Initial "typed element" support

This is a framework which provides a way to access an "element" via a series of interfaces, where each type of element may or may not choose to implement a particular interface.

An element type may either be "typed" or "untyped" depending on whether it requires extra payload information to store itself, or whether it can be represented purely by the element ID. The allocation and management of payload instances is handled internally as required, and the payload itself can be any C++ struct that uses the UE_DECLARE_TYPED_ELEMENT_DATA_RTTI and UE_DEFINE_TYPED_ELEMENT_DATA_RTTI macros to add minimal element payload RTTI data.

Client code will most commonly make use of FTypedElementHandle (a combination of an element ID and its internal data pointer) and TTypedElement (a combination of an element handle and an interface instance that is associated with it). FTypedElementList is useful for storing a set of elements (eg, a selection set), and FTypedElementId may be useful in limited circumstances where you need the most minimal representation of an element (current 32-bits).

Implementation code will need to make use of UTypedElementInterface to create element interface types, as well as UTypedElementRegistry to register, create, and destroy instances of an element. An element created via UTypedElementRegistry is returned as a TTypedElementOwner, which provides mutable access to the internal payload data (if any), and is the only time outside of implementing an interface for a specific element type that you should make any assumptions about the type of data held within an element instance.

TypedElementFrameworkTests.h/.cpp provides a basic example of usage. Note: While UTypedElementRegistry isn't intrinsically a singleton, there is a global instance available via UTypedElementRegistry::GetInstance().

diff --git a/Engine/Source/Runtime/TypedElementFramework/Private/TypedElementFrameworkModule.cpp b/Engine/Source/Runtime/TypedElementFramework/Private/TypedElementFrameworkModule.cpp
new file mode 100644
--- /dev/null
+++ b/Engine/Source/Runtime/TypedElementFramework/Private/TypedElementFrameworkModule.cpp
@@ -0,0 +1,20 @@
+// Copyright Epic Games, Inc. All Rights Reserved.
+
+#include "TypedElementRegistry.h"
+#include "Modules/ModuleManager.h"
+
+class FTypedElementFrameworkModule : public FDefaultModuleImpl
+{
+public:
+	virtual void StartupModule() override
+	{
+		UTypedElementRegistry::Private_InitializeInstance();
+	}
+
+	virtual void ShutdownModule() override
+	{
+		UTypedElementRegistry::Private_ShutdownInstance();
+	}
+};
+
+IMPLEMENT_MODULE(FTypedElementFrameworkModule, TypedElementFramework)
diff --git a/Engine/Source/Runtime/TypedElementFramework/Private/TypedElementFrameworkTests.cpp b/Engine/Source/Runtime/TypedElementFramework/Private/TypedElementFrameworkTests.cpp
new file mode 100644
--- /dev/null
+++ b/Engine/Source/Runtime/TypedElementFramework/Private/TypedElementFrameworkTests.cpp
@@ -0,0 +1,253 @@
+// Copyright Epic Games, Inc. All Rights Reserved.
+
+#include "TypedElementFrameworkTests.h"
+#include "Misc/AutomationTest.h"
+#include "ProfilingDebugging/ScopedTimers.h"
+
+UE_DEFINE_TYPED_ELEMENT_DATA_RTTI(FTestTypedElementData);
+
+FText UTestTypedElementInterfaceA_ImplTyped::GetDisplayName(const FTypedElementHandle& InElementHandle)
+{
+	const FTestTypedElementData* ElementData = InElementHandle.GetDataScript<FTestTypedElementData>();
+	if (!ElementData)
+	{
+		return FText();
+	}
+
+	return FText();
+}
+
+bool UTestTypedElementInterfaceA_ImplTyped::SetDisplayName(const FTypedElementHandle& InElementHandle, FText InNewName, bool bNotify)
+{
+	const FTestTypedElementData* ElementData = InElementHandle.GetDataScript<FTestTypedElementData>();
+	if (!ElementData)
+	{
+		return false;
+	}
+
+	return false;
+}
+
+FText UTestTypedElementInterfaceA_ImplUntyped::GetDisplayName(const FTypedElementHandle& InElementHandle)
+{
+	return FText();
+}
+
+bool UTestTypedElementInterfaceA_ImplUntyped::SetDisplayName(const FTypedElementHandle& InElementHandle, FText InNewName, bool bNotify)
+{
+	return false;
+}
+
+#if WITH_DEV_AUTOMATION_TESTS
+
+IMPLEMENT_SIMPLE_AUTOMATION_TEST(FTypedElementRegistrySmokeTest, "System.Runtime.TypedElementRegistry.SmokeTest", EAutomationTestFlags::EditorContext | EAutomationTestFlags::SmokeFilter)
+bool FTypedElementRegistrySmokeTest::RunTest(const FString& Parameters)
+{
+	UTypedElementRegistry* Registry = NewObject<UTypedElementRegistry>();
+	
+	auto TestInterfaceAHandle = [](const TTypedElement<UTestTypedElementInterfaceA>& InElementDataInterfaceHandle)
+	{
+		// Proxy API added via specialization
+		InElementDataInterfaceHandle.SetDisplayName(FText());
+		InElementDataInterfaceHandle.GetDisplayName();
+
+		// Verbose API
+		InElementDataInterfaceHandle.GetInterfaceChecked().SetDisplayName(InElementDataInterfaceHandle, FText());
+		InElementDataInterfaceHandle.GetInterfaceChecked().GetDisplayName(InElementDataInterfaceHandle);
+	};
+
+	auto TestInterfaceAccess = [&Registry, &TestInterfaceAHandle](const FTypedElementHandle& InElementHandle)
+	{
+		// Get the interface and the element handle in two calls - this is how scripting might work
+		if (UTestTypedElementInterfaceA* Interface = Registry->GetElementInterface<UTestTypedElementInterfaceA>(InElementHandle))
+		{
+			Interface->SetDisplayName(InElementHandle, FText());
+			Interface->GetDisplayName(InElementHandle);
+		}
+
+		// Get the interface and the element handle in a single call - this is how C++ might work
+		if (TTypedElement<UTestTypedElementInterfaceA> Element = Registry->GetElement<UTestTypedElementInterfaceA>(InElementHandle))
+		{
+			TestInterfaceAHandle(Element);
+		}
+	};
+
+	// Test that all the cast variants compile
+	{
+		FTypedElement DummyElement;
+		CastTypedElement<UTestTypedElementInterfaceA>(DummyElement);
+		CastTypedElement<UTestTypedElementInterfaceA>(FTypedElement());
+		CastTypedElementChecked<UTestTypedElementInterfaceA>(DummyElement);
+		CastTypedElementChecked<UTestTypedElementInterfaceA>(FTypedElement());
+
+		TTypedElement<UTestTypedElementInterfaceA> DummyCastedElement;
+		CastTypedElement<UTestTypedElementInterfaceA>(DummyElement, DummyCastedElement);
+		CastTypedElement<UTestTypedElementInterfaceA>(FTypedElement(), DummyCastedElement);
+		CastTypedElementChecked<UTestTypedElementInterfaceA>(DummyElement, DummyCastedElement);
+		CastTypedElementChecked<UTestTypedElementInterfaceA>(FTypedElement(), DummyCastedElement);
+	}
+
+	const FName DummyElementType_Typed = "DummyElementType_Typed";
+	Registry->RegisterElementType<FTestTypedElementData>(DummyElementType_Typed);
+	Registry->RegisterElementInterface<UTestTypedElementInterfaceA>(DummyElementType_Typed, NewObject<UTestTypedElementInterfaceA_ImplTyped>());
+
+	const FName DummyElementType_Untyped = "DummyElementType_Untyped";
+	Registry->RegisterElementType(DummyElementType_Untyped);
+	Registry->RegisterElementInterface<UTestTypedElementInterfaceA>(DummyElementType_Untyped, NewObject<UTestTypedElementInterfaceA_ImplUntyped>());
+
+	TTypedElementOwner<FTestTypedElementData> TypedElement1 = Registry->CreateElement<FTestTypedElementData>(DummyElementType_Typed);
+	TypedElement1.GetDataChecked().InternalElementId = "TypedElement1";
+	TTypedElementOwner<FTestTypedElementData> TypedElement2 = Registry->CreateElement<FTestTypedElementData>(DummyElementType_Typed);
+	TypedElement2.GetDataChecked().InternalElementId = "TypedElement2";
+	TTypedElementOwner<FTestTypedElementData> TypedElement3 = Registry->CreateElement<FTestTypedElementData>(DummyElementType_Typed);
+	TypedElement3.GetDataChecked().InternalElementId = "TypedElement3";
+
+	FTypedElementOwner UntypedElement1 = Registry->CreateElement(DummyElementType_Untyped, 0);
+	FTypedElementOwner UntypedElement2 = Registry->CreateElement(DummyElementType_Untyped, 1);
+	FTypedElementOwner UntypedElement3 = Registry->CreateElement(DummyElementType_Untyped, 2);
+
+	TestInterfaceAccess(TypedElement1.AcquireHandle());
+	TestInterfaceAccess(UntypedElement1.AcquireHandle());
+
+	FTypedElementListPtr ElementList = Registry->CreateElementList();
+	ElementList->Add(TypedElement1);
+	ElementList->Add(TypedElement2);
+	ElementList->Add(TypedElement3);
+	ElementList->Add(UntypedElement1);
+	ElementList->Add(UntypedElement2);
+	ElementList->Add(UntypedElement3);
+	
+	for (const FTypedElementHandle& ElementHandle : *ElementList)
+	{
+		TestInterfaceAccess(ElementHandle);
+	}
+	
+	for (const TTypedElement<UTestTypedElementInterfaceA>& ElementHandle : ElementList->InterateInterface<UTestTypedElementInterfaceA>())
+	{
+		TestInterfaceAHandle(ElementHandle);
+	}
+
+	ElementList.Reset();
+
+	Registry->DestroyElement(TypedElement1);
+	Registry->DestroyElement(TypedElement2);
+	Registry->DestroyElement(TypedElement3);
+
+	Registry->DestroyElement(UntypedElement1);
+	Registry->DestroyElement(UntypedElement2);
+	Registry->DestroyElement(UntypedElement3);
+
+	return true;
+}
+
+IMPLEMENT_SIMPLE_AUTOMATION_TEST(FTypedElementRegistryPerfTest, "System.Runtime.TypedElementRegistry.PerfTest", EAutomationTestFlags::EditorContext | EAutomationTestFlags::PerfFilter)
+bool FTypedElementRegistryPerfTest::RunTest(const FString& Parameters)
+{
+	constexpr int32 NumHandlesToTest = 1000000;
+
+	UTypedElementRegistry* Registry = NewObject<UTypedElementRegistry>();
+
+	const FName DummyElementType_Typed = "DummyElementType_Typed";
+	Registry->RegisterElementType<FTestTypedElementData>(DummyElementType_Typed);
+	Registry->RegisterElementInterface<UTestTypedElementInterfaceA>(DummyElementType_Typed, NewObject<UTestTypedElementInterfaceA_ImplTyped>());
+
+	const FName DummyElementType_Untyped = "DummyElementType_Untyped";
+	Registry->RegisterElementType(DummyElementType_Untyped);
+	Registry->RegisterElementInterface<UTestTypedElementInterfaceA>(DummyElementType_Untyped, NewObject<UTestTypedElementInterfaceA_ImplUntyped>());
+
+	TArray<TTypedElementOwner<FTestTypedElementData>> TypedOwnerHandles;
+	TArray<FTypedElementOwner> UntypedOwnerHandles;
+	FTypedElementListPtr ElementList = Registry->CreateElementList();
+
+	// Create typed handles
+	{
+		TypedOwnerHandles.Reserve(NumHandlesToTest);
+		{
+			FScopedDurationTimeLogger Timer(FString::Printf(TEXT("Creating %d typed handles"), NumHandlesToTest));
+			for (int32 Index = 0; Index < NumHandlesToTest; ++Index)
+			{
+				TypedOwnerHandles.Emplace(Registry->CreateElement<FTestTypedElementData>(DummyElementType_Typed));
+			}
+		}
+	}
+
+	// Create untyped handles
+	{
+		UntypedOwnerHandles.Reserve(NumHandlesToTest);
+		{
+			FScopedDurationTimeLogger Timer(FString::Printf(TEXT("Creating %d untyped handles"), NumHandlesToTest));
+			for (int32 Index = 0; Index < NumHandlesToTest; ++Index)
+			{
+				UntypedOwnerHandles.Emplace(Registry->CreateElement(DummyElementType_Untyped, Index));
+			}
+		}
+	}
+
+	// Populate an element list with all handles
+	{
+		FScopedDurationTimeLogger Timer(FString::Printf(TEXT("Appending %d typed handles to list"), TypedOwnerHandles.Num() + UntypedOwnerHandles.Num()));
+		ElementList->Reserve(TypedOwnerHandles.Num() + UntypedOwnerHandles.Num());
+		ElementList->Append(TypedOwnerHandles);
+		ElementList->Append(UntypedOwnerHandles);
+	}
+
+	// Find an interface from each handle
+	{
+		FScopedDurationTimeLogger Timer(FString::Printf(TEXT("Finding %d interfaces from list"), ElementList->Num()));
+		for (const FTypedElementHandle& ElementHandle : *ElementList)
+		{
+			Registry->GetElementInterface<UTestTypedElementInterfaceA>(ElementHandle);
+		}
+	}
+
+	// Find an element from each handle
+	{
+		FScopedDurationTimeLogger Timer(FString::Printf(TEXT("Finding %d elements from list"), ElementList->Num()));
+		for (const FTypedElementHandle& ElementHandle : *ElementList)
+		{
+			Registry->GetElement<UTestTypedElementInterfaceA>(ElementHandle);
+		}
+	}
+
+	// Iterate all elements that implement an interface
+	{
+		FScopedDurationTimeLogger Timer(FString::Printf(TEXT("Iterating %d elements in list"), ElementList->Num()));
+		for (const TTypedElement<UTestTypedElementInterfaceA>& ElementHandle : ElementList->InterateInterface<UTestTypedElementInterfaceA>())
+		{
+		}
+	}
+
+	// Clear the element list
+	{
+		FScopedDurationTimeLogger Timer(FString::Printf(TEXT("Reset %d elements in list"), ElementList->Num()));
+		ElementList.Reset();
+	}
+
+	// Destroy typed handles
+	{
+		{
+			FScopedDurationTimeLogger Timer(FString::Printf(TEXT("Destroying %d typed handles"), TypedOwnerHandles.Num()));
+			for (TTypedElementOwner<FTestTypedElementData>& TypedOwnerHandle : TypedOwnerHandles)
+			{
+				Registry->DestroyElement(TypedOwnerHandle);
+			}
+		}
+		TypedOwnerHandles.Reset();
+	}
+
+	// Destroy untyped handles
+	{
+		{
+			FScopedDurationTimeLogger Timer(FString::Printf(TEXT("Destroying %d untyped handles"), UntypedOwnerHandles.Num()));
+			for (FTypedElementOwner& UntypedOwnerHandle : UntypedOwnerHandles)
+			{
+				Registry->DestroyElement(UntypedOwnerHandle);
+			}
+		}
+		UntypedOwnerHandles.Reset();
+	}
+
+	return true;
+}
+
+#endif // WITH_DEV_AUTOMATION_TESTS
diff --git a/Engine/Source/Runtime/TypedElementFramework/Private/TypedElementFrameworkTests.h b/Engine/Source/Runtime/TypedElementFramework/Private/TypedElementFrameworkTests.h
new file mode 100644
--- /dev/null
+++ b/Engine/Source/Runtime/TypedElementFramework/Private/TypedElementFrameworkTests.h
@@ -0,0 +1,62 @@
+// Copyright Epic Games, Inc. All Rights Reserved.
+
+#pragma once
+
+#include "TypedElementRegistry.h"
+#include "TypedElementFrameworkTests.generated.h"
+
+/**
+ * Test interfaces
+ */
+UCLASS()
+class UTestTypedElementInterfaceA : public UTypedElementInterface
+{
+	GENERATED_BODY()
+
+public:
+	UFUNCTION(BlueprintCallable, Category="TypedElementFramework|Testing")
+	virtual FText GetDisplayName(const FTypedElementHandle& InElementHandle) { return FText(); }
+	
+	UFUNCTION(BlueprintCallable, Category="TypedElementFramework|Testing")
+	virtual bool SetDisplayName(const FTypedElementHandle& InElementHandle, FText InNewName, bool bNotify = true) { return false; }
+};
+
+template <>
+struct TTypedElement<UTestTypedElementInterfaceA> : public TTypedElementBase<UTestTypedElementInterfaceA>
+{
+	FText GetDisplayName() const { return InterfacePtr->GetDisplayName(*this); }
+	bool SetDisplayName(FText InNewName, bool bNotify = true) const { return InterfacePtr->SetDisplayName(*this, InNewName, bNotify); }
+};
+
+/**
+ * Test dummy type
+ */
+struct FTestTypedElementData
+{
+	UE_DECLARE_TYPED_ELEMENT_DATA_RTTI(FTestTypedElementData);
+	
+	FName InternalElementId;
+};
+
+UCLASS()
+class UTestTypedElementInterfaceA_ImplTyped : public UTestTypedElementInterfaceA
+{
+	GENERATED_BODY()
+
+public:
+	virtual FText GetDisplayName(const FTypedElementHandle& InElementHandle) override;
+	virtual bool SetDisplayName(const FTypedElementHandle& InElementHandle, FText InNewName, bool bNotify) override;
+};
+
+/**
+ * Test untyped
+ */
+UCLASS()
+class UTestTypedElementInterfaceA_ImplUntyped : public UTestTypedElementInterfaceA
+{
+	GENERATED_BODY()
+
+public:
+	virtual FText GetDisplayName(const FTypedElementHandle& InElementHandle) override;
+	virtual bool SetDisplayName(const FTypedElementHandle& InElementHandle, FText InNewName, bool bNotify) override;
+};
diff --git a/Engine/Source/Runtime/TypedElementFramework/Private/TypedElementList.cpp b/Engine/Source/Runtime/TypedElementFramework/Private/TypedElementList.cpp
new file mode 100644
--- /dev/null
+++ b/Engine/Source/Runtime/TypedElementFramework/Private/TypedElementList.cpp
@@ -0,0 +1,227 @@
+// Copyright Epic Games, Inc. All Rights Reserved.
+
+#include "TypedElementList.h"
+#include "TypedElementRegistry.h"
+
+namespace TypedElementList_Private
+{
+
+void GetElementImpl(const UTypedElementRegistry* InRegistry, const FTypedElementHandle& InElementHandle, const UClass* InBaseInterfaceType, FTypedElement& OutElement)
+{
+	InRegistry->Private_GetElementImpl(InElementHandle, InBaseInterfaceType, OutElement);
+}
+
+} // namespace TypedElementList_Private
+
+
+FTypedElementListLegacySync::FTypedElementListLegacySync(const FTypedElementList& InElementList)
+	: ElementList(InElementList)
+{
+}
+
+FTypedElementListLegacySync::FOnSyncEvent& FTypedElementListLegacySync::OnSyncEvent()
+{
+	return OnSyncEventDelegate;
+}
+
+void FTypedElementListLegacySync::Private_EmitSyncEvent(const ESyncType InSyncType, const FTypedElementHandle& InElementHandle)
+{
+	const bool bIsWithinBatchOperation = IsRunningBatchOperation();
+	bBatchOperationIsDirty |= bIsWithinBatchOperation;
+	OnSyncEventDelegate.Broadcast(ElementList, InSyncType, InElementHandle, bIsWithinBatchOperation);
+}
+
+bool FTypedElementListLegacySync::IsRunningBatchOperation() const
+{
+	return NumOpenBatchOperations > 0;
+}
+
+void FTypedElementListLegacySync::BeginBatchOperation()
+{
+	++NumOpenBatchOperations;
+}
+
+void FTypedElementListLegacySync::EndBatchOperation(const bool InNotify)
+{
+	checkf(NumOpenBatchOperations > 0, TEXT("Batch operation underflow!"));
+
+	if (--NumOpenBatchOperations == 0)
+	{
+		const bool bNotifyChange = bBatchOperationIsDirty && InNotify;
+		bBatchOperationIsDirty = false;
+
+		if (bNotifyChange)
+		{
+			Private_EmitSyncEvent(ESyncType::BatchComplete);
+			check(!bBatchOperationIsDirty); // This should still be false after emitting the notification!
+		}
+	}
+}
+
+bool FTypedElementListLegacySync::IsBatchOperationDirty() const
+{
+	return bBatchOperationIsDirty;
+}
+
+void FTypedElementListLegacySync::ForceBatchOperationDirty()
+{
+	if (NumOpenBatchOperations > 0)
+	{
+		bBatchOperationIsDirty = true;
+	}
+}
+
+
+FTypedElementList::FTypedElementList(UTypedElementRegistry* InRegistry)
+	: Registry(InRegistry)
+{
+	Registry->Private_OnElementListCreated(this);
+}
+
+FTypedElementList::~FTypedElementList()
+{
+	LegacySync.Reset();
+	Registry->Private_OnElementListDestroyed(this);
+}
+
+FTypedElementListPtr FTypedElementList::Private_CreateElementList(UTypedElementRegistry* InRegistry)
+{
+	return FTypedElementListPtr(new FTypedElementList(InRegistry));
+}
+
+FTypedElementListPtr FTypedElementList::Clone() const
+{
+	FTypedElementListPtr ClonedElementList = Private_CreateElementList(Registry);
+	ClonedElementList->ElementCombinedIds = ElementCombinedIds;
+	ClonedElementList->ElementHandles = ElementHandles;
+	return ClonedElementList;
+}
+
+bool FTypedElementList::AddElementImpl(FTypedElementHandle&& InElementHandle)
+{
+	if (!InElementHandle)
+	{
+		return false;
+	}
+
+	bool bAlreadyAdded = false;
+	ElementCombinedIds.Add(InElementHandle.GetId().GetCombinedId(), &bAlreadyAdded);
+
+	if (!bAlreadyAdded)
+	{
+		const FTypedElementHandle& AddedElementHandle = ElementHandles.Add_GetRef(MoveTemp(InElementHandle));
+		NoteListChanged(EChangeType::Added, AddedElementHandle);
+	}
+
+	return !bAlreadyAdded;
+}
+
+bool FTypedElementList::RemoveElementImpl(const FTypedElementId& InElementId)
+{
+	if (!InElementId)
+	{
+		return false;
+	}
+
+	const bool bRemoved = ElementCombinedIds.Remove(InElementId.GetCombinedId()) > 0;
+
+	if (bRemoved)
+	{
+		const int32 ElementHandleIndexToRemove = ElementHandles.IndexOfByPredicate([&InElementId](const FTypedElementHandle& InElementHandle)
+		{
+			return InElementHandle.GetId() == InElementId;
+		});
+		checkSlow(ElementHandleIndexToRemove != INDEX_NONE);
+
+		FTypedElementHandle RemovedElementHandle = MoveTemp(ElementHandles[ElementHandleIndexToRemove]);
+		ElementHandles.RemoveAt(ElementHandleIndexToRemove, 1, /*bAllowShrinking*/false);
+
+		NoteListChanged(EChangeType::Removed, RemovedElementHandle);
+	}
+
+	return bRemoved;
+}
+
+int32 FTypedElementList::RemoveAllElementsImpl(TFunctionRef<bool(const FTypedElementHandle&)> InPredicate)
+{
+	if (LegacySync)
+	{
+		LegacySync->BeginBatchOperation();
+	}
+
+	int32 RemovedCount = 0;
+	for (int32 Index = ElementHandles.Num() - 1; Index >= 0; --Index)
+	{
+		if (InPredicate(ElementHandles[Index]))
+		{
+			FTypedElementHandle RemovedElementHandle = MoveTemp(ElementHandles[Index]);
+			ElementCombinedIds.Remove(RemovedElementHandle.GetId().GetCombinedId());
+			ElementHandles.RemoveAt(Index, 1, /*bAllowShrinking*/false);
+
+			NoteListChanged(EChangeType::Removed, RemovedElementHandle);
+
+			++RemovedCount;
+		}
+	}
+
+	if (LegacySync)
+	{
+		LegacySync->EndBatchOperation();
+	}
+
+	return RemovedCount;
+}
+
+bool FTypedElementList::ContainsElementImpl(const FTypedElementId& InElementId) const
+{
+	return InElementId 
+		&& ElementCombinedIds.Contains(InElementId.GetCombinedId());
+}
+
+FTypedElementListLegacySync& FTypedElementList::Legacy_GetSync()
+{
+	if (!LegacySync)
+	{
+		LegacySync = MakeUnique<FTypedElementListLegacySync>(*this);
+	}
+	return *LegacySync;
+}
+
+void FTypedElementList::NotifyPendingChanges()
+{
+	if (bHasPendingNotify)
+	{
+		bHasPendingNotify = false;
+		OnChangedDelegate.Broadcast(*this);
+		check(!bHasPendingNotify); // This should still be false after emitting the notification!
+	}
+}
+
+void FTypedElementList::NoteListChanged(const EChangeType InChangeType, const FTypedElementHandle& InElementHandle)
+{
+	bHasPendingNotify = true;
+
+	if (LegacySync)
+	{
+		FTypedElementListLegacySync::ESyncType SyncType = FTypedElementListLegacySync::ESyncType::Modified;
+		switch (InChangeType)
+		{
+		case EChangeType::Added:
+			SyncType = FTypedElementListLegacySync::ESyncType::Added;
+			break;
+
+		case EChangeType::Removed:
+			SyncType = FTypedElementListLegacySync::ESyncType::Removed;
+			break;
+
+		case EChangeType::Cleared:
+			SyncType = FTypedElementListLegacySync::ESyncType::Cleared;
+			break;
+
+		default:
+			break;
+		}
+
+		LegacySync->Private_EmitSyncEvent(SyncType, InElementHandle);
+	}
+}
diff --git a/Engine/Source/Runtime/TypedElementFramework/Private/TypedElementRegistry.cpp b/Engine/Source/Runtime/TypedElementFramework/Private/TypedElementRegistry.cpp
new file mode 100644
--- /dev/null
+++ b/Engine/Source/Runtime/TypedElementFramework/Private/TypedElementRegistry.cpp
@@ -0,0 +1,163 @@
+// Copyright Epic Games, Inc. All Rights Reserved.
+
+#include "TypedElementRegistry.h"
+#include "Misc/CoreDelegates.h"
+#include "Misc/ScopeLock.h"
+
+const FTypedElementId FTypedElementId::Unset;
+
+UTypedElementRegistry* UTypedElementRegistry::Instance = nullptr;
+
+UTypedElementRegistry::UTypedElementRegistry()
+{
+	FCoreDelegates::OnEndFrame.AddUObject(this, &UTypedElementRegistry::NotifyElementListPendingChanges);
+}
+
+void UTypedElementRegistry::Private_InitializeInstance()
+{
+	checkf(!Instance, TEXT("Instance was already initialized!"));
+	Instance = NewObject<UTypedElementRegistry>();
+	Instance->AddToRoot();
+}
+
+void UTypedElementRegistry::Private_ShutdownInstance()
+{
+	checkf(Instance, TEXT("Instance was already shutdown!"));
+	Instance->RemoveFromRoot();
+	Instance = nullptr;
+}
+
+UTypedElementRegistry* UTypedElementRegistry::GetInstance()
+{
+	return Instance;
+}
+
+void UTypedElementRegistry::AddReferencedObjects(UObject* InThis, FReferenceCollector& Collector)
+{
+	Super::AddReferencedObjects(InThis, Collector);
+
+	UTypedElementRegistry* This = CastChecked<UTypedElementRegistry>(InThis);
+	for (TUniquePtr<FRegisteredElementType>& RegisteredElementType : This->RegisteredElementTypes)
+	{
+		if (RegisteredElementType)
+		{
+			for (auto& InterfacesPair : RegisteredElementType->Interfaces)
+			{
+				Collector.AddReferencedObject(InterfacesPair.Value);
+			}
+		}
+	}
+}
+
+void UTypedElementRegistry::RegisterElementTypeImpl(const FName InElementTypeName, TUniquePtr<FRegisteredElementType>&& InRegisteredElementType)
+{
+	// Query whether this type has previously been registered in any type registry, and if so re-use that ID
+	// If not (or if the element is typeless) then assign the next available ID
+	FTypedHandleTypeId TypeId = InRegisteredElementType->GetDataTypeId();
+	if (TypeId == 0)
+	{
+		static FCriticalSection NextTypeIdCS;
+		static FTypedHandleTypeId NextTypeId = 1;
+
+		FScopeLock NextTypeIdLock(&NextTypeIdCS);
+
+		checkf(NextTypeId <= TypedHandleMaxTypeId, TEXT("Ran out of typed element type IDs!"));
+
+		TypeId = NextTypeId++;
+		InRegisteredElementType->SetDataTypeId(TypeId);
+	}
+
+	InRegisteredElementType->TypeId = TypeId;
+	InRegisteredElementType->TypeName = InElementTypeName;
+	AddRegisteredElementType(MoveTemp(InRegisteredElementType));
+}
+
+void UTypedElementRegistry::RegisterElementInterfaceImpl(const FName InElementTypeName, UTypedElementInterface* InElementInterface, const TSubclassOf<UTypedElementInterface>& InBaseInterfaceType)
+{
+	checkf(InElementInterface->IsA(InBaseInterfaceType), TEXT("Interface '%s' of type '%s' does not derive from '%s'!"), *InElementInterface->GetPathName(), *InElementInterface->GetClass()->GetName(), *InBaseInterfaceType->GetName());
+
+	FRegisteredElementType* RegisteredElementType = GetRegisteredElementTypeFromName(InElementTypeName);
+	checkf(RegisteredElementType, TEXT("Element type '%s' has not been registered!"), *InElementTypeName.ToString());
+
+	checkf(!RegisteredElementType->Interfaces.Contains(InBaseInterfaceType->GetFName()), TEXT("Element type '%s' has already registered an interface for '%s'!"), *InElementTypeName.ToString(), *InBaseInterfaceType->GetName());
+	RegisteredElementType->Interfaces.Add(InBaseInterfaceType->GetFName(), InElementInterface);
+}
+
+UTypedElementInterface* UTypedElementRegistry::GetElementInterfaceImpl(const FTypedElementId& InElementId, const TSubclassOf<UTypedElementInterface>& InBaseInterfaceType) const
+{
+	if (!InElementId)
+	{
+		return nullptr;
+	}
+
+	FRegisteredElementType* RegisteredElementType = GetRegisteredElementTypeFromId(InElementId.GetTypeId());
+	checkf(RegisteredElementType, TEXT("Element type ID '%d' has not been registered!"), InElementId.GetTypeId());
+
+	return RegisteredElementType->Interfaces.FindRef(InBaseInterfaceType->GetFName());
+}
+
+void UTypedElementRegistry::ReleaseElementId(FTypedElementId& InOutElementId)
+{
+	if (!InOutElementId)
+	{
+		return;
+	}
+
+	FRegisteredElementType* RegisteredElementType = GetRegisteredElementTypeFromId(InOutElementId.GetTypeId());
+	checkf(RegisteredElementType, TEXT("Element type ID '%d' has not been registered!"), InOutElementId.GetTypeId());
+
+#if WITH_TYPED_ELEMENT_REFCOUNT
+	const FTypedElementInternalData& ElementData = RegisteredElementType->GetDataForElement(InOutElementId.GetElementId());
+	ElementData.ReleaseRef();
+#endif	// WITH_TYPED_ELEMENT_REFCOUNT
+
+	InOutElementId.Private_DestroyNoRef();
+}
+
+FTypedElementHandle UTypedElementRegistry::GetElementHandle(const FTypedElementId& InElementId) const
+{
+	if (!InElementId)
+	{
+		return FTypedElementHandle();
+	}
+
+	FRegisteredElementType* RegisteredElementType = GetRegisteredElementTypeFromId(InElementId.GetTypeId());
+	checkf(RegisteredElementType, TEXT("Element type ID '%d' has not been registered!"), InElementId.GetTypeId());
+
+	FTypedElementHandle ElementHandle;
+	ElementHandle.Private_InitializeAddRef(InElementId.GetTypeId(), InElementId.GetElementId(), RegisteredElementType->GetDataForElement(InElementId.GetElementId()));
+
+	return ElementHandle;
+}
+
+FTypedElementListPtr UTypedElementRegistry::CreateElementList(TArrayView<const FTypedElementId> InElementIds)
+{
+	FTypedElementListPtr ElementList = CreateElementList();
+
+	for (const FTypedElementId& ElementId : InElementIds)
+	{
+		if (FTypedElementHandle ElementHandle = GetElementHandle(ElementId))
+		{
+			ElementList->Add(MoveTemp(ElementHandle));
+		}
+	}
+
+	return ElementList;
+}
+
+FTypedElementListPtr UTypedElementRegistry::CreateElementList(TArrayView<const FTypedElementHandle> InElementHandles)
+{
+	FTypedElementListPtr ElementList = CreateElementList();
+	ElementList->Append(InElementHandles);
+	return ElementList;
+}
+
+void UTypedElementRegistry::NotifyElementListPendingChanges()
+{
+	FReadScopeLock ActiveElementListsLock(ActiveElementListsRW);
+
+	for (FTypedElementList* ActiveElementList : ActiveElementLists)
+	{
+		ActiveElementList->NotifyPendingChanges();
+	}
+}
diff --git a/Engine/Source/Runtime/TypedElementFramework/Public/TypedElementData.h b/Engine/Source/Runtime/TypedElementFramework/Public/TypedElementData.h
new file mode 100644
--- /dev/null
+++ b/Engine/Source/Runtime/TypedElementFramework/Public/TypedElementData.h
@@ -0,0 +1,269 @@
+// Copyright Epic Games, Inc. All Rights Reserved.
+
+#pragma once
+
+#include "CoreMinimal.h"
+#include "Misc/ScopeRWLock.h"
+#include "Containers/ChunkedArray.h"
+#include "TypedElementLimits.h"
+
+/**
+ * Macro to declare the required RTTI data for types representing element data.
+ * @note Place this in the public section of your type declaration.
+ */
+#define UE_DECLARE_TYPED_ELEMENT_DATA_RTTI(ELEMENT_DATA_TYPE)						\
+	static FTypedHandleTypeId Private_RegisteredTypeId;								\
+	static FTypedHandleTypeId StaticTypeId() { return Private_RegisteredTypeId; }	\
+	static FName StaticTypeName() { static const FName TypeName = #ELEMENT_DATA_TYPE; return TypeName; }
+
+/**
+ * Macro to define the required RTTI data for types representing element data.
+ * @note Place this in the cpp file for your type definition.
+ */
+#define UE_DEFINE_TYPED_ELEMENT_DATA_RTTI(ELEMENT_DATA_TYPE)						\
+	FTypedHandleTypeId ELEMENT_DATA_TYPE::Private_RegisteredTypeId = 0;
+
+/**
+ * Base class for the internal payload data associated with elements.
+ */
+class FTypedElementInternalData
+{
+public:
+	FTypedElementInternalData() = default;
+	
+	FTypedElementInternalData(const FTypedElementInternalData&) = delete;
+	FTypedElementInternalData& operator=(const FTypedElementInternalData&) = delete;
+	
+	FTypedElementInternalData(FTypedElementInternalData&& InOther) = default;
+	FTypedElementInternalData& operator=(FTypedElementInternalData&&) = default;
+
+	virtual ~FTypedElementInternalData() = default;
+
+#if WITH_TYPED_ELEMENT_REFCOUNT
+	FORCEINLINE void AddRef() const
+	{
+		checkSlow(RefCount < TNumericLimits<FTypedHandleRefCount>::Max());
+		FPlatformAtomics::InterlockedIncrement(&RefCount);
+	}
+
+	FORCEINLINE void ReleaseRef() const
+	{
+		checkSlow(RefCount > 0);
+		FPlatformAtomics::InterlockedDecrement(&RefCount);
+	}
+
+	FORCEINLINE FTypedHandleRefCount GetRefCount() const
+	{
+		return FPlatformAtomics::AtomicRead(&RefCount);
+	}
+#endif	// WITH_TYPED_ELEMENT_REFCOUNT
+
+	virtual const void* GetUntypedData() const
+	{
+		return nullptr;
+	}
+
+private:
+#if WITH_TYPED_ELEMENT_REFCOUNT
+	mutable FTypedHandleRefCount RefCount = 0;
+#endif	// WITH_TYPED_ELEMENT_REFCOUNT
+};
+
+/**
+ * Internal payload data associated with typed elements.
+ */
+template <typename ElementDataType>
+class TTypedElementInternalData : public FTypedElementInternalData
+{
+public:
+	TTypedElementInternalData() = default;
+
+	TTypedElementInternalData(TTypedElementInternalData&& InOther) = default;
+	TTypedElementInternalData& operator=(TTypedElementInternalData&&) = default;
+
+	virtual ~TTypedElementInternalData() = default;
+
+	FORCEINLINE const ElementDataType& GetData() const
+	{
+		return Data;
+	}
+
+	FORCEINLINE ElementDataType& GetMutableData()
+	{
+		return Data;
+	}
+
+	virtual const void* GetUntypedData() const override
+	{
+		return &Data;
+	}
+
+private:
+	ElementDataType Data;
+};
+
+/**
+ * Internal payload data associated with typeless elements.
+ */
+template <>
+class TTypedElementInternalData<void> : public FTypedElementInternalData
+{
+public:
+	TTypedElementInternalData() = default;
+
+	TTypedElementInternalData(TTypedElementInternalData&& InOther) = default;
+	TTypedElementInternalData& operator=(TTypedElementInternalData&&) = default;
+
+	virtual ~TTypedElementInternalData() = default;
+};
+
+/**
+ * Data store implementation used by the element registry to manage internal data. 
+ * @note This is the generic implementation that uses an array and manages the IDs itself.
+ */
+template <typename ElementDataType>
+class TTypedElementInternalDataStore
+{
+public:
+	static_assert(TNumericLimits<int32>::Max() >= TypedHandleMaxElementId, "TTypedElementInternalDataStore internally uses signed 32-bit indices so cannot store TypedHandleMaxElementId! Consider making this container 64-bit aware, or explicitly remove this compile time check.");
+
+	TTypedElementInternalData<ElementDataType>& AddDataForElement(FTypedHandleElementId& InOutElementId)
+	{
+		FWriteScopeLock InternalDataLock(InternalDataRW);
+
+		checkSlow(InOutElementId < 0);
+
+		InOutElementId = InternalDataFreeIndices.Num() > 0
+			? InternalDataFreeIndices.Pop(/*bAllowShrinking*/false)
+			: InternalDataArray.Add();
+
+		return InternalDataArray[InOutElementId];
+	}
+
+	void RemoveDataForElement(const FTypedHandleElementId InElementId, const FTypedElementInternalData* InExpectedDataPtr)
+	{
+		FWriteScopeLock InternalDataLock(InternalDataRW);
+
+		checkSlow(InternalDataArray.IsValidIndex(InElementId));
+		
+		TTypedElementInternalData<ElementDataType>& InternalData = InternalDataArray[InElementId];
+		checkf(InExpectedDataPtr == &InternalData, TEXT("Internal data pointer did not match the expected value! Does this handle belong to a different element registry?"));
+		InternalData = TTypedElementInternalData<ElementDataType>();
+		InternalDataFreeIndices.Add(InElementId);
+	}
+
+	const TTypedElementInternalData<ElementDataType>& GetDataForElement(const FTypedHandleElementId InElementId) const
+	{
+		FReadScopeLock InternalDataLock(InternalDataRW);
+
+		checkSlow(InternalDataArray.IsValidIndex(InElementId));
+		return InternalDataArray[InElementId];
+	}
+
+	static FORCEINLINE void SetStaticDataTypeId(const FTypedHandleTypeId InTypeId)
+	{
+		checkSlow(ElementDataType::Private_RegisteredTypeId == 0);
+		ElementDataType::Private_RegisteredTypeId = InTypeId;
+	}
+
+	static FORCEINLINE FTypedHandleTypeId StaticDataTypeId()
+	{
+		return ElementDataType::StaticTypeId();
+	}
+
+	static FORCEINLINE FName StaticDataTypeName()
+	{
+		return ElementDataType::StaticTypeName();
+	}
+
+private:
+	mutable FRWLock InternalDataRW;
+	TChunkedArray<TTypedElementInternalData<ElementDataType>> InternalDataArray;
+	TArray<int32> InternalDataFreeIndices;
+};
+
+/**
+ * Data store implementation used by the element registry to manage internal data. 
+ * @note This is the typeless implementation that uses external IDs, and exists only to track ref counts.
+ */
+template <>
+class TTypedElementInternalDataStore<void>
+{
+public:
+	static_assert(TNumericLimits<int32>::Max() >= TypedHandleMaxElementId, "TTypedElementInternalDataStore internally uses signed 32-bit indices so cannot store TypedHandleMaxElementId! Consider making this container 64-bit aware, or explicitly remove this compile time check.");
+
+	TTypedElementInternalData<void>& AddDataForElement(FTypedHandleElementId& InOutElementId)
+	{
+#if WITH_TYPED_ELEMENT_REFCOUNT
+		FWriteScopeLock InternalDataLock(InternalDataRW);
+
+		checkSlow(InOutElementId >= 0);
+		checkSlow(!ElementIdToArrayIndex.Contains(InOutElementId));
+
+		const int32 InternalDataArrayIndex = InternalDataFreeIndices.Num() > 0
+			? InternalDataFreeIndices.Pop(/*bAllowShrinking*/false)
+			: InternalDataArray.Add();
+
+		ElementIdToArrayIndex.Add(InOutElementId, InternalDataArrayIndex);
+		return InternalDataArray[InternalDataArrayIndex];
+#else	// WITH_TYPED_ELEMENT_REFCOUNT
+		return SharedInternalData;
+#endif	// WITH_TYPED_ELEMENT_REFCOUNT
+	}
+
+	void RemoveDataForElement(const FTypedHandleElementId InElementId, const FTypedElementInternalData* InExpectedDataPtr)
+	{
+#if WITH_TYPED_ELEMENT_REFCOUNT
+		FWriteScopeLock InternalDataLock(InternalDataRW);
+
+		int32 InternalDataArrayIndex = INDEX_NONE;
+		ElementIdToArrayIndex.RemoveAndCopyValue(InElementId, InternalDataArrayIndex);
+
+		checkSlow(InternalDataArray.IsValidIndex(InternalDataArrayIndex));
+
+		TTypedElementInternalData<void>& InternalData = InternalDataArray[InElementId];
+		checkf(InExpectedDataPtr == &InternalData, TEXT("Internal data pointer did not match the expected value! Does this handle belong to a different element registry?"));
+		InternalData = TTypedElementInternalData<void>();
+		InternalDataFreeIndices.Add(InternalDataArrayIndex);
+#else	// WITH_TYPED_ELEMENT_REFCOUNT
+		checkf(InExpectedDataPtr == &SharedInternalData, TEXT("Internal data pointer did not match the expected value! Does this handle belong to a different element registry?"));
+#endif	// WITH_TYPED_ELEMENT_REFCOUNT
+	}
+
+	const TTypedElementInternalData<void>& GetDataForElement(const FTypedHandleElementId InElementId) const
+	{
+#if WITH_TYPED_ELEMENT_REFCOUNT
+		FReadScopeLock InternalDataLock(InternalDataRW);
+
+		const int32* InternalDataArrayIndexPtr = ElementIdToArrayIndex.Find(InElementId);
+		checkSlow(InternalDataArrayIndexPtr && InternalDataArray.IsValidIndex(*InternalDataArrayIndexPtr));
+		return InternalDataArray[*InternalDataArrayIndexPtr];
+#else	// WITH_TYPED_ELEMENT_REFCOUNT
+		return SharedInternalData;
+#endif	// WITH_TYPED_ELEMENT_REFCOUNT
+	}
+
+	static FORCEINLINE void SetStaticDataTypeId(const FTypedHandleTypeId InTypeId)
+	{
+	}
+
+	static FORCEINLINE FTypedHandleTypeId StaticDataTypeId()
+	{
+		return 0;
+	}
+
+	static FORCEINLINE FName StaticDataTypeName()
+	{
+		return FName();
+	}
+
+private:
+#if WITH_TYPED_ELEMENT_REFCOUNT
+	mutable FRWLock InternalDataRW;
+	TChunkedArray<TTypedElementInternalData<void>> InternalDataArray;
+	TArray<int32> InternalDataFreeIndices;
+	TMap<FTypedHandleElementId, int32> ElementIdToArrayIndex;
+#else	// WITH_TYPED_ELEMENT_REFCOUNT
+	TTypedElementInternalData<void> SharedInternalData;
+#endif	// WITH_TYPED_ELEMENT_REFCOUNT
+};
diff --git a/Engine/Source/Runtime/TypedElementFramework/Public/TypedElementHandle.h b/Engine/Source/Runtime/TypedElementFramework/Public/TypedElementHandle.h
new file mode 100644
--- /dev/null
+++ b/Engine/Source/Runtime/TypedElementFramework/Public/TypedElementHandle.h
@@ -0,0 +1,847 @@
+// Copyright Epic Games, Inc. All Rights Reserved.
+
+#pragma once
+
+#include <type_traits>
+#include "CoreMinimal.h"
+#include "TypedElementData.h"
+#include "TypedElementHandle.generated.h"
+
+/**
+ * Base typed used to represent element interfaces.
+ * @note Top-level element interfaces that inherit from this should also specialize TTypedElement for their API.
+ */
+UCLASS(Abstract)
+class TYPEDELEMENTFRAMEWORK_API UTypedElementInterface : public UObject
+{
+	GENERATED_BODY()
+};
+
+/**
+ * The most minimal representation of an element - its ID!
+ * This type is not immediately useful on its own, but can be used to find an element from the element registry or an element list.
+ * @note This is ref-counted like handles themselves are, so as long as an ID is available, the handle will be too.
+ * @note IDs lack the information needed to auto-release on destruction, so must be manually released, either via the corresponding handle or their owner element registry.
+ */
+struct TYPEDELEMENTFRAMEWORK_API FTypedElementId
+{
+public:
+	FTypedElementId()
+		: CombinedId(0)
+	{
+	}
+
+	FTypedElementId(const FTypedElementId&) = delete;
+	FTypedElementId& operator=(const FTypedElementId&) = delete;
+	
+	FTypedElementId(FTypedElementId&& InOther)
+		: CombinedId(InOther.CombinedId)
+	{
+		InOther.Private_DestroyNoRef();
+	}
+
+	FTypedElementId& operator=(FTypedElementId&& InOther)
+	{
+		if (this != &InOther)
+		{
+			CombinedId = InOther.CombinedId;
+
+			InOther.Private_DestroyNoRef();
+		}
+		return *this;
+	}
+
+	~FTypedElementId()
+	{
+		checkf(!IsSet(), TEXT("Element ID was still set during destruction! This will leak an element reference, and you should release this ID prior to destruction!"));
+	}
+
+	FORCEINLINE explicit operator bool() const
+	{
+		return IsSet();
+	}
+
+	/**
+	 * Has this ID been initialized to a valid element?
+	 */
+	FORCEINLINE bool IsSet() const
+	{
+		return TypeId != 0;
+	}
+
+	/**
+	 * Access the type ID portion of this element ID.
+	 */
+	FORCEINLINE FTypedHandleTypeId GetTypeId() const
+	{
+		return TypeId;
+	}
+
+	/**
+	 * Access the element ID portion of this element ID.
+	 */
+	FORCEINLINE FTypedHandleElementId GetElementId() const
+	{
+		return ElementId;
+	}
+
+	/**
+	 * Access the combined value of this element ID.
+	 * @note You typically don't want to store this directly as the element ID could be re-used.
+	 *       It is primarily useful as a secondary cache where something is keeping a reference to an element ID or element handle (eg, how FTypedElementList uses it internally).
+	 */
+	FORCEINLINE FTypedHandleCombinedId GetCombinedId() const
+	{
+		return CombinedId;
+	}
+
+	FORCEINLINE friend bool operator==(const FTypedElementId& InLHS, const FTypedElementId& InRHS)
+	{
+		return InLHS.CombinedId == InRHS.CombinedId;
+	}
+
+	FORCEINLINE friend bool operator!=(const FTypedElementId& InLHS, const FTypedElementId& InRHS)
+	{
+		return !(InLHS == InRHS);
+	}
+
+	friend inline uint32 GetTypeHash(const FTypedElementId& InElementId)
+	{
+		return GetTypeHash(InElementId.CombinedId);
+		
+	}
+
+	FORCEINLINE void Private_InitializeNoRef(const FTypedHandleTypeId InTypeId, const FTypedHandleElementId InElementId)
+	{
+		TypeId = InTypeId;
+		ElementId = InElementId;
+	}
+
+	FORCEINLINE void Private_DestroyNoRef()
+	{
+		CombinedId = 0;
+	}
+
+	/** An unset element ID */
+	static const FTypedElementId Unset;
+
+private:
+	union
+	{
+		struct
+		{
+			// Note: These are arranged in this order to give CombinedId better hash distribution for GetTypeHash!
+			FTypedHandleCombinedId ElementId : TypedHandleElementIdBits;
+			FTypedHandleCombinedId TypeId : TypedHandleTypeIdBits;
+		};
+		FTypedHandleCombinedId CombinedId;
+	};
+};
+
+/**
+ * A representation of an element that includes its handle data.
+ * This type is the most standard way that an element is passed through to interfaces, and also the type that is stored in element lists.
+ * C++ code may choose to use TTypedElement instead, which is a combination of an element handle and its associated element interface.
+ * @note Handles auto-release on destruction.
+ */
+USTRUCT(BlueprintType)
+struct TYPEDELEMENTFRAMEWORK_API FTypedElementHandle
+{
+	GENERATED_BODY()
+
+public:
+	FTypedElementHandle() = default;
+
+	FTypedElementHandle(const FTypedElementHandle& InOther)
+	{
+		if (InOther)
+		{
+			Private_InitializeAddRef(InOther.Id.GetTypeId(), InOther.Id.GetElementId(), *InOther.DataPtr);
+		}
+	}
+
+	FTypedElementHandle& operator=(const FTypedElementHandle& InOther)
+	{
+		if (this != &InOther)
+		{
+			Private_DestroyReleaseRef();
+
+			if (InOther)
+			{
+				Private_InitializeAddRef(InOther.Id.GetTypeId(), InOther.Id.GetElementId(), *InOther.DataPtr);
+			}
+		}
+		return *this;
+	}
+
+	FTypedElementHandle(FTypedElementHandle&& InOther)
+		: Id(MoveTemp(InOther.Id))
+		, DataPtr(InOther.DataPtr)
+	{
+		InOther.DataPtr = nullptr;
+		checkSlow(!InOther.IsSet());
+	}
+
+	FTypedElementHandle& operator=(FTypedElementHandle&& InOther)
+	{
+		if (this != &InOther)
+		{
+			Private_DestroyReleaseRef();
+
+			Id = MoveTemp(InOther.Id);
+			DataPtr = InOther.DataPtr;
+
+			InOther.DataPtr = nullptr;
+			checkSlow(!InOther.IsSet());
+		}
+		return *this;
+	}
+
+	~FTypedElementHandle()
+	{
+		Private_DestroyReleaseRef();
+	}
+
+	FORCEINLINE explicit operator bool() const
+	{
+		return IsSet();
+	}
+
+	/**
+	 * Has this handle been initialized to a valid element?
+	 */
+	FORCEINLINE bool IsSet() const
+	{
+		return Id.IsSet();
+	}
+
+	/**
+	 * Release this handle and set it back to an empty state.
+	 */
+	FORCEINLINE void Release()
+	{
+		Private_DestroyReleaseRef();
+	}
+
+	/**
+	 * Get the ID that this element represents.
+	 */
+	FORCEINLINE const FTypedElementId& GetId() const
+	{
+		return Id;
+	}
+
+	/**
+	 * Test to see whether the data stored within this handle is of the given type.
+	 * @note This is not typically something you'd want to query outside of data access within an interface implementation.
+	 */
+	template <typename ElementDataType>
+	FORCEINLINE bool IsDataOfType() const
+	{
+		return Id.GetTypeId() == ElementDataType::StaticTypeId();
+	}
+
+	/**
+	 * Attempt to access the data stored within this handle as the given type, returning null if it isn't possible.
+	 * @note This is not typically something you'd want to use outside of data access within an interface implementation.
+	 */
+	template <typename ElementDataType>
+	FORCEINLINE const ElementDataType* GetData() const
+	{
+		return IsDataOfType<ElementDataType>()
+			? static_cast<const ElementDataType*>(DataPtr->GetUntypedData())
+			: nullptr;
+	}
+
+	/**
+	 * Attempt to access the data stored within this handle as the given type, returning null if it isn't possible and logging an access error for scripting.
+	 * @note This is not typically something you'd want to use outside of data access within an interface implementation.
+	 */
+	template <typename ElementDataType>
+	const ElementDataType* GetDataScript() const
+	{
+		if (!DataPtr)
+		{
+			FFrame::KismetExecutionMessage(TEXT("Element handle data is null!"), ELogVerbosity::Error);
+			return nullptr;
+		}
+
+		if (!IsDataOfType<ElementDataType>())
+		{
+			FFrame::KismetExecutionMessage(*FString::Printf(TEXT("Element handle data type is '%d', but '%d' (%s) was requested!"), Id.GetTypeId(), ElementDataType::StaticTypeId(), *ElementDataType::StaticTypeName().ToString()), ELogVerbosity::Error);
+			return nullptr;
+		}
+
+		return static_cast<const ElementDataType*>(DataPtr->GetUntypedData());
+	}
+
+	/**
+	 * Attempt to access the data stored within this handle as the given type, asserting if it isn't possible.
+	 * @note This is not typically something you'd want to use outside of data access within an interface implementation.
+	 */
+	template <typename ElementDataType>
+	FORCEINLINE const ElementDataType& GetDataChecked() const
+	{
+		checkf(DataPtr, TEXT("Element handle data is null!"));
+		checkf(IsDataOfType<ElementDataType>(), TEXT("Element handle data type is '%d', but '%d' (%s) was requested!"), Id.GetTypeId(), ElementDataType::StaticTypeId(), *ElementDataType::StaticTypeName().ToString());
+		return *static_cast<const ElementDataType*>(DataPtr->GetUntypedData());
+	}
+	
+	/**
+	 * Acquire a copy of the ID that this element represents.
+	 * @note This must be paired with a call to ReleaseId.
+	 */
+	FTypedElementId AcquireId() const
+	{
+		FTypedElementId ElementId;
+		if (IsSet())
+		{
+			AddRef();
+			ElementId.Private_InitializeNoRef(Id.GetTypeId(), Id.GetElementId());
+		}
+		return ElementId;
+	}
+
+	/**
+	 * Release a copy of the ID that this element represents.
+	 * @note This should have come from a call to AcquireId.
+	 */
+	void ReleaseId(FTypedElementId& InOutElementId) const
+	{
+		checkf(InOutElementId == Id, TEXT("Element ID does not match this handle!"));
+		if (InOutElementId)
+		{
+			ReleaseRef();
+			InOutElementId.Private_DestroyNoRef();
+		}
+	}
+
+	FORCEINLINE friend bool operator==(const FTypedElementHandle& InLHS, const FTypedElementHandle& InRHS)
+	{
+		return InLHS.Id == InRHS.Id;
+	}
+
+	FORCEINLINE friend bool operator!=(const FTypedElementHandle& InLHS, const FTypedElementHandle& InRHS)
+	{
+		return !(InLHS == InRHS);
+	}
+
+	friend inline uint32 GetTypeHash(const FTypedElementHandle& InElementHandle)
+	{
+		return GetTypeHash(InElementHandle.Id);
+	}
+
+	FORCEINLINE void Private_InitializeNoRef(const FTypedHandleTypeId InTypeId, const FTypedHandleElementId InElementId, const FTypedElementInternalData& InData)
+	{
+		Id.Private_InitializeNoRef(InTypeId, InElementId);
+		DataPtr = &InData;
+	}
+
+	FORCEINLINE void Private_InitializeAddRef(const FTypedHandleTypeId InTypeId, const FTypedHandleElementId InElementId, const FTypedElementInternalData& InData)
+	{
+		Private_InitializeNoRef(InTypeId, InElementId, InData);
+		AddRef();
+	}
+
+	FORCEINLINE void Private_DestroyNoRef()
+	{
+		Id.Private_DestroyNoRef();
+		DataPtr = nullptr;
+	}
+
+	FORCEINLINE void Private_DestroyReleaseRef()
+	{
+		ReleaseRef();
+		Private_DestroyNoRef();
+	}
+
+	FORCEINLINE const FTypedElementInternalData* Private_GetInternalData() const
+	{
+		return DataPtr;
+	}
+
+private:
+	FORCEINLINE void AddRef() const
+	{
+#if WITH_TYPED_ELEMENT_REFCOUNT
+		if (DataPtr)
+		{
+			DataPtr->AddRef();
+		}
+#endif	// WITH_TYPED_ELEMENT_REFCOUNT
+	}
+
+	FORCEINLINE void ReleaseRef() const
+	{
+#if WITH_TYPED_ELEMENT_REFCOUNT
+		if (DataPtr)
+		{
+			DataPtr->ReleaseRef();
+		}
+#endif	// WITH_TYPED_ELEMENT_REFCOUNT
+	}
+
+	FTypedElementId Id;
+	const FTypedElementInternalData* DataPtr = nullptr;
+};
+
+/**
+ * Common implementation of TTypedElement that is inherited by all specializations.
+ */
+template <typename BaseInterfaceType>
+struct TTypedElementBase : public FTypedElementHandle
+{
+public:
+	TTypedElementBase() = default;
+
+	TTypedElementBase(const TTypedElementBase&) = default;
+	TTypedElementBase& operator=(const TTypedElementBase&) = default;
+
+	TTypedElementBase(TTypedElementBase&& InOther)
+		: FTypedElementHandle(MoveTemp(InOther))
+		, InterfacePtr(InOther.InterfacePtr)
+	{
+		InOther.InterfacePtr = nullptr;
+		checkSlow(!InOther.IsSet());
+	}
+
+	TTypedElementBase& operator=(TTypedElementBase&& InOther)
+	{
+		if (this != &InOther)
+		{
+			FTypedElementHandle::operator=(MoveTemp(InOther));
+			InterfacePtr = InOther.InterfacePtr;
+
+			InOther.InterfacePtr = nullptr;
+			checkSlow(!InOther.IsSet());
+		}
+		return *this;
+	}
+
+	~TTypedElementBase()
+	{
+		Private_DestroyReleaseRef();
+	}
+
+	FORCEINLINE explicit operator bool() const
+	{
+		return IsSet();
+	}
+
+	/**
+	 * Has this element been initialized to a valid handle and interface?
+	 */
+	FORCEINLINE bool IsSet() const
+	{
+		return FTypedElementHandle::IsSet()
+			&& InterfacePtr;
+	}
+
+	/**
+	 * Release this element and set it back to an empty state.
+	 */
+	FORCEINLINE void Release()
+	{
+		Private_DestroyReleaseRef();
+	}
+
+	/**
+	 * Test to see whether the interface stored within this element is of the given type.
+	 */
+	template <typename InterfaceType>
+	FORCEINLINE bool HasInterface() const
+	{
+		return HasInterface(InterfaceType::StaticClass());
+	}
+
+	/**
+	 * Test to see whether the interface stored within this element is of the given type.
+	 */
+	FORCEINLINE bool HasInterface(const UClass* InInterfaceType) const
+	{
+		return InterfacePtr
+			&& InterfacePtr->IsA(InInterfaceType);
+	}
+
+	/**
+	 * Attempt to access the interface stored within this element, returning null if it isn't set.
+	 */
+	FORCEINLINE BaseInterfaceType* GetInterface() const
+	{
+		return InterfacePtr;
+	}
+
+	/**
+	 * Attempt to access the interface stored within this element, asserting if it isn't set.
+	 */
+	FORCEINLINE BaseInterfaceType& GetInterfaceChecked() const
+	{
+		checkf(InterfacePtr, TEXT("Interface is null!"));
+		return *InterfacePtr;
+	}
+
+	FORCEINLINE friend bool operator==(const TTypedElementBase& InLHS, const TTypedElementBase& InRHS)
+	{
+		return static_cast<const FTypedElementHandle&>(InLHS) == static_cast<const FTypedElementHandle&>(InRHS)
+			&& InLHS.InterfacePtr == InRHS.InterfacePtr;
+	}
+
+	FORCEINLINE friend bool operator!=(const TTypedElementBase& InLHS, const TTypedElementBase& InRHS)
+	{
+		return !(InLHS == InRHS);
+	}
+
+	friend inline uint32 GetTypeHash(const TTypedElementBase& InElement)
+	{
+		return HashCombine(GetTypeHash(static_cast<const FTypedElementHandle&>(InElement)), GetTypeHash(InElement.InterfacePtr));
+	}
+
+	FORCEINLINE void Private_InitializeNoRef(const FTypedHandleTypeId InTypeId, const FTypedHandleElementId InElementId, const FTypedElementInternalData& InData, BaseInterfaceType* InInterfacePtr)
+	{
+		FTypedElementHandle::Private_InitializeNoRef(InTypeId, InElementId, InData);
+		InterfacePtr = InInterfacePtr;
+	}
+
+	FORCEINLINE void Private_InitializeAddRef(const FTypedHandleTypeId InTypeId, const FTypedHandleElementId InElementId, const FTypedElementInternalData& InData, BaseInterfaceType* InInterfacePtr)
+	{
+		FTypedElementHandle::Private_InitializeAddRef(InTypeId, InElementId, InData);
+		InterfacePtr = InInterfacePtr;
+	}
+
+	FORCEINLINE void Private_DestroyNoRef()
+	{
+		FTypedElementHandle::Private_DestroyNoRef();
+		InterfacePtr = nullptr;
+	}
+
+	FORCEINLINE void Private_DestroyReleaseRef()
+	{
+		FTypedElementHandle::Private_DestroyReleaseRef();
+		InterfacePtr = nullptr;
+	}
+
+protected:
+	BaseInterfaceType* InterfacePtr = nullptr;
+};
+
+/**
+ * A combination of an element handle and its associated element interface.
+ * @note This should be specialized for top-level element interfaces to include their interface API.
+ * @note Elements auto-release on destruction.
+ */
+template <typename BaseInterfaceType>
+struct TTypedElement : public TTypedElementBase<BaseInterfaceType>
+{
+};
+using FTypedElement = TTypedElement<UTypedElementInterface>;
+
+template <typename OtherInterfaceType, typename ThisInterfaceType>
+FORCEINLINE TTypedElement<OtherInterfaceType> CastTypedElement(const TTypedElement<ThisInterfaceType>& InElement)
+{
+	static_assert(sizeof(TTypedElement<ThisInterfaceType>) == sizeof(TTypedElement<OtherInterfaceType>), "All TTypedElement instances must be the same size for this cast implementation to work!");
+	if (InElement && InElement.template HasInterface<OtherInterfaceType>())
+	{
+		return reinterpret_cast<const TTypedElement<OtherInterfaceType>&>(InElement);
+	}
+	return TTypedElement<OtherInterfaceType>();
+}
+
+template <typename OtherInterfaceType, typename ThisInterfaceType>
+FORCEINLINE void CastTypedElement(const TTypedElement<ThisInterfaceType>& InElement, TTypedElement<OtherInterfaceType>& OutCastElement)
+{
+	static_assert(sizeof(TTypedElement<ThisInterfaceType>) == sizeof(TTypedElement<OtherInterfaceType>), "All TTypedElement instances must be the same size for this cast implementation to work!");
+	if (InElement && InElement.template HasInterface<OtherInterfaceType>())
+	{
+		OutCastElement = reinterpret_cast<const TTypedElement<OtherInterfaceType>&>(InElement);
+	}
+	else
+	{
+		OutCastElement.Private_DestroyReleaseRef();
+	}
+}
+
+template <typename OtherInterfaceType, typename ThisInterfaceType>
+FORCEINLINE TTypedElement<OtherInterfaceType> CastTypedElement(TTypedElement<ThisInterfaceType>&& InElement)
+{
+	static_assert(sizeof(TTypedElement<ThisInterfaceType>) == sizeof(TTypedElement<OtherInterfaceType>), "All TTypedElement instances must be the same size for this cast implementation to work!");
+	if (InElement && InElement.template HasInterface<OtherInterfaceType>())
+	{
+		return reinterpret_cast<TTypedElement<OtherInterfaceType>&&>(InElement);
+	}
+	return TTypedElement<OtherInterfaceType>();
+}
+
+template <typename OtherInterfaceType, typename ThisInterfaceType>
+FORCEINLINE void CastTypedElement(TTypedElement<ThisInterfaceType>&& InElement, TTypedElement<OtherInterfaceType>& OutCastElement)
+{
+	static_assert(sizeof(TTypedElement<ThisInterfaceType>) == sizeof(TTypedElement<OtherInterfaceType>), "All TTypedElement instances must be the same size for this cast implementation to work!");
+	if (InElement && InElement.template HasInterface<OtherInterfaceType>())
+	{
+		OutCastElement = reinterpret_cast<TTypedElement<OtherInterfaceType>&&>(InElement);
+	}
+	else
+	{
+		OutCastElement.Private_DestroyReleaseRef();
+	}
+}
+
+template <typename OtherInterfaceType, typename ThisInterfaceType>
+FORCEINLINE TTypedElement<OtherInterfaceType> CastTypedElementChecked(const TTypedElement<ThisInterfaceType>& InElement)
+{
+	static_assert(sizeof(TTypedElement<ThisInterfaceType>) == sizeof(TTypedElement<OtherInterfaceType>), "All TTypedElement instances must be the same size for this cast implementation to work!");
+	if (InElement)
+	{
+		checkf(InElement.template HasInterface<OtherInterfaceType>(), TEXT("Element does not implement the required interface for this cast!"));
+		return reinterpret_cast<const TTypedElement<OtherInterfaceType>&>(InElement);
+	}
+	return TTypedElement<OtherInterfaceType>();
+}
+
+template <typename OtherInterfaceType, typename ThisInterfaceType>
+FORCEINLINE void CastTypedElementChecked(const TTypedElement<ThisInterfaceType>& InElement, TTypedElement<OtherInterfaceType>& OutCastElement)
+{
+	static_assert(sizeof(TTypedElement<ThisInterfaceType>) == sizeof(TTypedElement<OtherInterfaceType>), "All TTypedElement instances must be the same size for this cast implementation to work!");
+	if (InElement)
+	{
+		checkf(InElement.template HasInterface<OtherInterfaceType>(), TEXT("Element does not implement the required interface for this cast!"));
+		OutCastElement = reinterpret_cast<const TTypedElement<OtherInterfaceType>&>(InElement);
+	}
+	else
+	{
+		OutCastElement.Private_DestroyReleaseRef();
+	}
+}
+
+template <typename OtherInterfaceType, typename ThisInterfaceType>
+FORCEINLINE TTypedElement<OtherInterfaceType> CastTypedElementChecked(TTypedElement<ThisInterfaceType>&& InElement)
+{
+	static_assert(sizeof(TTypedElement<ThisInterfaceType>) == sizeof(TTypedElement<OtherInterfaceType>), "All TTypedElement instances must be the same size for this cast implementation to work!");
+	if (InElement)
+	{
+		checkf(InElement.template HasInterface<OtherInterfaceType>(), TEXT("Element does not implement the required interface for this cast!"));
+		return reinterpret_cast<TTypedElement<OtherInterfaceType>&&>(InElement);
+	}
+	return TTypedElement<OtherInterfaceType>();
+}
+
+template <typename OtherInterfaceType, typename ThisInterfaceType>
+FORCEINLINE void CastTypedElementChecked(TTypedElement<ThisInterfaceType>&& InElement, TTypedElement<OtherInterfaceType>& OutCastElement)
+{
+	static_assert(sizeof(TTypedElement<ThisInterfaceType>) == sizeof(TTypedElement<OtherInterfaceType>), "All TTypedElement instances must be the same size for this cast implementation to work!");
+	if (InElement)
+	{
+		checkf(InElement.template HasInterface<OtherInterfaceType>(), TEXT("Element does not implement the required interface for this cast!"));
+		OutCastElement = reinterpret_cast<TTypedElement<OtherInterfaceType>&&>(InElement);
+	}
+	else
+	{
+		OutCastElement.Private_DestroyReleaseRef();
+	}
+}
+
+/**
+ * A representation of the owner of an element that includes its mutable handle data.
+ * This type is returned when creating an element, and should be used to populate its internal payload data (if any).
+ * @note Owners do not auto-release on destruction, and must be manually destroyed via their owner element registry.
+ */
+template <typename ElementDataType>
+struct TTypedElementOwner
+{
+public:
+	TTypedElementOwner() = default;
+
+	TTypedElementOwner(const TTypedElementOwner&) = delete;
+	TTypedElementOwner& operator=(const TTypedElementOwner&) = delete;
+
+	TTypedElementOwner(TTypedElementOwner&& InOther)
+		: Id(MoveTemp(InOther.Id))
+		, DataPtr(InOther.DataPtr)
+	{
+		InOther.DataPtr = nullptr;
+		checkSlow(!InOther.IsSet());
+	}
+
+	TTypedElementOwner& operator=(TTypedElementOwner&& InOther)
+	{
+		if (this != &InOther)
+		{
+			Id = MoveTemp(InOther.Id);
+			DataPtr = InOther.DataPtr;
+
+			InOther.DataPtr = nullptr;
+			checkSlow(!InOther.IsSet());
+		}
+		return *this;
+	}
+
+	~TTypedElementOwner()
+	{
+		checkf(!IsSet(), TEXT("Element owner was still set during destruction! This will leak an element, and you should destroy this element prior to destruction!"));
+	}
+
+	FORCEINLINE explicit operator bool() const
+	{
+		return IsSet();
+	}
+
+	/**
+	 * Has this owner been initialized to a valid element?
+	 */
+	FORCEINLINE bool IsSet() const
+	{
+		return Id.IsSet();
+	}
+
+	/**
+	 * Get the ID that this element represents.
+	 */
+	FORCEINLINE const FTypedElementId& GetId() const
+	{
+		return Id;
+	}
+
+	/**
+	 * Attempt to access the mutable data stored within this owner, returning null if it isn't possible.
+	 */
+	template <typename U = ElementDataType, std::enable_if_t<!std::is_void<U>::value>* = nullptr>
+	FORCEINLINE U* GetData() const
+	{
+		static_assert(std::is_same<U, ElementDataType>::value, "Don't explicitly specify U!");
+		return DataPtr
+			? &DataPtr->GetMutableData()
+			: nullptr;
+	}
+
+	/**
+	 * Attempt to access the mutable data stored within this owner, asserting if it isn't possible.
+	 */
+	template <typename U = ElementDataType, std::enable_if_t<!std::is_void<U>::value>* = nullptr>
+	FORCEINLINE U& GetDataChecked() const
+	{
+		static_assert(std::is_same<U, ElementDataType>::value, "Don't explicitly specify U!");
+		checkf(DataPtr, TEXT("Handle data is null!"));
+		return DataPtr->GetMutableData();
+	}
+
+	/**
+	 * Acquire a copy of the ID that this element represents.
+	 * @note This must be paired with a call to ReleaseId.
+	 */
+	FTypedElementId AcquireId() const
+	{
+		FTypedElementId ElementId;
+		if (IsSet())
+		{
+			AddRef();
+			ElementId.Private_InitializeNoRef(Id.GetTypeId(), Id.GetElementId());
+		}
+		return ElementId;
+	}
+
+	/**
+	 * Release a copy of the ID that this element represents.
+	 * @note This should have come from a call to AcquireId.
+	 */
+	void ReleaseId(FTypedElementId& InOutElementId) const
+	{
+		checkf(InOutElementId == Id, TEXT("Element ID does not match this owner!"));
+		if (InOutElementId)
+		{
+			ReleaseRef();
+			InOutElementId.Private_DestroyNoRef();
+		}
+	}
+
+	/**
+	 * Acquire a copy of the handle that this element represents.
+	 * @note This must be paired with a call to ReleaseHandle (or a call to Release on the handle instance).
+	 */
+	FTypedElementHandle AcquireHandle() const
+	{
+		FTypedElementHandle ElementHandle;
+		if (IsSet())
+		{
+			ElementHandle.Private_InitializeAddRef(Id.GetTypeId(), Id.GetElementId(), *DataPtr);
+		}
+		return ElementHandle;
+	}
+
+	/**
+	 * Release a copy of the handle that this element represents.
+	 * @note This should have come from a call to AcquireHandle.
+	 */
+	void ReleaseHandle(FTypedElementHandle& InOutElementHandle) const
+	{
+		checkf(InOutElementHandle.GetId() == Id, TEXT("Element handle ID does not match this owner!"));
+		InOutElementHandle.Release();
+	}
+
+	FORCEINLINE friend bool operator==(const TTypedElementOwner& InLHS, const TTypedElementOwner& InRHS)
+	{
+		return InLHS.Id == InRHS.Id;
+	}
+
+	FORCEINLINE friend bool operator!=(const TTypedElementOwner& InLHS, const TTypedElementOwner& InRHS)
+	{
+		return !(InLHS == InRHS);
+	}
+
+	friend inline uint32 GetTypeHash(const TTypedElementOwner& InElementOwner)
+	{
+		return GetTypeHash(InElementOwner.Id);
+	}
+
+	FORCEINLINE void Private_InitializeNoRef(const FTypedHandleTypeId InTypeId, const FTypedHandleElementId InElementId, TTypedElementInternalData<ElementDataType>& InData)
+	{
+		Id.Private_InitializeNoRef(InTypeId, InElementId);
+		DataPtr = &InData;
+	}
+
+	FORCEINLINE void Private_InitializeAddRef(const FTypedHandleTypeId InTypeId, const FTypedHandleElementId InElementId, TTypedElementInternalData<ElementDataType>& InData)
+	{
+		Private_InitializeNoRef(InTypeId, InElementId, InData);
+		AddRef();
+	}
+
+	FORCEINLINE void Private_DestroyNoRef()
+	{
+		Id.Private_DestroyNoRef();
+		DataPtr = nullptr;
+	}
+
+	FORCEINLINE void Private_DestroyReleaseRef()
+	{
+		ReleaseRef();
+		Private_DestroyNoRef();
+	}
+
+	FORCEINLINE const TTypedElementInternalData<ElementDataType>* Private_GetInternalData() const
+	{
+		return DataPtr;
+	}
+
+private:
+	FORCEINLINE void AddRef() const
+	{
+#if WITH_TYPED_ELEMENT_REFCOUNT
+		if (DataPtr)
+		{
+			DataPtr->AddRef();
+		}
+#endif	// WITH_TYPED_ELEMENT_REFCOUNT
+	}
+
+	FORCEINLINE void ReleaseRef() const
+	{
+#if WITH_TYPED_ELEMENT_REFCOUNT
+		if (DataPtr)
+		{
+			DataPtr->ReleaseRef();
+		}
+#endif	// WITH_TYPED_ELEMENT_REFCOUNT
+	}
+
+	FTypedElementId Id;
+	TTypedElementInternalData<ElementDataType>* DataPtr = nullptr;
+};
+using FTypedElementOwner = TTypedElementOwner<void>;
diff --git a/Engine/Source/Runtime/TypedElementFramework/Public/TypedElementLimits.h b/Engine/Source/Runtime/TypedElementFramework/Public/TypedElementLimits.h
new file mode 100644
--- /dev/null
+++ b/Engine/Source/Runtime/TypedElementFramework/Public/TypedElementLimits.h
@@ -0,0 +1,34 @@
+// Copyright Epic Games, Inc. All Rights Reserved.
+
+#pragma once
+
+#include "CoreTypes.h"
+#include "Math/NumericLimits.h"
+
+#define WITH_TYPED_ELEMENT_REFCOUNT (1)
+
+/**
+ * Handle ID limits, as used by FTypedElementId.
+ * @note Limited to a combined 32-bits so that they can be used directly within render targets, though could be made 64-bits if the 
+ *       editor used 64-bit render targets (this would also require 64-bit container support in TTypedElementInternalDataStore).
+ */
+constexpr SIZE_T TypedHandleTypeIdBits = 8;
+constexpr SIZE_T TypedHandleElementIdBits = 24;
+
+constexpr SIZE_T TypedHandleTypeIdBytes = TypedHandleTypeIdBits >> 3;
+constexpr SIZE_T TypedHandleElementIdBytes = TypedHandleElementIdBits >> 3;
+
+constexpr SIZE_T TypedHandleMaxTypeId = ((SIZE_T)1 << TypedHandleTypeIdBits) - 1;
+constexpr SIZE_T TypedHandleMaxElementId = ((SIZE_T)1 << TypedHandleElementIdBits) - 1;
+
+using FTypedHandleTypeId = uint8;
+using FTypedHandleElementId = int32;
+using FTypedHandleCombinedId = uint32;
+#if WITH_TYPED_ELEMENT_REFCOUNT
+using FTypedHandleRefCount = int32;
+#endif	// WITH_TYPED_ELEMENT_REFCOUNT
+
+static_assert(sizeof(FTypedHandleCombinedId) >= (TypedHandleTypeIdBytes + TypedHandleElementIdBytes), "FTypedHandleCombinedId is not large enough to hold the combination of TypedHandleTypeIdBytes and TypedHandleElementIdBytes!");
+
+static_assert(TNumericLimits<FTypedHandleTypeId>::Max() >= TypedHandleMaxTypeId, "FTypedHandleTypeId is not large enough to hold TypedHandleMaxTypeId!");
+static_assert(TNumericLimits<FTypedHandleElementId>::Max() >= TypedHandleMaxElementId, "FTypedHandleElementId is not large enough to hold TypedHandleMaxElementId!");
diff --git a/Engine/Source/Runtime/TypedElementFramework/Public/TypedElementList.h b/Engine/Source/Runtime/TypedElementFramework/Public/TypedElementList.h
new file mode 100644
--- /dev/null
+++ b/Engine/Source/Runtime/TypedElementFramework/Public/TypedElementList.h
@@ -0,0 +1,580 @@
+// Copyright Epic Games, Inc. All Rights Reserved.
+
+#pragma once
+
+#include "CoreMinimal.h"
+#include "Delegates/Delegate.h"
+#include "Delegates/DelegateCombinations.h"
+#include "Containers/ArrayView.h"
+#include "TypedElementHandle.h"
+#include "TypedElementListPtr.h"
+
+class UTypedElementRegistry;
+
+namespace TypedElementList_Private
+{
+
+TYPEDELEMENTFRAMEWORK_API void GetElementImpl(const UTypedElementRegistry* InRegistry, const FTypedElementHandle& InElementHandle, const UClass* InBaseInterfaceType, FTypedElement& OutElement);
+
+template <typename BaseInterfaceType>
+FORCEINLINE void GetElement(const UTypedElementRegistry* InRegistry, const FTypedElementHandle& InElementHandle, TTypedElement<BaseInterfaceType>& OutElement)
+{
+	static_assert(sizeof(TTypedElement<BaseInterfaceType>) == sizeof(FTypedElement), "All TTypedElement instances must be the same size for this cast implementation to work!");
+	GetElementImpl(InRegistry, InElementHandle, BaseInterfaceType::StaticClass(), reinterpret_cast<FTypedElement&>(OutElement));
+}
+
+template <typename BaseInterfaceType>
+FORCEINLINE TTypedElement<BaseInterfaceType> GetElement(const UTypedElementRegistry* InRegistry, const FTypedElementHandle& InElementHandle)
+{
+	TTypedElement<BaseInterfaceType> Element;
+	GetElement(InRegistry, InElementHandle, Element);
+	return Element;
+}
+
+/**
+ * Iterator type for ranged-for loops which only returns handles that implement the given interface.
+ * @note This iterator type only supports the minimal functionality needed to support C++ ranged-for syntax.
+ */
+template <typename BaseInterfaceType>
+struct TTypedElementListInterfaceIterator
+{
+public:
+	explicit TTypedElementListInterfaceIterator(const UTypedElementRegistry* InRegistry, const TArray<FTypedElementHandle>& InElementHandles, const int32 InIndex)
+		: Registry(InRegistry)
+		, ElementHandles(InElementHandles)
+		, InitialNum(InElementHandles.Num())
+		, Index(InIndex)
+	{
+		AdvanceIterator();
+	}
+
+	FORCEINLINE const TTypedElement<BaseInterfaceType>& operator*() const
+	{
+		return Element;
+	}
+
+	FORCEINLINE TTypedElementListInterfaceIterator& operator++()
+	{
+		AdvanceIterator();
+		return *this;
+	}
+
+private:
+	const UTypedElementRegistry* Registry;
+	const TArray<FTypedElementHandle>& ElementHandles;
+	const int32 InitialNum;
+	int32 Index;
+	TTypedElement<BaseInterfaceType> Element;
+
+	void AdvanceIterator()
+	{
+		while (Index < ElementHandles.Num())
+		{
+			TypedElementList_Private::GetElement(Registry, ElementHandles[Index++], Element);
+			if (Element)
+			{
+				break;
+			}
+		}
+	}
+
+	FORCEINLINE friend bool operator!=(const TTypedElementListInterfaceIterator& Lhs, const TTypedElementListInterfaceIterator& Rhs)
+	{
+		// We only need to do the check in this operator, because no other operator will be
+		// called until after this one returns.
+		//
+		// Also, we should only need to check one side of this comparison - if the other iterator isn't
+		// even from the same array then the compiler has generated bad code.
+		ensureMsgf(Lhs.ElementHandles.Num() == Lhs.InitialNum, TEXT("Array has changed during ranged-for iteration!"));
+		return Lhs.Index != Rhs.Index;
+	}
+};
+
+/**
+ * Proxy type to enable ranged-for loops which only returns handles that implement the given interface.
+ */
+template <typename BaseInterfaceType>
+struct TTypedElementListInterfaceIteratorProxy
+{
+public:
+	explicit TTypedElementListInterfaceIteratorProxy(const UTypedElementRegistry* InRegistry, const TArray<FTypedElementHandle>& InElementHandles)
+		: Registry(InRegistry)
+		, ElementHandles(InElementHandles)
+	{
+	}
+
+	TTypedElementListInterfaceIteratorProxy(const TTypedElementListInterfaceIteratorProxy&) = delete;
+	TTypedElementListInterfaceIteratorProxy& operator=(const TTypedElementListInterfaceIteratorProxy&) = delete;
+
+	TTypedElementListInterfaceIteratorProxy(TTypedElementListInterfaceIteratorProxy&&) = default;
+	TTypedElementListInterfaceIteratorProxy& operator=(TTypedElementListInterfaceIteratorProxy&&) = default;
+
+	/**
+	 * DO NOT USE DIRECTLY
+	 * STL-like iterators to enable range-based for loop support.
+	 */
+	using RangedForConstIteratorType = TTypedElementListInterfaceIterator<BaseInterfaceType>;
+	FORCEINLINE RangedForConstIteratorType begin() const { return RangedForConstIteratorType(Registry, ElementHandles, 0); }
+	FORCEINLINE RangedForConstIteratorType end() const { return RangedForConstIteratorType(Registry, ElementHandles, ElementHandles.Num()); }
+
+private:
+	const UTypedElementRegistry* Registry;
+	const TArray<FTypedElementHandle>& ElementHandles;
+};
+
+} // namespace TypedElementList_Private
+
+/**
+ * Interface to allow external systems (such as USelection) to receive immediate sync notifications as an element list is changed.
+ * This exists purely as a bridging mechanism and shouldn't be relied on for new code. It is lazily created as needed.
+ */
+class TYPEDELEMENTFRAMEWORK_API FTypedElementListLegacySync
+{
+public:
+	enum class ESyncType : uint8
+	{
+		/**
+		 * An element was added to the element list.
+		 * The ElementHandle argument will be set to the element that was added.
+		 */
+		Added,
+
+		/**
+		 * An element was removed from the element list.
+		 * The ElementHandle argument will be set to the element that was removed.
+		 */
+		Removed,
+
+		/**
+		 * The element list was modified in an unknown way.
+		 * The ElementHandle argument will be unset.
+		 */
+		Modified,
+
+		/**
+		 * The element list was cleared.
+		 * The ElementHandle argument will be unset.
+		 */
+		Cleared,
+
+		/**
+		 * The element list was modified as part of a batch or bulk operation.
+		 * The ElementHandle argument will be unset.
+		 * @note A batch operation will emit internal (bIsWithinBatchOperation=true) Added, Removed, Modified and Cleared updates during the batch, 
+		 *       so if you respond to those internal updates you may choose to ignore this one. Otherwise you should treat it the same as Modified.
+		 */
+		BatchComplete,
+	};
+	
+	FTypedElementListLegacySync(const FTypedElementList& InElementList);
+
+	void Private_EmitSyncEvent(const ESyncType InSyncType, const FTypedElementHandle& InElementHandle = FTypedElementHandle());
+
+	DECLARE_EVENT_FourParams(FTypedElementListLegacySync, FOnSyncEvent, const FTypedElementList& /*InElementList*/, ESyncType /*InSyncType*/, const FTypedElementHandle& /*InElementHandle*/, bool /*bIsWithinBatchOperation*/);
+	FOnSyncEvent& OnSyncEvent();
+
+	bool IsRunningBatchOperation() const;
+	void BeginBatchOperation();
+	void EndBatchOperation(const bool InNotify = true);
+	bool IsBatchOperationDirty() const;
+	void ForceBatchOperationDirty();
+
+private:
+	const FTypedElementList& ElementList;
+
+	FOnSyncEvent OnSyncEventDelegate;
+
+	int32 NumOpenBatchOperations = 0;
+	bool bBatchOperationIsDirty = false;
+};
+
+/**
+ * A list of element handles.
+ * Provides high-level access to groups of elements, including accessing elements that implement specific interfaces.
+ */
+class TYPEDELEMENTFRAMEWORK_API FTypedElementList
+{
+public:
+	FTypedElementList(const FTypedElementList&) = delete;
+	FTypedElementList& operator=(const FTypedElementList&) = delete;
+
+	FTypedElementList(FTypedElementList&&) = delete;
+	FTypedElementList& operator=(FTypedElementList&&) = delete;
+
+	~FTypedElementList();
+
+	static FTypedElementListPtr Private_CreateElementList(UTypedElementRegistry* InRegistry);
+
+	/**
+	 * Clone this list instance.
+	 * @note Only copies elements; does not copy any bindings!
+	 */
+	FTypedElementListPtr Clone() const;
+
+	/**
+	 * Get the element handle at the given index.
+	 * @note Use IsValidIndex to test for validity.
+	 */
+	FORCEINLINE const FTypedElementHandle& operator[](const int32 InIndex) const
+	{
+		return GetElementHandleAt(InIndex);
+	}
+
+	/**
+	 * Get the element handle at the given index.
+	 * @note Use IsValidIndex to test for validity.
+	 */
+	FORCEINLINE const FTypedElementHandle& GetElementHandleAt(const int32 InIndex) const
+	{
+		return ElementHandles[InIndex];
+	}
+
+	/**
+	 * Get the element at the given index.
+	 * @note Use IsValidIndex to test for validity.
+	 */
+	template <typename BaseInterfaceType>
+	FORCEINLINE TTypedElement<BaseInterfaceType> GetElementAt(const int32 InIndex) const
+	{
+		return GetElement<BaseInterfaceType>(GetElementHandleAt(InIndex));
+	}
+
+	/**
+	 * Get the element at the given index.
+	 * @note Use IsValidIndex to test for validity.
+	 */
+	template <typename BaseInterfaceType>
+	FORCEINLINE void GetElementAt(const int32 InIndex, TTypedElement<BaseInterfaceType>& OutElement) const
+	{
+		GetElement(GetElementHandleAt(InIndex), OutElement);
+	}
+
+	/**
+	 * Get the element from the given handle.
+	 */
+	template <typename BaseInterfaceType>
+	FORCEINLINE TTypedElement<BaseInterfaceType> GetElement(const FTypedElementHandle& InElementHandle) const
+	{
+		return TypedElementList_Private::GetElement<BaseInterfaceType>(Registry, InElementHandle);
+	}
+
+	/**
+	 * Get the element from the given handle.
+	 */
+	template <typename BaseInterfaceType>
+	FORCEINLINE void GetElement(const FTypedElementHandle& InElementHandle, TTypedElement<BaseInterfaceType>& OutElement) const
+	{
+		TypedElementList_Private::GetElement(Registry, InElementHandle, OutElement);
+	}
+
+	/**
+	 * Is the given index a valid entry within this element list?
+	 */
+	FORCEINLINE bool IsValidIndex(const int32 InIndex) const
+	{
+		return ElementHandles.IsValidIndex(InIndex);
+	}
+
+	/**
+	 * Get the number of entries within this element list.
+	 */
+	FORCEINLINE int32 Num() const
+	{
+		return ElementHandles.Num();
+	}
+
+	/**
+	 * Shrink this element list storage to avoid slack.
+	 */
+	FORCEINLINE void Shrink()
+	{
+		ElementCombinedIds.Shrink();
+		ElementHandles.Shrink();
+	}
+
+	/**
+	 * Pre-allocate enough memory in this element list to store the given number of entries.
+	 */
+	FORCEINLINE void Reserve(const int32 InSize)
+	{
+		ElementCombinedIds.Reserve(InSize);
+		ElementHandles.Reserve(InSize);
+	}
+
+	/**
+	 * Remove all entries from this element list, potentially leaving space allocated for the given number of entries.
+	 */
+	FORCEINLINE void Empty(const int32 InSlack = 0)
+	{
+		ElementCombinedIds.Empty(InSlack);
+		ElementHandles.Empty(InSlack);
+		NoteListChanged(EChangeType::Cleared);
+	}
+
+	/**
+	 * Remove all entries from this element list, preserving existing allocations.
+	 */
+	FORCEINLINE void Reset()
+	{
+		ElementCombinedIds.Reset();
+		ElementHandles.Reset();
+		NoteListChanged(EChangeType::Cleared);
+	}
+
+	/**
+	 * Does this element list contain an entry for the given element ID?
+	 */
+	FORCEINLINE bool Contains(const FTypedElementId& InElementId) const
+	{
+		return ContainsElementImpl(InElementId);
+	}
+
+	/**
+	 * Does this element list contain an entry for the given element handle?
+	 */
+	FORCEINLINE bool Contains(const FTypedElementHandle& InElementHandle) const
+	{
+		return ContainsElementImpl(InElementHandle.GetId());
+	}
+
+	/**
+	 * Does this element list contain an entry for the given element owner?
+	 */
+	template <typename ElementDataType>
+	FORCEINLINE bool Contains(const TTypedElementOwner<ElementDataType>& InElementOwner)
+	{
+		return ContainsElementImpl(InElementOwner.GetId());
+	}
+
+	/**
+	 * Add the given element handle to this element list, if it isn't already in the list.
+	 * @return True if the element handle was added, false if it is already in the list.
+	 */
+	FORCEINLINE bool Add(const FTypedElementHandle& InElementHandle)
+	{
+		return AddElementImpl(CopyTemp(InElementHandle));
+	}
+
+	/**
+	 * Add the given element handle to this element list, if it isn't already in the list.
+	 * @return True if the element handle was added, false if it is already in the list.
+	 */
+	FORCEINLINE bool Add(FTypedElementHandle&& InElementHandle)
+	{
+		return AddElementImpl(MoveTemp(InElementHandle));
+	}
+
+	/**
+	 * Add the given element owner to this element list, if it isn't already in the list.
+	 * @return True if the element owner was added, false if it is already in the list.
+	 */
+	template <typename ElementDataType>
+	FORCEINLINE bool Add(const TTypedElementOwner<ElementDataType>& InElementOwner)
+	{
+		return AddElementImpl(InElementOwner.AcquireHandle());
+	}
+
+	/**
+	 * Append the given element handles to this element list, for any that already in the list.
+	 */
+	void Append(TArrayView<const FTypedElementHandle> InElementHandles)
+	{
+		if (LegacySync)
+		{
+			LegacySync->BeginBatchOperation();
+		}
+
+		Reserve(Num() + InElementHandles.Num());
+		for (const FTypedElementHandle& ElementHandle : InElementHandles)
+		{
+			AddElementImpl(CopyTemp(ElementHandle));
+		}
+
+		if (LegacySync)
+		{
+			LegacySync->EndBatchOperation();
+		}
+	}
+
+	/**
+	 * Append the given element owners to this element list, for any that already in the list.
+	 */
+	template <typename ElementDataType>
+	FORCEINLINE void Append(const TArray<TTypedElementOwner<ElementDataType>>& InElementOwners)
+	{
+		Append(MakeArrayView(InElementOwners));
+	}
+
+	/**
+	 * Append the given element owners to this element list, for any that already in the list.
+	 */
+	template <typename ElementDataType>
+	void Append(TArrayView<const TTypedElementOwner<ElementDataType>> InElementOwners)
+	{
+		if (LegacySync)
+		{
+			LegacySync->BeginBatchOperation();
+		}
+
+		Reserve(Num() + InElementOwners.Num());
+		for (const TTypedElementOwner<ElementDataType>& ElementOwner : InElementOwners)
+		{
+			AddElementImpl(ElementOwner.AcquireHandle());
+		}
+
+		if (LegacySync)
+		{
+			LegacySync->EndBatchOperation();
+		}
+	}
+
+	/**
+	 * Remove the given element ID from this element list, if it is in the list.
+	 * @return True if the element ID was removed, false if it isn't in the list.
+	 */
+	FORCEINLINE bool Remove(const FTypedElementId& InElementId)
+	{
+		return RemoveElementImpl(InElementId);
+	}
+
+	/**
+	 * Remove the given element handle from this element list, if it is in the list.
+	 * @return True if the element handle was removed, false if it isn't in the list.
+	 */
+	FORCEINLINE bool Remove(const FTypedElementHandle& InElementHandle)
+	{
+		return RemoveElementImpl(InElementHandle.GetId());
+	}
+
+	/**
+	 * Remove the given element owner from this element list, if it is in the list.
+	 * @return True if the element owner was removed, false if it isn't in the list.
+	 */
+	template <typename ElementDataType>
+	FORCEINLINE bool Remove(const TTypedElementOwner<ElementDataType>& InElementOwner)
+	{
+		return RemoveElementImpl(InElementOwner.GetId());
+	}
+
+	/**
+	 * Remove any element handles that match the given predicate from this element list.
+	 * @return The number of element handles removed.
+	 */
+	FORCEINLINE int32 RemoveAll(TFunctionRef<bool(const FTypedElementHandle&)> InPredicate)
+	{
+		return RemoveAllElementsImpl(InPredicate);
+	}
+
+	/**
+	 * Remove any elements that match the given predicate from this element list.
+	 * @return The number of elements removed.
+	 */
+	template <typename BaseInterfaceType>
+	int32 RemoveAll(TFunctionRef<bool(const TTypedElement<BaseInterfaceType>&)> InPredicate)
+	{
+		TTypedElement<BaseInterfaceType> TempElement;
+		return RemoveAllElementsImpl([this, &TempElement, &InPredicate](const FTypedElementHandle& InElementHandle)
+		{
+			GetElement(InElementHandle, TempElement);
+			return TempElement && InPredicate(TempElement);
+		});
+	}
+
+	/**
+	 * Access the delegate that is invoked whenever this element list has been changed.
+	 * @note This is called automatically at the end of each frame, but can also be manually invoked by NotifyPendingChanges.
+	 */
+	DECLARE_EVENT_OneParam(FTypedElementList, FOnChanged, const FTypedElementList& /*InElementList*/);
+	FOnChanged& OnChanged()
+	{
+		return OnChangedDelegate;
+	}
+
+	/**
+	 * Invoke the delegate called whenever this element list has been changed.
+	 */
+	void NotifyPendingChanges();
+
+	/**
+	 * Access the interface to allow external systems (such as USelection) to receive immediate sync notifications as an element list is changed.
+	 * This exists purely as a bridging mechanism and shouldn't be relied on for new code. It is lazily created as needed.
+	 */
+	FTypedElementListLegacySync& Legacy_GetSync();
+
+	/**
+	 * Get an iterator for elements that implement the given element interface.
+	 */
+	template <typename BaseInterfaceType>
+	FORCEINLINE TypedElementList_Private::TTypedElementListInterfaceIteratorProxy<BaseInterfaceType> InterateInterface() const
+	{
+		return TypedElementList_Private::TTypedElementListInterfaceIteratorProxy<BaseInterfaceType>(Registry, ElementHandles);
+	}
+
+	/**
+	 * DO NOT USE DIRECTLY
+	 * STL-like iterators to enable range-based for loop support.
+	 */
+	FORCEINLINE TArray<FTypedElementHandle>::RangedForConstIteratorType begin() const { return ElementHandles.begin(); }
+	FORCEINLINE TArray<FTypedElementHandle>::RangedForConstIteratorType end() const { return ElementHandles.end(); }
+
+private:
+	enum class EChangeType : uint8
+	{
+		/**
+		* An element was added to the element list.
+		* The ElementHandle argument will be set to the element that was added.
+		*/
+		Added,
+
+		/**
+		* An element was removed from the element list.
+		* The ElementHandle argument will be set to the element that was removed.
+		*/
+		Removed,
+
+		/**
+		* The element list was cleared.
+		* The ElementHandle argument will be unset.
+		*/
+		Cleared,
+	};
+
+	explicit FTypedElementList(UTypedElementRegistry* InRegistry);
+
+	bool AddElementImpl(FTypedElementHandle&& InElementHandle);
+	bool RemoveElementImpl(const FTypedElementId& InElementId);
+	int32 RemoveAllElementsImpl(TFunctionRef<bool(const FTypedElementHandle&)> InPredicate);
+	bool ContainsElementImpl(const FTypedElementId& InElementId) const;
+
+	void NoteListChanged(const EChangeType InChangeType, const FTypedElementHandle& InElementHandle = FTypedElementHandle());
+
+	UTypedElementRegistry* Registry = nullptr;
+
+	/**
+	 * Set of combined ID values that are currently present in this element list.
+	 * Used to perform optimized querying of which elements are in this list, and to avoid adding duplicate entries.
+	 */
+	TSet<FTypedHandleCombinedId> ElementCombinedIds;
+
+	/**
+	 * Array of element handles present in this element list.
+	 * These are stored in the same order that they are added, and the set above can be used to optimize certain queries.
+	 */
+	TArray<FTypedElementHandle> ElementHandles;
+
+	/**
+	 * Delegate that is invoked whenever this element list has been changed
+	 */
+	FOnChanged OnChangedDelegate;
+
+	/**
+	 * Whether there are pending changes for OnChangedDelegate to notify for.
+	 */
+	bool bHasPendingNotify = false;
+
+	/**
+	 * Interface to allow external systems (such as USelection) to receive immediate sync notifications as an element list is changed.
+	 * This exists purely as a bridging mechanism and shouldn't be relied on for new code. It is lazily created as needed.
+	 */
+	TUniquePtr<FTypedElementListLegacySync> LegacySync;
+};
diff --git a/Engine/Source/Runtime/TypedElementFramework/Public/TypedElementListPtr.h b/Engine/Source/Runtime/TypedElementFramework/Public/TypedElementListPtr.h
new file mode 100644
--- /dev/null
+++ b/Engine/Source/Runtime/TypedElementFramework/Public/TypedElementListPtr.h
@@ -0,0 +1,8 @@
+// Copyright Epic Games, Inc. All Rights Reserved.
+
+#pragma once
+
+#include "Templates/UniquePtr.h"
+
+class FTypedElementList;
+using FTypedElementListPtr = TUniquePtr<FTypedElementList>;
diff --git a/Engine/Source/Runtime/TypedElementFramework/Public/TypedElementOwnerStore.h b/Engine/Source/Runtime/TypedElementFramework/Public/TypedElementOwnerStore.h
new file mode 100644
--- /dev/null
+++ b/Engine/Source/Runtime/TypedElementFramework/Public/TypedElementOwnerStore.h
@@ -0,0 +1,182 @@
+// Copyright Epic Games, Inc. All Rights Reserved.
+
+#pragma once
+
+#include "CoreMinimal.h"
+#include "TypedElementHandle.h"
+#include "Misc/ScopeLock.h"
+
+/**
+ * Type to keep a lock on the TTypedElementOwnerStore for the duration 
+ * of time that a TTypedElementOwner is being externally referenced.
+ * This avoids invalid memory access if the map is reallocated.
+ */
+template <typename ElementOwnerType>
+struct TTypedElementOwnerScopedAccessImpl
+{
+public:
+	explicit TTypedElementOwnerScopedAccessImpl(FCriticalSection* InSynchObject)
+		: SynchObject(InSynchObject)
+	{
+		check(SynchObject);
+		SynchObject->Lock();
+	}
+
+	~TTypedElementOwnerScopedAccessImpl()
+	{
+		if (SynchObject)
+		{
+			SynchObject->Unlock();
+		}
+	}
+
+	TTypedElementOwnerScopedAccessImpl(const TTypedElementOwnerScopedAccessImpl&) = delete;
+	TTypedElementOwnerScopedAccessImpl& operator=(const TTypedElementOwnerScopedAccessImpl&) = delete;
+
+	TTypedElementOwnerScopedAccessImpl(TTypedElementOwnerScopedAccessImpl&& InOther)
+		: SynchObject(InOther.SynchObject)
+		, ElementOwner(InOther.ElementOwner)
+	{
+		InOther.SynchObject = nullptr;
+		InOther.ElementOwner = nullptr;
+	}
+
+	TTypedElementOwnerScopedAccessImpl& operator=(TTypedElementOwnerScopedAccessImpl&&) = delete;
+
+	FORCEINLINE explicit operator bool() const
+	{
+		return IsSet();
+	}
+
+	FORCEINLINE bool IsSet() const
+	{
+		return ElementOwner != nullptr;
+	}
+
+	FORCEINLINE ElementOwnerType& operator*() const
+	{
+		check(ElementOwner);
+		return *ElementOwner;
+	}
+
+	FORCEINLINE ElementOwnerType* operator->() const
+	{
+		check(ElementOwner);
+		return ElementOwner;
+	}
+
+	void Private_SetElementOwner(ElementOwnerType* InElementOwner)
+	{
+		checkf(!ElementOwner, TEXT("Element owner was already set!"));
+		ElementOwner = InElementOwner;
+
+		// If we didn't find an element then we can unlock immediately
+		// Otherwise we hold the lock for as long as the element is being read to avoid a map write from invalidating the reference from reallocation
+		check(SynchObject);
+		if (!ElementOwner)
+		{
+			SynchObject->Unlock();
+			SynchObject = nullptr;
+		}
+	}
+
+private:
+	FCriticalSection* SynchObject = nullptr;
+	ElementOwnerType* ElementOwner = nullptr;
+};
+
+template <typename ElementDataType>
+using TTypedElementOwnerScopedAccess = TTypedElementOwnerScopedAccessImpl<const TTypedElementOwner<ElementDataType>>;
+
+template <typename ElementDataType>
+using TTypedElementOwnerScopedMutableAccess = TTypedElementOwnerScopedAccessImpl<TTypedElementOwner<ElementDataType>>;
+
+inline FString GetTypedElementOwnerStoreKeyDebugString(const void* InKey)
+{
+	return TEXT("<GetTypedElementOwnerStoreKeyDebugString not implemented for key type>");
+}
+
+inline FString GetTypedElementOwnerStoreKeyDebugString(const UObject* InKey)
+{
+	return InKey->GetPathName();
+}
+
+/**
+ * A store of element owners tied to their corresponding owner instance.
+ * This can be used to track element owners for an instance without having to add new data to the type itself.
+ */
+template <typename ElementDataType, typename KeyDataType = const void*>
+struct TTypedElementOwnerStore
+{
+public:
+	static_assert(TNumericLimits<int32>::Max() >= TypedHandleMaxElementId, "TTypedElementOwnerStore internally uses TMap so cannot store TypedHandleMaxElementId! Consider making this container 64-bit aware, or explicitly remove this compile time check.");
+
+	TTypedElementOwnerStore() = default;
+
+	~TTypedElementOwnerStore()
+	{
+		checkf(ElementOwnerMap.Num() == 0, TEXT("Element owners were still registered during destruction! This will leak elements, and you should unregister and destroy all elements prior to destruction!"));
+	}
+
+	/**
+	 * Register the given element owner, transferring ownership of the owner instance to this store.
+	 * @note This will lock the store for writes until the scoped access struct is no longer being used.
+	 * @note This must be paired with a call to UnregisterElementOwner.
+	 */
+	TTypedElementOwnerScopedAccess<ElementDataType> RegisterElementOwner(const KeyDataType& InKey, TTypedElementOwner<ElementDataType>&& InElementOwner)
+	{
+		TTypedElementOwnerScopedAccess<ElementDataType> ScopedAccess(&ElementOwnerMapCS);
+		checkf(!ElementOwnerMap.Contains(InKey), TEXT("Element owner has already been registered for this key (%s)! This will leak elements, and you should unregister and destroy the old element owner for this instance before adding a new one!"), *GetTypedElementOwnerStoreKeyDebugString(InKey));
+		ScopedAccess.Private_SetElementOwner(&ElementOwnerMap.Add(InKey, MoveTemp(InElementOwner)));
+		return ScopedAccess;
+	}
+
+	/**
+	 * Unregister the given element owner, transferring ownership of the owner instance back to the caller.
+	 * @note The returned owner may be unset if the key was not registered.
+	 */
+	TTypedElementOwner<ElementDataType> UnregisterElementOwner(const KeyDataType& InKey)
+	{
+		TTypedElementOwner<ElementDataType> ElementOwner;
+		{
+			FScopeLock Lock(&ElementOwnerMapCS);
+			ElementOwnerMap.RemoveAndCopyValue(InKey, ElementOwner);
+		}
+		return ElementOwner;
+	}
+
+	/**
+	 * Provide const access to the given element owner.
+	 * @note This will lock the store for writes until the scoped access struct is no longer being used. The returned struct may be unset if the key was not registered.
+	 */
+	TTypedElementOwnerScopedAccess<ElementDataType> FindElementOwner(const KeyDataType& InKey) const
+	{
+		TTypedElementOwnerScopedAccess<ElementDataType> ScopedAccess(&ElementOwnerMapCS);
+		ScopedAccess.Private_SetElementOwner(ElementOwnerMap.Find(InKey));
+		return ScopedAccess;
+	}
+
+	/**
+	 * Provide mutable access to the given element owner.
+	 * @note This will lock the store for writes until the scoped access struct is no longer being used. The returned struct may be unset if the key was not registered.
+	 */
+	TTypedElementOwnerScopedMutableAccess<ElementDataType> FindMutableElementOwner(const KeyDataType& InKey)
+	{
+		TTypedElementOwnerScopedMutableAccess<ElementDataType> ScopedAccess(&ElementOwnerMapCS);
+		ScopedAccess.Private_SetElementOwner(ElementOwnerMap.Find(InKey));
+		return ScopedAccess;
+	}
+
+	/**
+	 * Test to see whether the store has an entry for the given element owner.
+	 */
+	bool ContainsElementOwner(const KeyDataType& InKey) const
+	{
+		FScopeLock Lock(&ElementOwnerMapCS);
+		return ElementOwnerMap.Contains(InKey);
+	}
+
+private:
+	TMap<KeyDataType, TTypedElementOwner<ElementDataType>> ElementOwnerMap;
+	mutable FCriticalSection ElementOwnerMapCS;
+};
diff --git a/Engine/Source/Runtime/TypedElementFramework/Public/TypedElementRegistry.h b/Engine/Source/Runtime/TypedElementFramework/Public/TypedElementRegistry.h
new file mode 100644
--- /dev/null
+++ b/Engine/Source/Runtime/TypedElementFramework/Public/TypedElementRegistry.h
@@ -0,0 +1,397 @@
+// Copyright Epic Games, Inc. All Rights Reserved.
+
+#pragma once
+
+#include "CoreMinimal.h"
+#include "Misc/ScopeRWLock.h"
+#include "Containers/ArrayView.h"
+#include "Containers/SortedMap.h"
+#include "Templates/SubclassOf.h"
+#include "TypedElementHandle.h"
+#include "TypedElementList.h"
+#include "TypedElementRegistry.generated.h"
+
+/**
+ * Registry of element types and their associated interfaces, along with the elements that represent their instances.
+ */
+UCLASS()
+class TYPEDELEMENTFRAMEWORK_API UTypedElementRegistry : public UObject
+{
+	GENERATED_BODY()
+
+public:
+	UTypedElementRegistry();
+
+	//~ UObject interface
+	static void AddReferencedObjects(UObject* InThis, FReferenceCollector& Collector);
+
+	/**
+	 * Initialize the singleton instance of the registry used in most cases.
+	 */
+	static void Private_InitializeInstance();
+
+	/**
+	 * Shutdown the singleton instance of the registry used in most cases.
+	 */
+	static void Private_ShutdownInstance();
+
+	/**
+	 * Get the singleton instance of the registry used in most cases.
+	 */
+	static UTypedElementRegistry* GetInstance();
+
+	/**
+	 * Register an element type that doesn't require any additional payload data.
+	 */
+	FORCEINLINE void RegisterElementType(const FName InElementTypeName)
+	{
+		RegisterElementTypeImpl(InElementTypeName, MakeUnique<TRegisteredElementType<void>>());
+	}
+
+	/**
+	 * Register an element type that has additional payload data.
+	 */
+	template <typename ElementDataType>
+	FORCEINLINE void RegisterElementType(const FName InElementTypeName)
+	{
+		RegisterElementTypeImpl(InElementTypeName, MakeUnique<TRegisteredElementType<ElementDataType>>());
+	}
+
+	/**
+	 * Register that an element interface is supported for the given type, which must have previously been registered via RegisterElementType.
+	 */
+	template <typename BaseInterfaceType>
+	FORCEINLINE void RegisterElementInterface(const FName InElementTypeName, UTypedElementInterface* InElementInterface)
+	{
+		RegisterElementInterfaceImpl(InElementTypeName, InElementInterface, BaseInterfaceType::StaticClass());
+	}
+
+	/**
+	 * Get the element interface supported by the given handle, or null if there is no support for this interface.
+	 */
+	template <typename BaseInterfaceType>
+	FORCEINLINE BaseInterfaceType* GetElementInterface(const FTypedElementId& InElementId) const
+	{
+		return static_cast<BaseInterfaceType*>(GetElementInterfaceImpl(InElementId, BaseInterfaceType::StaticClass()));
+	}
+
+	/**
+	 * Get the element interface supported by the given handle, or null if there is no support for this interface.
+	 */
+	template <typename BaseInterfaceType>
+	FORCEINLINE BaseInterfaceType* GetElementInterface(const FTypedElementHandle& InElementHandle) const
+	{
+		return static_cast<BaseInterfaceType*>(GetElementInterfaceImpl(InElementHandle.GetId(), BaseInterfaceType::StaticClass()));
+	}
+
+	/**
+	 * Create an element that doesn't require any additional payload data.
+	 * @note The associated handle ID should be something that can externally be used to uniquely identify this element, until DestroyElementHandle is called on this handle.
+	 */
+	FORCEINLINE FTypedElementOwner CreateElement(const FName InElementTypeName, const FTypedHandleElementId InElementId)
+	{
+		return CreateElementImpl<void>(InElementTypeName, InElementId);
+	}
+
+	/**
+	 * Create an element that has additional payload data.
+	 * @note Allocation of the payload data and the associated handle ID are managed internally, and the data will remain valid until DestroyElementHandle is called on this handle.
+	 */
+	template <typename ElementDataType>
+	FORCEINLINE TTypedElementOwner<ElementDataType> CreateElement(const FName InElementTypeName)
+	{
+		return CreateElementImpl<ElementDataType>(InElementTypeName, INDEX_NONE);
+	}
+
+	/**
+	 * Destroy an element.
+	 */
+	FORCEINLINE void DestroyElement(FTypedElementOwner& InOutElementOwner)
+	{
+		return DestroyElementImpl<void>(InOutElementOwner);
+	}
+
+	/**
+	 * Destroy an element.
+	 */
+	template <typename ElementDataType>
+	FORCEINLINE void DestroyElement(TTypedElementOwner<ElementDataType>& InOutElementOwner)
+	{
+		return DestroyElementImpl<ElementDataType>(InOutElementOwner);
+	}
+
+	/**
+	 * Release an element ID that was previously acquired from an existing handle.
+	 */
+	void ReleaseElementId(FTypedElementId& InOutElementId);
+
+	/**
+	 * Get an element handle from its minimal ID.
+	 */
+	FTypedElementHandle GetElementHandle(const FTypedElementId& InElementId) const;
+
+	/**
+	 * Get an element that implements the given interface from its minimal ID.
+	 */
+	FORCEINLINE FTypedElement GetElement(const FTypedElementId& InElementId, const TSubclassOf<UTypedElementInterface>& InBaseInterfaceType) const
+	{
+		FTypedElement Element;
+		GetElementImpl(InElementId, InBaseInterfaceType, Element);
+		return Element;
+	}
+
+	/**
+	 * Get an element that implements the given interface from its minimal ID.
+	 */
+	template <typename BaseInterfaceType>
+	FORCEINLINE TTypedElement<BaseInterfaceType> GetElement(const FTypedElementId& InElementId) const
+	{
+		TTypedElement<BaseInterfaceType> Element;
+		GetElementImpl(InElementId, BaseInterfaceType::StaticClass(), Element);
+		return Element;
+	}
+
+	/**
+	 * Get an element that implements the given interface from its handle.
+	 */
+	FORCEINLINE FTypedElement GetElement(const FTypedElementHandle& InElementHandle, const TSubclassOf<UTypedElementInterface>& InBaseInterfaceType) const
+	{
+		FTypedElement Element;
+		GetElementImpl(InElementHandle, InBaseInterfaceType, Element);
+		return Element;
+	}
+
+	/**
+	 * Get an element that implements the given interface from its handle.
+	 */
+	template <typename BaseInterfaceType>
+	FORCEINLINE TTypedElement<BaseInterfaceType> GetElement(const FTypedElementHandle& InElementHandle) const
+	{
+		TTypedElement<BaseInterfaceType> Element;
+		GetElementImpl(InElementHandle, BaseInterfaceType::StaticClass(), Element);
+		return Element;
+	}
+
+	/**
+	 * Create an empty list of elements associated with this registry.
+	 */
+	FORCEINLINE FTypedElementListPtr CreateElementList()
+	{
+		return FTypedElementList::Private_CreateElementList(this);
+	}
+
+	/**
+	 * Create an empty list of elements associated with this registry, populated from the given minimal IDs that are valid.
+	 */
+	FTypedElementListPtr CreateElementList(TArrayView<const FTypedElementId> InElementIds);
+
+	/**
+	 * Create an empty list of elements associated with this registry, populated from the given handles that are valid.
+	 */
+	FTypedElementListPtr CreateElementList(TArrayView<const FTypedElementHandle> InElementHandles);
+
+	/**
+	 * Create an empty list of elements associated with this registry, populated from the given owners that are valid.
+	 */
+	template <typename ElementDataType>
+	FORCEINLINE FTypedElementListPtr CreateElementList(const TArray<TTypedElementOwner<ElementDataType>>& InElementOwners)
+	{
+		return CreateElementList(MakeArrayView(InElementOwners));
+	}
+
+	/**
+	 * Create an empty list of elements associated with this registry, populated from the given owners that are valid.
+	 */
+	template <typename ElementDataType>
+	FTypedElementListPtr CreateElementList(TArrayView<const TTypedElementOwner<ElementDataType>> InElementOwners)
+	{
+		FTypedElementListPtr ElementList = CreateElementList();
+		ElementList->Append(InElementOwners);
+		return ElementList;
+	}
+
+	void Private_OnElementListCreated(FTypedElementList* InElementList)
+	{
+		FWriteScopeLock ActiveElementListsLock(ActiveElementListsRW);
+		ActiveElementLists.Add(InElementList);
+	}
+
+	void Private_OnElementListDestroyed(FTypedElementList* InElementList)
+	{
+		FWriteScopeLock ActiveElementListsLock(ActiveElementListsRW);
+		ActiveElementLists.Remove(InElementList);
+	}
+
+	// Note: Access for FTypedElementList
+	FORCEINLINE void Private_GetElementImpl(const FTypedElementHandle& InElementHandle, const UClass* InBaseInterfaceType, FTypedElement& OutElement) const
+	{
+		GetElementImpl(InElementHandle, InBaseInterfaceType, OutElement);
+	}
+
+private:
+	struct FRegisteredElementType
+	{
+		virtual ~FRegisteredElementType() = default;
+
+		virtual FTypedElementInternalData& AddDataForElement(FTypedHandleElementId& InOutElementId) = 0;
+		virtual void RemoveDataForElement(const FTypedHandleElementId InElementId, const FTypedElementInternalData* InExpectedDataPtr) = 0;
+		virtual const FTypedElementInternalData& GetDataForElement(const FTypedHandleElementId InElementId) const = 0;
+		virtual void SetDataTypeId(const FTypedHandleTypeId InTypeId) = 0;
+		virtual FTypedHandleTypeId GetDataTypeId() const = 0;
+		virtual FName GetDataTypeName() const = 0;
+
+		FTypedHandleTypeId TypeId = 0;
+		FName TypeName;
+		TSortedMap<FName, UTypedElementInterface*, FDefaultAllocator, FNameFastLess> Interfaces;
+	};
+
+	template <typename ElementDataType>
+	struct TRegisteredElementType : public FRegisteredElementType
+	{
+		virtual ~TRegisteredElementType() = default;
+		
+		virtual FTypedElementInternalData& AddDataForElement(FTypedHandleElementId& InOutElementId) override
+		{
+			return HandleDataStore.AddDataForElement(InOutElementId);
+		}
+
+		virtual void RemoveDataForElement(const FTypedHandleElementId InElementId, const FTypedElementInternalData* InExpectedDataPtr) override
+		{
+			HandleDataStore.RemoveDataForElement(InElementId, InExpectedDataPtr);
+		}
+
+		virtual const FTypedElementInternalData& GetDataForElement(const FTypedHandleElementId InElementId) const override
+		{
+			return HandleDataStore.GetDataForElement(InElementId);
+		}
+
+		virtual void SetDataTypeId(const FTypedHandleTypeId InTypeId) override
+		{
+			TTypedElementInternalDataStore<ElementDataType>::SetStaticDataTypeId(InTypeId);
+		}
+
+		virtual FTypedHandleTypeId GetDataTypeId() const override
+		{
+			return TTypedElementInternalDataStore<ElementDataType>::StaticDataTypeId();
+		}
+
+		virtual FName GetDataTypeName() const override
+		{
+			return TTypedElementInternalDataStore<ElementDataType>::StaticDataTypeName();
+		}
+
+		TTypedElementInternalDataStore<ElementDataType> HandleDataStore;
+	};
+
+	void RegisterElementTypeImpl(const FName InElementTypeName, TUniquePtr<FRegisteredElementType>&& InRegisteredElementType);
+	void RegisterElementInterfaceImpl(const FName InElementTypeName, UTypedElementInterface* InElementInterface, const TSubclassOf<UTypedElementInterface>& InBaseInterfaceType);
+	UTypedElementInterface* GetElementInterfaceImpl(const FTypedElementId& InElementId, const TSubclassOf<UTypedElementInterface>& InBaseInterfaceType) const;
+
+	template <typename ElementDataType>
+	TTypedElementOwner<ElementDataType> CreateElementImpl(const FName InElementTypeName, const FTypedHandleElementId InElementId)
+	{
+		FRegisteredElementType* RegisteredElementType = GetRegisteredElementTypeFromName(InElementTypeName);
+		checkf(RegisteredElementType, TEXT("Element type '%s' has not been registered!"), *InElementTypeName.ToString());
+
+		checkf(RegisteredElementType->GetDataTypeId() == TTypedElementInternalDataStore<ElementDataType>::StaticDataTypeId(), TEXT("Element type '%s' uses '%s' as its handle data type, but '%s' was requested!"), *InElementTypeName.ToString(), *RegisteredElementType->GetDataTypeName().ToString(), *TTypedElementInternalDataStore<ElementDataType>::StaticDataTypeName().ToString());
+
+		FTypedHandleElementId NewElementId = InElementId;
+		FTypedElementInternalData& NewElementData = RegisteredElementType->AddDataForElement(NewElementId);
+
+		TTypedElementOwner<ElementDataType> ElementOwner;
+		ElementOwner.Private_InitializeAddRef(RegisteredElementType->TypeId, NewElementId, static_cast<TTypedElementInternalData<ElementDataType>&>(NewElementData));
+
+		return ElementOwner;
+	}
+
+	template <typename ElementDataType>
+	void DestroyElementImpl(TTypedElementOwner<ElementDataType>& InOutElementOwner)
+	{
+		FRegisteredElementType* RegisteredElementType = GetRegisteredElementTypeFromId(InOutElementOwner.GetId().GetTypeId());
+		checkf(RegisteredElementType, TEXT("Element type ID '%d' has not been registered!"), InOutElementOwner.GetId().GetTypeId());
+
+#if DO_CHECK && WITH_TYPED_ELEMENT_REFCOUNT
+		{
+			const FTypedHandleRefCount RefCount = RegisteredElementType->GetDataForElement(InOutElementOwner.GetId().GetElementId()).GetRefCount();
+			checkf(RefCount == 1, TEXT("Element is still externally referenced when being destroyed (ref-count: %d)!"), RefCount);
+		}
+#endif	// DO_CHECK && WITH_TYPED_ELEMENT_REFCOUNT
+
+		RegisteredElementType->RemoveDataForElement(InOutElementOwner.GetId().GetElementId(), InOutElementOwner.Private_GetInternalData());
+		InOutElementOwner.Private_DestroyNoRef();
+	}
+
+	template <typename BaseInterfaceType>
+	void GetElementImpl(const FTypedElementId& InElementId, const UClass* InBaseInterfaceType, TTypedElement<BaseInterfaceType>& OutElement) const
+	{
+		OutElement.Private_DestroyReleaseRef();
+
+		if (InElementId)
+		{
+			FRegisteredElementType* RegisteredElementType = GetRegisteredElementTypeFromId(InElementId.GetTypeId());
+			checkf(RegisteredElementType, TEXT("Element type ID '%d' has not been registered!"), InElementId.GetTypeId());
+
+			OutElement.Private_InitializeAddRef(InElementId.GetTypeId(), InElementId.GetElementId(), RegisteredElementType->GetDataForElement(InElementId.GetElementId()), static_cast<BaseInterfaceType*>(RegisteredElementType->Interfaces.FindRef(InBaseInterfaceType->GetFName())));
+		}
+	}
+
+	template <typename BaseInterfaceType>
+	void GetElementImpl(const FTypedElementHandle& InElementHandle, const UClass* InBaseInterfaceType, TTypedElement<BaseInterfaceType>& OutElement) const
+	{
+		OutElement.Private_DestroyReleaseRef();
+
+		if (InElementHandle)
+		{
+			FRegisteredElementType* RegisteredElementType = GetRegisteredElementTypeFromId(InElementHandle.GetId().GetTypeId());
+			checkf(RegisteredElementType, TEXT("Element type ID '%d' has not been registered!"), InElementHandle.GetId().GetTypeId());
+
+			OutElement.Private_InitializeAddRef(InElementHandle.GetId().GetTypeId(), InElementHandle.GetId().GetElementId(), *InElementHandle.Private_GetInternalData(), static_cast<BaseInterfaceType*>(RegisteredElementType->Interfaces.FindRef(InBaseInterfaceType->GetFName())));
+		}
+	}
+
+	void AddRegisteredElementType(TUniquePtr<FRegisteredElementType>&& InRegisteredElementType)
+	{
+		checkf(InRegisteredElementType->TypeId > 0, TEXT("Element type ID was unassigned!"));
+		checkf(!GetRegisteredElementTypeFromId(InRegisteredElementType->TypeId), TEXT("Element type '%d' has already been registered!"), InRegisteredElementType->TypeId);
+		checkf(!GetRegisteredElementTypeFromName(InRegisteredElementType->TypeName), TEXT("Element type '%s' has already been registered!"), *InRegisteredElementType->TypeName.ToString());
+
+		{
+			FWriteScopeLock RegisteredElementTypesLock(RegisteredElementTypesRW);
+
+			RegisteredElementTypesNameToId.Add(InRegisteredElementType->TypeName, InRegisteredElementType->TypeId);
+			RegisteredElementTypes[InRegisteredElementType->TypeId - 1] = MoveTemp(InRegisteredElementType);
+		}
+	}
+
+	FRegisteredElementType* GetRegisteredElementTypeFromId(const FTypedHandleTypeId InTypeId) const
+	{
+		FReadScopeLock RegisteredElementTypesLock(RegisteredElementTypesRW);
+
+		return InTypeId > 0
+			? RegisteredElementTypes[InTypeId - 1].Get()
+			: nullptr;
+	}
+
+	FRegisteredElementType* GetRegisteredElementTypeFromName(const FName& InTypeName) const
+	{
+		FReadScopeLock RegisteredElementTypesLock(RegisteredElementTypesRW);
+
+		if (const FTypedHandleTypeId* TypeId = RegisteredElementTypesNameToId.Find(InTypeName))
+		{
+			return RegisteredElementTypes[(*TypeId) - 1].Get();
+		}
+
+		return nullptr;
+	}
+
+	void NotifyElementListPendingChanges();
+
+	static UTypedElementRegistry* Instance;
+	
+	mutable FRWLock RegisteredElementTypesRW;
+	TUniquePtr<FRegisteredElementType> RegisteredElementTypes[TypedHandleMaxTypeId - 1];
+	TSortedMap<FName, FTypedHandleTypeId, FDefaultAllocator, FNameFastLess> RegisteredElementTypesNameToId;
+
+	mutable FRWLock ActiveElementListsRW;
+	TSet<FTypedElementList*> ActiveElementLists;
+};
diff --git a/Engine/Source/Runtime/TypedElementFramework/TypedElementFramework.Build.cs b/Engine/Source/Runtime/TypedElementFramework/TypedElementFramework.Build.cs
new file mode 100644
--- /dev/null
+++ b/Engine/Source/Runtime/TypedElementFramework/TypedElementFramework.Build.cs
@@ -0,0 +1,18 @@
+// Copyright Epic Games, Inc. All Rights Reserved.
+
+using UnrealBuildTool;
+using System.IO;
+
+public class TypedElementFramework : ModuleRules
+{
+	public TypedElementFramework(ReadOnlyTargetRules Target) : base(Target)
+	{
+		PublicDependencyModuleNames.AddRange(
+			new string[]
+			{
+				"Core",
+				"CoreUObject",
+			}
+		);
+    }
+}
diff --git a/Engine/Source/Runtime/TypedElementInterfaces/Private/TypedElementInterfacesModule.cpp b/Engine/Source/Runtime/TypedElementInterfaces/Private/TypedElementInterfacesModule.cpp
new file mode 100644
--- /dev/null
+++ b/Engine/Source/Runtime/TypedElementInterfaces/Private/TypedElementInterfacesModule.cpp
@@ -0,0 +1,14 @@
+// Copyright Epic Games, Inc. All Rights Reserved.
+
+#include "Modules/ModuleManager.h"
+
+class FTypedElementInterfacesModule : public FDefaultModuleImpl
+{
+public:
+	virtual void StartupModule() override
+	{
+		FModuleManager::Get().LoadModuleChecked("TypedElementFramework");
+	}
+};
+
+IMPLEMENT_MODULE(FTypedElementInterfacesModule, TypedElementInterfaces)
diff --git a/Engine/Source/Runtime/TypedElementInterfaces/TypedElementInterfaces.Build.cs b/Engine/Source/Runtime/TypedElementInterfaces/TypedElementInterfaces.Build.cs
new file mode 100644
--- /dev/null
+++ b/Engine/Source/Runtime/TypedElementInterfaces/TypedElementInterfaces.Build.cs
@@ -0,0 +1,19 @@
+// Copyright Epic Games, Inc. All Rights Reserved.
+
+using UnrealBuildTool;
+using System.IO;
+
+public class TypedElementInterfaces : ModuleRules
+{
+	public TypedElementInterfaces(ReadOnlyTargetRules Target) : base(Target)
+	{
+		PublicDependencyModuleNames.AddRange(
+			new string[]
+			{
+				"Core",
+				"CoreUObject",
+				"TypedElementFramework",
+			}
+		);
+    }
+}
