Renamed FTypedElementHandle::GetDataScript() to FTypedElementHandle::GetData()

#rnx

diff --git a/Engine/Source/Runtime/Engine/Private/Elements/Actor/ActorElementSelectionInterface.cpp b/Engine/Source/Runtime/Engine/Private/Elements/Actor/ActorElementSelectionInterface.cpp
--- a/Engine/Source/Runtime/Engine/Private/Elements/Actor/ActorElementSelectionInterface.cpp
+++ b/Engine/Source/Runtime/Engine/Private/Elements/Actor/ActorElementSelectionInterface.cpp
@@ -12,6 +12,6 @@ bool UActorElementSelectionInterface::IsValidSelection(const FTypedElementHandle
 
 UObject* UActorElementSelectionInterface::Legacy_GetSelectionObject(const FTypedElementHandle& InElementHandle)
 {
-	const FActorElementData* ActorData = InElementHandle.GetDataScript<FActorElementData>();
+	const FActorElementData* ActorData = InElementHandle.GetData<FActorElementData>();
 	return ActorData ? ActorData->Actor : nullptr;
 }
diff --git a/Engine/Source/Runtime/Engine/Private/Elements/Component/ComponentElementSelectionInterface.cpp b/Engine/Source/Runtime/Engine/Private/Elements/Component/ComponentElementSelectionInterface.cpp
--- a/Engine/Source/Runtime/Engine/Private/Elements/Component/ComponentElementSelectionInterface.cpp
+++ b/Engine/Source/Runtime/Engine/Private/Elements/Component/ComponentElementSelectionInterface.cpp
@@ -12,6 +12,6 @@ bool UComponentElementSelectionInterface::IsValidSelection(const FTypedElementHa
 
 UObject* UComponentElementSelectionInterface::Legacy_GetSelectionObject(const FTypedElementHandle& InElementHandle)
 {
-	const FComponentElementData* ComponentData = InElementHandle.GetDataScript<FComponentElementData>();
+	const FComponentElementData* ComponentData = InElementHandle.GetData<FComponentElementData>();
 	return ComponentData ? ComponentData->Component : nullptr;
 }
diff --git a/Engine/Source/Runtime/TypedElementFramework/Private/TypedElementFrameworkTests.cpp b/Engine/Source/Runtime/TypedElementFramework/Private/TypedElementFrameworkTests.cpp
--- a/Engine/Source/Runtime/TypedElementFramework/Private/TypedElementFrameworkTests.cpp
+++ b/Engine/Source/Runtime/TypedElementFramework/Private/TypedElementFrameworkTests.cpp
@@ -8,7 +8,7 @@ UE_DEFINE_TYPED_ELEMENT_DATA_RTTI(FTestTypedElementData);
 
 FText UTestTypedElementInterfaceA_ImplTyped::GetDisplayName(const FTypedElementHandle& InElementHandle)
 {
-	const FTestTypedElementData* ElementData = InElementHandle.GetDataScript<FTestTypedElementData>();
+	const FTestTypedElementData* ElementData = InElementHandle.GetData<FTestTypedElementData>();
 	if (!ElementData)
 	{
 		return FText();
@@ -19,7 +19,7 @@ FText UTestTypedElementInterfaceA_ImplTyped::GetDisplayName(const FTypedElementH
 
 bool UTestTypedElementInterfaceA_ImplTyped::SetDisplayName(const FTypedElementHandle& InElementHandle, FText InNewName, bool bNotify)
 {
-	const FTestTypedElementData* ElementData = InElementHandle.GetDataScript<FTestTypedElementData>();
+	const FTestTypedElementData* ElementData = InElementHandle.GetData<FTestTypedElementData>();
 	if (!ElementData)
 	{
 		return false;
diff --git a/Engine/Source/Runtime/TypedElementFramework/Public/TypedElementHandle.h b/Engine/Source/Runtime/TypedElementFramework/Public/TypedElementHandle.h
--- a/Engine/Source/Runtime/TypedElementFramework/Public/TypedElementHandle.h
+++ b/Engine/Source/Runtime/TypedElementFramework/Public/TypedElementHandle.h
@@ -241,24 +241,12 @@ public:
 		return Id.GetTypeId() == ElementDataType::StaticTypeId();
 	}
 
-	/**
-	 * Attempt to access the data stored within this handle as the given type, returning null if it isn't possible.
-	 * @note This is not typically something you'd want to use outside of data access within an interface implementation.
-	 */
-	template <typename ElementDataType>
-	FORCEINLINE const ElementDataType* GetData() const
-	{
-		return IsDataOfType<ElementDataType>()
-			? static_cast<const ElementDataType*>(DataPtr->GetUntypedData())
-			: nullptr;
-	}
-
 	/**
 	 * Attempt to access the data stored within this handle as the given type, returning null if it isn't possible and logging an access error for scripting.
 	 * @note This is not typically something you'd want to use outside of data access within an interface implementation.
 	 */
 	template <typename ElementDataType>
-	const ElementDataType* GetDataScript() const
+	const ElementDataType* GetData() const
 	{
 		if (!DataPtr)
 		{
