Fixed shutdown crash from UTypedElementRegistry already being GC'd

#rnx

diff --git a/Engine/Source/Runtime/TypedElementFramework/Private/TypedElementRegistry.cpp b/Engine/Source/Runtime/TypedElementFramework/Private/TypedElementRegistry.cpp
--- a/Engine/Source/Runtime/TypedElementFramework/Private/TypedElementRegistry.cpp
+++ b/Engine/Source/Runtime/TypedElementFramework/Private/TypedElementRegistry.cpp
@@ -3,10 +3,15 @@
 #include "TypedElementRegistry.h"
 #include "Misc/CoreDelegates.h"
 #include "Misc/ScopeLock.h"
+#include "UObject/StrongObjectPtr.h"
 
 const FTypedElementId FTypedElementId::Unset;
 
-UTypedElementRegistry* UTypedElementRegistry::Instance = nullptr;
+TStrongObjectPtr<UTypedElementRegistry>& GetTypedElementRegistryInstance()
+{
+	static TStrongObjectPtr<UTypedElementRegistry> TypedElementRegistryInstance;
+	return TypedElementRegistryInstance;
+}
 
 UTypedElementRegistry::UTypedElementRegistry()
 {
@@ -15,21 +20,32 @@ UTypedElementRegistry::UTypedElementRegistry()
 
 void UTypedElementRegistry::Private_InitializeInstance()
 {
+	TStrongObjectPtr<UTypedElementRegistry>& Instance = GetTypedElementRegistryInstance();
 	checkf(!Instance, TEXT("Instance was already initialized!"));
-	Instance = NewObject<UTypedElementRegistry>();
-	Instance->AddToRoot();
+	Instance.Reset(NewObject<UTypedElementRegistry>());
 }
 
 void UTypedElementRegistry::Private_ShutdownInstance()
 {
-	checkf(Instance, TEXT("Instance was already shutdown!"));
-	Instance->RemoveFromRoot();
-	Instance = nullptr;
+	TStrongObjectPtr<UTypedElementRegistry>& Instance = GetTypedElementRegistryInstance();
+	Instance.Reset();
 }
 
 UTypedElementRegistry* UTypedElementRegistry::GetInstance()
 {
-	return Instance;
+	TStrongObjectPtr<UTypedElementRegistry>& Instance = GetTypedElementRegistryInstance();
+	return Instance.Get();
+}
+
+void UTypedElementRegistry::FinishDestroy()
+{
+	TStrongObjectPtr<UTypedElementRegistry>& Instance = GetTypedElementRegistryInstance();
+	if (Instance.Get() == this)
+	{
+		Instance.Reset();
+	}
+
+	Super::FinishDestroy();
 }
 
 void UTypedElementRegistry::AddReferencedObjects(UObject* InThis, FReferenceCollector& Collector)
diff --git a/Engine/Source/Runtime/TypedElementFramework/Public/TypedElementRegistry.h b/Engine/Source/Runtime/TypedElementFramework/Public/TypedElementRegistry.h
--- a/Engine/Source/Runtime/TypedElementFramework/Public/TypedElementRegistry.h
+++ b/Engine/Source/Runtime/TypedElementFramework/Public/TypedElementRegistry.h
@@ -23,6 +23,7 @@ public:
 	UTypedElementRegistry();
 
 	//~ UObject interface
+	virtual void FinishDestroy() override;
 	static void AddReferencedObjects(UObject* InThis, FReferenceCollector& Collector);
 
 	/**
@@ -385,8 +386,6 @@ private:
 	}
 
 	void NotifyElementListPendingChanges();
-
-	static UTypedElementRegistry* Instance;
 	
 	mutable FRWLock RegisteredElementTypesRW;
 	TUniquePtr<FRegisteredElementType> RegisteredElementTypes[TypedHandleMaxTypeId - 1];
