Renamed macro to match coding standards

#rnx

diff --git a/Engine/Source/Runtime/TypedElementFramework/Private/TypedElementRegistry.cpp b/Engine/Source/Runtime/TypedElementFramework/Private/TypedElementRegistry.cpp
--- a/Engine/Source/Runtime/TypedElementFramework/Private/TypedElementRegistry.cpp
+++ b/Engine/Source/Runtime/TypedElementFramework/Private/TypedElementRegistry.cpp
@@ -122,10 +122,10 @@ void UTypedElementRegistry::ReleaseElementId(FTypedElementId& InOutElementId)
 	FRegisteredElementType* RegisteredElementType = GetRegisteredElementTypeFromId(InOutElementId.GetTypeId());
 	checkf(RegisteredElementType, TEXT("Element type ID '%d' has not been registered!"), InOutElementId.GetTypeId());
 
-#if WITH_TYPED_ELEMENT_REFCOUNT
+#if UE_TYPED_ELEMENT_HAS_REFCOUNT
 	const FTypedElementInternalData& ElementData = RegisteredElementType->GetDataForElement(InOutElementId.GetElementId());
 	ElementData.ReleaseRef();
-#endif	// WITH_TYPED_ELEMENT_REFCOUNT
+#endif	// UE_TYPED_ELEMENT_HAS_REFCOUNT
 
 	InOutElementId.Private_DestroyNoRef();
 }
diff --git a/Engine/Source/Runtime/TypedElementFramework/Public/TypedElementData.h b/Engine/Source/Runtime/TypedElementFramework/Public/TypedElementData.h
--- a/Engine/Source/Runtime/TypedElementFramework/Public/TypedElementData.h
+++ b/Engine/Source/Runtime/TypedElementFramework/Public/TypedElementData.h
@@ -39,7 +39,7 @@ public:
 
 	virtual ~FTypedElementInternalData() = default;
 
-#if WITH_TYPED_ELEMENT_REFCOUNT
+#if UE_TYPED_ELEMENT_HAS_REFCOUNT
 	FORCEINLINE void AddRef() const
 	{
 		checkSlow(RefCount < TNumericLimits<FTypedHandleRefCount>::Max());
@@ -56,7 +56,7 @@ public:
 	{
 		return FPlatformAtomics::AtomicRead(&RefCount);
 	}
-#endif	// WITH_TYPED_ELEMENT_REFCOUNT
+#endif	// UE_TYPED_ELEMENT_HAS_REFCOUNT
 
 	virtual const void* GetUntypedData() const
 	{
@@ -64,9 +64,9 @@ public:
 	}
 
 private:
-#if WITH_TYPED_ELEMENT_REFCOUNT
+#if UE_TYPED_ELEMENT_HAS_REFCOUNT
 	mutable FTypedHandleRefCount RefCount = 0;
-#endif	// WITH_TYPED_ELEMENT_REFCOUNT
+#endif	// UE_TYPED_ELEMENT_HAS_REFCOUNT
 };
 
 /**
@@ -194,7 +194,7 @@ public:
 
 	TTypedElementInternalData<void>& AddDataForElement(FTypedHandleElementId& InOutElementId)
 	{
-#if WITH_TYPED_ELEMENT_REFCOUNT
+#if UE_TYPED_ELEMENT_HAS_REFCOUNT
 		FWriteScopeLock InternalDataLock(InternalDataRW);
 
 		checkSlow(InOutElementId >= 0);
@@ -206,14 +206,14 @@ public:
 
 		ElementIdToArrayIndex.Add(InOutElementId, InternalDataArrayIndex);
 		return InternalDataArray[InternalDataArrayIndex];
-#else	// WITH_TYPED_ELEMENT_REFCOUNT
+#else	// UE_TYPED_ELEMENT_HAS_REFCOUNT
 		return SharedInternalData;
-#endif	// WITH_TYPED_ELEMENT_REFCOUNT
+#endif	// UE_TYPED_ELEMENT_HAS_REFCOUNT
 	}
 
 	void RemoveDataForElement(const FTypedHandleElementId InElementId, const FTypedElementInternalData* InExpectedDataPtr)
 	{
-#if WITH_TYPED_ELEMENT_REFCOUNT
+#if UE_TYPED_ELEMENT_HAS_REFCOUNT
 		FWriteScopeLock InternalDataLock(InternalDataRW);
 
 		int32 InternalDataArrayIndex = INDEX_NONE;
@@ -225,22 +225,22 @@ public:
 		checkf(InExpectedDataPtr == &InternalData, TEXT("Internal data pointer did not match the expected value! Does this handle belong to a different element registry?"));
 		InternalData = TTypedElementInternalData<void>();
 		InternalDataFreeIndices.Add(InternalDataArrayIndex);
-#else	// WITH_TYPED_ELEMENT_REFCOUNT
+#else	// UE_TYPED_ELEMENT_HAS_REFCOUNT
 		checkf(InExpectedDataPtr == &SharedInternalData, TEXT("Internal data pointer did not match the expected value! Does this handle belong to a different element registry?"));
-#endif	// WITH_TYPED_ELEMENT_REFCOUNT
+#endif	// UE_TYPED_ELEMENT_HAS_REFCOUNT
 	}
 
 	const TTypedElementInternalData<void>& GetDataForElement(const FTypedHandleElementId InElementId) const
 	{
-#if WITH_TYPED_ELEMENT_REFCOUNT
+#if UE_TYPED_ELEMENT_HAS_REFCOUNT
 		FReadScopeLock InternalDataLock(InternalDataRW);
 
 		const int32* InternalDataArrayIndexPtr = ElementIdToArrayIndex.Find(InElementId);
 		checkSlow(InternalDataArrayIndexPtr && InternalDataArray.IsValidIndex(*InternalDataArrayIndexPtr));
 		return InternalDataArray[*InternalDataArrayIndexPtr];
-#else	// WITH_TYPED_ELEMENT_REFCOUNT
+#else	// UE_TYPED_ELEMENT_HAS_REFCOUNT
 		return SharedInternalData;
-#endif	// WITH_TYPED_ELEMENT_REFCOUNT
+#endif	// UE_TYPED_ELEMENT_HAS_REFCOUNT
 	}
 
 	static FORCEINLINE void SetStaticDataTypeId(const FTypedHandleTypeId InTypeId)
@@ -258,12 +258,12 @@ public:
 	}
 
 private:
-#if WITH_TYPED_ELEMENT_REFCOUNT
+#if UE_TYPED_ELEMENT_HAS_REFCOUNT
 	mutable FRWLock InternalDataRW;
 	TChunkedArray<TTypedElementInternalData<void>> InternalDataArray;
 	TArray<int32> InternalDataFreeIndices;
 	TMap<FTypedHandleElementId, int32> ElementIdToArrayIndex;
-#else	// WITH_TYPED_ELEMENT_REFCOUNT
+#else	// UE_TYPED_ELEMENT_HAS_REFCOUNT
 	TTypedElementInternalData<void> SharedInternalData;
-#endif	// WITH_TYPED_ELEMENT_REFCOUNT
+#endif	// UE_TYPED_ELEMENT_HAS_REFCOUNT
 };
diff --git a/Engine/Source/Runtime/TypedElementFramework/Public/TypedElementHandle.h b/Engine/Source/Runtime/TypedElementFramework/Public/TypedElementHandle.h
--- a/Engine/Source/Runtime/TypedElementFramework/Public/TypedElementHandle.h
+++ b/Engine/Source/Runtime/TypedElementFramework/Public/TypedElementHandle.h
@@ -351,22 +351,22 @@ public:
 private:
 	FORCEINLINE void AddRef() const
 	{
-#if WITH_TYPED_ELEMENT_REFCOUNT
+#if UE_TYPED_ELEMENT_HAS_REFCOUNT
 		if (DataPtr)
 		{
 			DataPtr->AddRef();
 		}
-#endif	// WITH_TYPED_ELEMENT_REFCOUNT
+#endif	// UE_TYPED_ELEMENT_HAS_REFCOUNT
 	}
 
 	FORCEINLINE void ReleaseRef() const
 	{
-#if WITH_TYPED_ELEMENT_REFCOUNT
+#if UE_TYPED_ELEMENT_HAS_REFCOUNT
 		if (DataPtr)
 		{
 			DataPtr->ReleaseRef();
 		}
-#endif	// WITH_TYPED_ELEMENT_REFCOUNT
+#endif	// UE_TYPED_ELEMENT_HAS_REFCOUNT
 	}
 
 	FTypedElementId Id;
@@ -811,22 +811,22 @@ public:
 private:
 	FORCEINLINE void AddRef() const
 	{
-#if WITH_TYPED_ELEMENT_REFCOUNT
+#if UE_TYPED_ELEMENT_HAS_REFCOUNT
 		if (DataPtr)
 		{
 			DataPtr->AddRef();
 		}
-#endif	// WITH_TYPED_ELEMENT_REFCOUNT
+#endif	// UE_TYPED_ELEMENT_HAS_REFCOUNT
 	}
 
 	FORCEINLINE void ReleaseRef() const
 	{
-#if WITH_TYPED_ELEMENT_REFCOUNT
+#if UE_TYPED_ELEMENT_HAS_REFCOUNT
 		if (DataPtr)
 		{
 			DataPtr->ReleaseRef();
 		}
-#endif	// WITH_TYPED_ELEMENT_REFCOUNT
+#endif	// UE_TYPED_ELEMENT_HAS_REFCOUNT
 	}
 
 	FTypedElementId Id;
diff --git a/Engine/Source/Runtime/TypedElementFramework/Public/TypedElementLimits.h b/Engine/Source/Runtime/TypedElementFramework/Public/TypedElementLimits.h
--- a/Engine/Source/Runtime/TypedElementFramework/Public/TypedElementLimits.h
+++ b/Engine/Source/Runtime/TypedElementFramework/Public/TypedElementLimits.h
@@ -5,7 +5,7 @@
 #include "CoreTypes.h"
 #include "Math/NumericLimits.h"
 
-#define WITH_TYPED_ELEMENT_REFCOUNT (1)
+#define UE_TYPED_ELEMENT_HAS_REFCOUNT (1)
 
 /**
  * Handle ID limits, as used by FTypedElementId.
@@ -24,9 +24,9 @@ constexpr SIZE_T TypedHandleMaxElementId = ((SIZE_T)1 << TypedHandleElementIdBit
 using FTypedHandleTypeId = uint8;
 using FTypedHandleElementId = int32;
 using FTypedHandleCombinedId = uint32;
-#if WITH_TYPED_ELEMENT_REFCOUNT
+#if UE_TYPED_ELEMENT_HAS_REFCOUNT
 using FTypedHandleRefCount = int32;
-#endif	// WITH_TYPED_ELEMENT_REFCOUNT
+#endif	// UE_TYPED_ELEMENT_HAS_REFCOUNT
 
 static_assert(sizeof(FTypedHandleCombinedId) >= (TypedHandleTypeIdBytes + TypedHandleElementIdBytes), "FTypedHandleCombinedId is not large enough to hold the combination of TypedHandleTypeIdBytes and TypedHandleElementIdBytes!");
 
diff --git a/Engine/Source/Runtime/TypedElementFramework/Public/TypedElementRegistry.h b/Engine/Source/Runtime/TypedElementFramework/Public/TypedElementRegistry.h
--- a/Engine/Source/Runtime/TypedElementFramework/Public/TypedElementRegistry.h
+++ b/Engine/Source/Runtime/TypedElementFramework/Public/TypedElementRegistry.h
@@ -311,12 +311,12 @@ private:
 		FRegisteredElementType* RegisteredElementType = GetRegisteredElementTypeFromId(InOutElementOwner.GetId().GetTypeId());
 		checkf(RegisteredElementType, TEXT("Element type ID '%d' has not been registered!"), InOutElementOwner.GetId().GetTypeId());
 
-#if DO_CHECK && WITH_TYPED_ELEMENT_REFCOUNT
+#if DO_CHECK && UE_TYPED_ELEMENT_HAS_REFCOUNT
 		{
 			const FTypedHandleRefCount RefCount = RegisteredElementType->GetDataForElement(InOutElementOwner.GetId().GetElementId()).GetRefCount();
 			checkf(RefCount == 1, TEXT("Element is still externally referenced when being destroyed (ref-count: %d)!"), RefCount);
 		}
-#endif	// DO_CHECK && WITH_TYPED_ELEMENT_REFCOUNT
+#endif	// DO_CHECK && UE_TYPED_ELEMENT_HAS_REFCOUNT
 
 		RegisteredElementType->RemoveDataForElement(InOutElementOwner.GetId().GetElementId(), InOutElementOwner.Private_GetInternalData());
 		InOutElementOwner.Private_DestroyNoRef();
