Added optional reference tracking for typed element handles

This is controlled by the UE_TYPED_ELEMENT_HAS_REFTRACKING define, and adds extra debugging information that can be used to track down reference leaks.

diff --git a/Engine/Source/Runtime/TypedElementFramework/Private/TypedElementRegistry.cpp b/Engine/Source/Runtime/TypedElementFramework/Private/TypedElementRegistry.cpp
--- a/Engine/Source/Runtime/TypedElementFramework/Private/TypedElementRegistry.cpp
+++ b/Engine/Source/Runtime/TypedElementFramework/Private/TypedElementRegistry.cpp
@@ -122,10 +122,8 @@ void UTypedElementRegistry::ReleaseElementId(FTypedElementId& InOutElementId)
 	FRegisteredElementType* RegisteredElementType = GetRegisteredElementTypeFromId(InOutElementId.GetTypeId());
 	checkf(RegisteredElementType, TEXT("Element type ID '%d' has not been registered!"), InOutElementId.GetTypeId());
 
-#if UE_TYPED_ELEMENT_HAS_REFCOUNT
 	const FTypedElementInternalData& ElementData = RegisteredElementType->GetDataForElement(InOutElementId.GetElementId());
-	ElementData.ReleaseRef();
-#endif	// UE_TYPED_ELEMENT_HAS_REFCOUNT
+	ElementData.ReleaseRef(INDEX_NONE); // Cannot track element ID references as we have no space to store the reference ID
 
 	InOutElementId.Private_DestroyNoRef();
 }
diff --git a/Engine/Source/Runtime/TypedElementFramework/Public/TypedElementData.h b/Engine/Source/Runtime/TypedElementFramework/Public/TypedElementData.h
--- a/Engine/Source/Runtime/TypedElementFramework/Public/TypedElementData.h
+++ b/Engine/Source/Runtime/TypedElementFramework/Public/TypedElementData.h
@@ -3,7 +3,10 @@
 #pragma once
 
 #include "CoreMinimal.h"
+#include "Misc/ScopeLock.h"
 #include "Misc/ScopeRWLock.h"
+#include "HAL/PlatformStackWalk.h"
+#include "Containers/SparseArray.h"
 #include "Containers/ChunkedArray.h"
 #include "TypedElementLimits.h"
 
@@ -23,6 +26,53 @@
 #define UE_DEFINE_TYPED_ELEMENT_DATA_RTTI(ELEMENT_DATA_TYPE)						\
 	FTypedHandleTypeId ELEMENT_DATA_TYPE::Private_RegisteredTypeId = 0;
 
+#if UE_TYPED_ELEMENT_HAS_REFTRACKING
+/**
+ * Debugging information used to locate reference leaks.
+ */
+class FTypedElementReference
+{
+public:
+	FTypedElementReference()
+	{
+		CallstackDepth = FPlatformStackWalk::CaptureStackBackTrace(Callstack, UE_ARRAY_COUNT(Callstack));
+	}
+
+	FTypedElementReference(const FTypedElementReference&) = default;
+	FTypedElementReference& operator=(const FTypedElementReference&) = default;
+
+	FTypedElementReference(FTypedElementReference&&) = default;
+	FTypedElementReference& operator=(FTypedElementReference&&) = default;
+
+	void AppendString(FString& OutStr) const
+	{
+		ANSICHAR CallstackText[4096];
+		for (uint32 CallstackIndex = TypedHandleRefTrackingSkipCount; CallstackIndex < CallstackDepth; ++CallstackIndex)
+		{
+			if (CallstackIndex > TypedHandleRefTrackingSkipCount)
+			{
+				OutStr.Append(LINE_TERMINATOR);
+			}
+
+			CallstackText[0] = 0;
+			FPlatformStackWalk::ProgramCounterToHumanReadableString(CallstackIndex - TypedHandleRefTrackingSkipCount, Callstack[CallstackIndex], CallstackText, UE_ARRAY_COUNT(CallstackText));
+			OutStr.Append(CallstackText);
+		}
+	}
+
+	FString ToString() const
+	{
+		FString Str;
+		AppendString(Str);
+		return Str;
+	}
+
+private:
+	uint64 Callstack[TypedHandleRefTrackingDepth + TypedHandleRefTrackingSkipCount];
+	uint32 CallstackDepth = 0;
+};
+#endif	// UE_TYPED_ELEMENT_HAS_REFTRACKING
+
 /**
  * Base class for the internal payload data associated with elements.
  */
@@ -34,29 +84,90 @@ public:
 	FTypedElementInternalData(const FTypedElementInternalData&) = delete;
 	FTypedElementInternalData& operator=(const FTypedElementInternalData&) = delete;
 	
-	FTypedElementInternalData(FTypedElementInternalData&& InOther) = default;
-	FTypedElementInternalData& operator=(FTypedElementInternalData&&) = default;
+	FTypedElementInternalData(FTypedElementInternalData&& InOther) = delete;
+	FTypedElementInternalData& operator=(FTypedElementInternalData&&) = delete;
 
 	virtual ~FTypedElementInternalData() = default;
 
-#if UE_TYPED_ELEMENT_HAS_REFCOUNT
-	FORCEINLINE void AddRef() const
+	virtual void Reset()
 	{
-		checkSlow(RefCount < TNumericLimits<FTypedHandleRefCount>::Max());
+#if UE_TYPED_ELEMENT_HAS_REFCOUNTING
+		FPlatformAtomics::InterlockedExchange(&RefCount, 0);
+#endif	// UE_TYPED_ELEMENT_HAS_REFCOUNTING
+#if UE_TYPED_ELEMENT_HAS_REFTRACKING
+		{
+			FScopeLock ReferencesLock(&ReferencesCS);
+			References.Reset();
+		}
+#endif	// UE_TYPED_ELEMENT_HAS_REFTRACKING
+	}
+
+	FORCEINLINE FTypedElementReferenceId AddRef(const bool bCanTrackReference) const
+	{
+		FTypedElementReferenceId ReferenceId = INDEX_NONE;
+#if UE_TYPED_ELEMENT_HAS_REFCOUNTING
+		checkSlow(RefCount < TNumericLimits<FTypedElementRefCount>::Max());
 		FPlatformAtomics::InterlockedIncrement(&RefCount);
+#endif	// UE_TYPED_ELEMENT_HAS_REFCOUNTING
+#if UE_TYPED_ELEMENT_HAS_REFTRACKING
+		if (bCanTrackReference)
+		{
+			FScopeLock ReferencesLock(&ReferencesCS);
+			ReferenceId = References.Add(FTypedElementReference());
+		}
+#endif	// UE_TYPED_ELEMENT_HAS_REFTRACKING
+		return ReferenceId;
 	}
 
-	FORCEINLINE void ReleaseRef() const
+	FORCEINLINE void ReleaseRef(const FTypedElementReferenceId InReferenceId) const
 	{
+#if UE_TYPED_ELEMENT_HAS_REFCOUNTING
 		checkSlow(RefCount > 0);
 		FPlatformAtomics::InterlockedDecrement(&RefCount);
+#endif	// UE_TYPED_ELEMENT_HAS_REFCOUNTING
+#if UE_TYPED_ELEMENT_HAS_REFTRACKING
+		if (InReferenceId != INDEX_NONE)
+		{
+			FScopeLock ReferencesLock(&ReferencesCS);
+			References.RemoveAt(InReferenceId);
+		}
+#endif	// UE_TYPED_ELEMENT_HAS_REFTRACKING
 	}
 
-	FORCEINLINE FTypedHandleRefCount GetRefCount() const
+	FORCEINLINE FTypedElementRefCount GetRefCount() const
 	{
+#if UE_TYPED_ELEMENT_HAS_REFCOUNTING
 		return FPlatformAtomics::AtomicRead(&RefCount);
+#else	// UE_TYPED_ELEMENT_HAS_REFCOUNTING
+		return 0;
+#endif	// UE_TYPED_ELEMENT_HAS_REFCOUNTING
+	}
+
+	void AppendRefString(FString& OutStr) const
+	{
+#if UE_TYPED_ELEMENT_HAS_REFCOUNTING
+		OutStr.Appendf(TEXT("Ref-count: %d"), FPlatformAtomics::AtomicRead(&RefCount));
+#endif	// UE_TYPED_ELEMENT_HAS_REFCOUNTING
+#if UE_TYPED_ELEMENT_HAS_REFTRACKING
+		{
+			FScopeLock ReferencesLock(&ReferencesCS);
+			OutStr.Append(LINE_TERMINATOR TEXT("References:") LINE_TERMINATOR);
+			for (const FTypedElementReference& Reference : References)
+			{
+				OutStr.Append(TEXT("-----------------------------------------------") LINE_TERMINATOR);
+				Reference.AppendString(OutStr);
+			}
+			OutStr.Append(LINE_TERMINATOR TEXT("==============================================="));
+		}
+#endif	// UE_TYPED_ELEMENT_HAS_REFTRACKING
+	}
+
+	FString GetRefString() const
+	{
+		FString Str;
+		AppendRefString(Str);
+		return Str;
 	}
-#endif	// UE_TYPED_ELEMENT_HAS_REFCOUNT
 
 	virtual const void* GetUntypedData() const
 	{
@@ -64,9 +175,13 @@ public:
 	}
 
 private:
-#if UE_TYPED_ELEMENT_HAS_REFCOUNT
-	mutable FTypedHandleRefCount RefCount = 0;
-#endif	// UE_TYPED_ELEMENT_HAS_REFCOUNT
+#if UE_TYPED_ELEMENT_HAS_REFCOUNTING
+	mutable FTypedElementRefCount RefCount = 0;
+#endif	// UE_TYPED_ELEMENT_HAS_REFCOUNTING
+#if UE_TYPED_ELEMENT_HAS_REFTRACKING
+	mutable FCriticalSection ReferencesCS;
+	mutable TSparseArray<FTypedElementReference> References;
+#endif	// UE_TYPED_ELEMENT_HAS_REFTRACKING
 };
 
 /**
@@ -76,12 +191,11 @@ template <typename ElementDataType>
 class TTypedElementInternalData : public FTypedElementInternalData
 {
 public:
-	TTypedElementInternalData() = default;
-
-	TTypedElementInternalData(TTypedElementInternalData&& InOther) = default;
-	TTypedElementInternalData& operator=(TTypedElementInternalData&&) = default;
-
-	virtual ~TTypedElementInternalData() = default;
+	virtual void Reset() override
+	{
+		FTypedElementInternalData::Reset();
+		Data = ElementDataType();
+	}
 
 	FORCEINLINE const ElementDataType& GetData() const
 	{
@@ -108,13 +222,6 @@ private:
 template <>
 class TTypedElementInternalData<void> : public FTypedElementInternalData
 {
-public:
-	TTypedElementInternalData() = default;
-
-	TTypedElementInternalData(TTypedElementInternalData&& InOther) = default;
-	TTypedElementInternalData& operator=(TTypedElementInternalData&&) = default;
-
-	virtual ~TTypedElementInternalData() = default;
 };
 
 /**
@@ -148,7 +255,7 @@ public:
 		
 		TTypedElementInternalData<ElementDataType>& InternalData = InternalDataArray[InElementId];
 		checkf(InExpectedDataPtr == &InternalData, TEXT("Internal data pointer did not match the expected value! Does this handle belong to a different element registry?"));
-		InternalData = TTypedElementInternalData<ElementDataType>();
+		InternalData.Reset();
 		InternalDataFreeIndices.Add(InElementId);
 	}
 
@@ -194,7 +301,7 @@ public:
 
 	TTypedElementInternalData<void>& AddDataForElement(FTypedHandleElementId& InOutElementId)
 	{
-#if UE_TYPED_ELEMENT_HAS_REFCOUNT
+#if UE_TYPED_ELEMENT_HAS_REFERENCING
 		FWriteScopeLock InternalDataLock(InternalDataRW);
 
 		checkSlow(InOutElementId >= 0);
@@ -206,14 +313,14 @@ public:
 
 		ElementIdToArrayIndex.Add(InOutElementId, InternalDataArrayIndex);
 		return InternalDataArray[InternalDataArrayIndex];
-#else	// UE_TYPED_ELEMENT_HAS_REFCOUNT
+#else	// UE_TYPED_ELEMENT_HAS_REFERENCING
 		return SharedInternalData;
-#endif	// UE_TYPED_ELEMENT_HAS_REFCOUNT
+#endif	// UE_TYPED_ELEMENT_HAS_REFERENCING
 	}
 
 	void RemoveDataForElement(const FTypedHandleElementId InElementId, const FTypedElementInternalData* InExpectedDataPtr)
 	{
-#if UE_TYPED_ELEMENT_HAS_REFCOUNT
+#if UE_TYPED_ELEMENT_HAS_REFERENCING
 		FWriteScopeLock InternalDataLock(InternalDataRW);
 
 		int32 InternalDataArrayIndex = INDEX_NONE;
@@ -223,24 +330,24 @@ public:
 
 		TTypedElementInternalData<void>& InternalData = InternalDataArray[InElementId];
 		checkf(InExpectedDataPtr == &InternalData, TEXT("Internal data pointer did not match the expected value! Does this handle belong to a different element registry?"));
-		InternalData = TTypedElementInternalData<void>();
+		InternalData.Reset();
 		InternalDataFreeIndices.Add(InternalDataArrayIndex);
-#else	// UE_TYPED_ELEMENT_HAS_REFCOUNT
+#else	// UE_TYPED_ELEMENT_HAS_REFERENCING
 		checkf(InExpectedDataPtr == &SharedInternalData, TEXT("Internal data pointer did not match the expected value! Does this handle belong to a different element registry?"));
-#endif	// UE_TYPED_ELEMENT_HAS_REFCOUNT
+#endif	// UE_TYPED_ELEMENT_HAS_REFERENCING
 	}
 
 	const TTypedElementInternalData<void>& GetDataForElement(const FTypedHandleElementId InElementId) const
 	{
-#if UE_TYPED_ELEMENT_HAS_REFCOUNT
+#if UE_TYPED_ELEMENT_HAS_REFERENCING
 		FReadScopeLock InternalDataLock(InternalDataRW);
 
 		const int32* InternalDataArrayIndexPtr = ElementIdToArrayIndex.Find(InElementId);
 		checkSlow(InternalDataArrayIndexPtr && InternalDataArray.IsValidIndex(*InternalDataArrayIndexPtr));
 		return InternalDataArray[*InternalDataArrayIndexPtr];
-#else	// UE_TYPED_ELEMENT_HAS_REFCOUNT
+#else	// UE_TYPED_ELEMENT_HAS_REFERENCING
 		return SharedInternalData;
-#endif	// UE_TYPED_ELEMENT_HAS_REFCOUNT
+#endif	// UE_TYPED_ELEMENT_HAS_REFERENCING
 	}
 
 	static FORCEINLINE void SetStaticDataTypeId(const FTypedHandleTypeId InTypeId)
@@ -258,12 +365,12 @@ public:
 	}
 
 private:
-#if UE_TYPED_ELEMENT_HAS_REFCOUNT
+#if UE_TYPED_ELEMENT_HAS_REFERENCING
 	mutable FRWLock InternalDataRW;
 	TChunkedArray<TTypedElementInternalData<void>> InternalDataArray;
 	TArray<int32> InternalDataFreeIndices;
 	TMap<FTypedHandleElementId, int32> ElementIdToArrayIndex;
-#else	// UE_TYPED_ELEMENT_HAS_REFCOUNT
+#else	// UE_TYPED_ELEMENT_HAS_REFERENCING
 	TTypedElementInternalData<void> SharedInternalData;
-#endif	// UE_TYPED_ELEMENT_HAS_REFCOUNT
+#endif	// UE_TYPED_ELEMENT_HAS_REFERENCING
 };
diff --git a/Engine/Source/Runtime/TypedElementFramework/Public/TypedElementHandle.h b/Engine/Source/Runtime/TypedElementFramework/Public/TypedElementHandle.h
--- a/Engine/Source/Runtime/TypedElementFramework/Public/TypedElementHandle.h
+++ b/Engine/Source/Runtime/TypedElementFramework/Public/TypedElementHandle.h
@@ -177,8 +177,14 @@ public:
 	FTypedElementHandle(FTypedElementHandle&& InOther)
 		: Id(MoveTemp(InOther.Id))
 		, DataPtr(InOther.DataPtr)
+#if UE_TYPED_ELEMENT_HAS_REFTRACKING
+		, ReferenceId(InOther.ReferenceId)
+#endif	// UE_TYPED_ELEMENT_HAS_REFTRACKING
 	{
 		InOther.DataPtr = nullptr;
+#if UE_TYPED_ELEMENT_HAS_REFTRACKING
+		InOther.ReferenceId = INDEX_NONE;
+#endif	// UE_TYPED_ELEMENT_HAS_REFTRACKING
 		checkSlow(!InOther.IsSet());
 	}
 
@@ -190,8 +196,14 @@ public:
 
 			Id = MoveTemp(InOther.Id);
 			DataPtr = InOther.DataPtr;
+#if UE_TYPED_ELEMENT_HAS_REFTRACKING
+			ReferenceId = InOther.ReferenceId;
+#endif	// UE_TYPED_ELEMENT_HAS_REFTRACKING
 
 			InOther.DataPtr = nullptr;
+#if UE_TYPED_ELEMENT_HAS_REFTRACKING
+			InOther.ReferenceId = INDEX_NONE;
+#endif	// UE_TYPED_ELEMENT_HAS_REFTRACKING
 			checkSlow(!InOther.IsSet());
 		}
 		return *this;
@@ -284,7 +296,7 @@ public:
 		FTypedElementId ElementId;
 		if (IsSet())
 		{
-			AddRef();
+			DataPtr->AddRef(/*bCanTrackReference*/false); // Cannot track element ID references as we have no space to store the reference ID
 			ElementId.Private_InitializeNoRef(Id.GetTypeId(), Id.GetElementId());
 		}
 		return ElementId;
@@ -299,7 +311,7 @@ public:
 		checkf(InOutElementId == Id, TEXT("Element ID does not match this handle!"));
 		if (InOutElementId)
 		{
-			ReleaseRef();
+			DataPtr->ReleaseRef(INDEX_NONE); // Cannot track element ID references as we have no space to store the reference ID
 			InOutElementId.Private_DestroyNoRef();
 		}
 	}
@@ -328,18 +340,21 @@ public:
 	FORCEINLINE void Private_InitializeAddRef(const FTypedHandleTypeId InTypeId, const FTypedHandleElementId InElementId, const FTypedElementInternalData& InData)
 	{
 		Private_InitializeNoRef(InTypeId, InElementId, InData);
-		AddRef();
+		RegisterRef();
 	}
 
 	FORCEINLINE void Private_DestroyNoRef()
 	{
 		Id.Private_DestroyNoRef();
 		DataPtr = nullptr;
+#if UE_TYPED_ELEMENT_HAS_REFTRACKING
+		ReferenceId = INDEX_NONE;
+#endif	// UE_TYPED_ELEMENT_HAS_REFTRACKING
 	}
 
 	FORCEINLINE void Private_DestroyReleaseRef()
 	{
-		ReleaseRef();
+		UnregisterRef();
 		Private_DestroyNoRef();
 	}
 
@@ -349,28 +364,33 @@ public:
 	}
 
 private:
-	FORCEINLINE void AddRef() const
+	FORCEINLINE void RegisterRef()
 	{
-#if UE_TYPED_ELEMENT_HAS_REFCOUNT
 		if (DataPtr)
 		{
-			DataPtr->AddRef();
+#if !UE_TYPED_ELEMENT_HAS_REFTRACKING
+			FTypedElementReferenceId ReferenceId = INDEX_NONE;
+#endif	// !UE_TYPED_ELEMENT_HAS_REFTRACKING
+			ReferenceId = DataPtr->AddRef(/*bCanTrackReference*/true);
 		}
-#endif	// UE_TYPED_ELEMENT_HAS_REFCOUNT
 	}
 
-	FORCEINLINE void ReleaseRef() const
+	FORCEINLINE void UnregisterRef()
 	{
-#if UE_TYPED_ELEMENT_HAS_REFCOUNT
 		if (DataPtr)
 		{
-			DataPtr->ReleaseRef();
+#if !UE_TYPED_ELEMENT_HAS_REFTRACKING
+			FTypedElementReferenceId ReferenceId = INDEX_NONE;
+#endif	// !UE_TYPED_ELEMENT_HAS_REFTRACKING
+			DataPtr->ReleaseRef(ReferenceId);
 		}
-#endif	// UE_TYPED_ELEMENT_HAS_REFCOUNT
 	}
 
 	FTypedElementId Id;
 	const FTypedElementInternalData* DataPtr = nullptr;
+#if UE_TYPED_ELEMENT_HAS_REFTRACKING
+	FTypedElementReferenceId ReferenceId = INDEX_NONE;
+#endif	// UE_TYPED_ELEMENT_HAS_REFTRACKING
 };
 
 /**
@@ -720,7 +740,7 @@ public:
 		FTypedElementId ElementId;
 		if (IsSet())
 		{
-			AddRef();
+			DataPtr->AddRef(/*bCanTrackReference*/false); // Cannot track element ID references as we have no space to store the reference ID
 			ElementId.Private_InitializeNoRef(Id.GetTypeId(), Id.GetElementId());
 		}
 		return ElementId;
@@ -735,7 +755,7 @@ public:
 		checkf(InOutElementId == Id, TEXT("Element ID does not match this owner!"));
 		if (InOutElementId)
 		{
-			ReleaseRef();
+			DataPtr->ReleaseRef(INDEX_NONE); // Cannot track element ID references as we have no space to store the reference ID
 			InOutElementId.Private_DestroyNoRef();
 		}
 	}
@@ -788,18 +808,21 @@ public:
 	FORCEINLINE void Private_InitializeAddRef(const FTypedHandleTypeId InTypeId, const FTypedHandleElementId InElementId, TTypedElementInternalData<ElementDataType>& InData)
 	{
 		Private_InitializeNoRef(InTypeId, InElementId, InData);
-		AddRef();
+		RegisterRef();
 	}
 
 	FORCEINLINE void Private_DestroyNoRef()
 	{
 		Id.Private_DestroyNoRef();
 		DataPtr = nullptr;
+#if UE_TYPED_ELEMENT_HAS_REFTRACKING
+		ReferenceId = INDEX_NONE;
+#endif	// UE_TYPED_ELEMENT_HAS_REFTRACKING
 	}
 
 	FORCEINLINE void Private_DestroyReleaseRef()
 	{
-		ReleaseRef();
+		UnregisterRef();
 		Private_DestroyNoRef();
 	}
 
@@ -809,27 +832,32 @@ public:
 	}
 
 private:
-	FORCEINLINE void AddRef() const
+	FORCEINLINE void RegisterRef()
 	{
-#if UE_TYPED_ELEMENT_HAS_REFCOUNT
 		if (DataPtr)
 		{
-			DataPtr->AddRef();
+#if !UE_TYPED_ELEMENT_HAS_REFTRACKING
+			FTypedElementReferenceId ReferenceId = INDEX_NONE;
+#endif	// !UE_TYPED_ELEMENT_HAS_REFTRACKING
+			ReferenceId = DataPtr->AddRef(/*bCanTrackReference*/true);
 		}
-#endif	// UE_TYPED_ELEMENT_HAS_REFCOUNT
 	}
 
-	FORCEINLINE void ReleaseRef() const
+	FORCEINLINE void UnregisterRef()
 	{
-#if UE_TYPED_ELEMENT_HAS_REFCOUNT
 		if (DataPtr)
 		{
-			DataPtr->ReleaseRef();
+#if !UE_TYPED_ELEMENT_HAS_REFTRACKING
+			FTypedElementReferenceId ReferenceId = INDEX_NONE;
+#endif	// !UE_TYPED_ELEMENT_HAS_REFTRACKING
+			DataPtr->ReleaseRef(ReferenceId);
 		}
-#endif	// UE_TYPED_ELEMENT_HAS_REFCOUNT
 	}
 
 	FTypedElementId Id;
 	TTypedElementInternalData<ElementDataType>* DataPtr = nullptr;
+#if UE_TYPED_ELEMENT_HAS_REFTRACKING
+	FTypedElementReferenceId ReferenceId = INDEX_NONE;
+#endif	// UE_TYPED_ELEMENT_HAS_REFTRACKING
 };
 using FTypedElementOwner = TTypedElementOwner<void>;
diff --git a/Engine/Source/Runtime/TypedElementFramework/Public/TypedElementLimits.h b/Engine/Source/Runtime/TypedElementFramework/Public/TypedElementLimits.h
--- a/Engine/Source/Runtime/TypedElementFramework/Public/TypedElementLimits.h
+++ b/Engine/Source/Runtime/TypedElementFramework/Public/TypedElementLimits.h
@@ -5,7 +5,9 @@
 #include "CoreTypes.h"
 #include "Math/NumericLimits.h"
 
-#define UE_TYPED_ELEMENT_HAS_REFCOUNT (1)
+#define UE_TYPED_ELEMENT_HAS_REFCOUNTING (1)
+#define UE_TYPED_ELEMENT_HAS_REFTRACKING (0)
+#define UE_TYPED_ELEMENT_HAS_REFERENCING (UE_TYPED_ELEMENT_HAS_REFCOUNTING || UE_TYPED_ELEMENT_HAS_REFTRACKING)
 
 /**
  * Handle ID limits, as used by FTypedElementId.
@@ -21,12 +23,15 @@ constexpr SIZE_T TypedHandleElementIdBytes = TypedHandleElementIdBits >> 3;
 constexpr SIZE_T TypedHandleMaxTypeId = ((SIZE_T)1 << TypedHandleTypeIdBits) - 1;
 constexpr SIZE_T TypedHandleMaxElementId = ((SIZE_T)1 << TypedHandleElementIdBits) - 1;
 
+constexpr SIZE_T TypedHandleRefTrackingSkipCount = 1;
+constexpr SIZE_T TypedHandleRefTrackingDepth = 15;
+
 using FTypedHandleTypeId = uint8;
 using FTypedHandleElementId = int32;
 using FTypedHandleCombinedId = uint32;
-#if UE_TYPED_ELEMENT_HAS_REFCOUNT
-using FTypedHandleRefCount = int32;
-#endif	// UE_TYPED_ELEMENT_HAS_REFCOUNT
+
+using FTypedElementRefCount = int32;
+using FTypedElementReferenceId = int32;
 
 static_assert(sizeof(FTypedHandleCombinedId) >= (TypedHandleTypeIdBytes + TypedHandleElementIdBytes), "FTypedHandleCombinedId is not large enough to hold the combination of TypedHandleTypeIdBytes and TypedHandleElementIdBytes!");
 
diff --git a/Engine/Source/Runtime/TypedElementFramework/Public/TypedElementRegistry.h b/Engine/Source/Runtime/TypedElementFramework/Public/TypedElementRegistry.h
--- a/Engine/Source/Runtime/TypedElementFramework/Public/TypedElementRegistry.h
+++ b/Engine/Source/Runtime/TypedElementFramework/Public/TypedElementRegistry.h
@@ -311,12 +311,13 @@ private:
 		FRegisteredElementType* RegisteredElementType = GetRegisteredElementTypeFromId(InOutElementOwner.GetId().GetTypeId());
 		checkf(RegisteredElementType, TEXT("Element type ID '%d' has not been registered!"), InOutElementOwner.GetId().GetTypeId());
 
-#if DO_CHECK && UE_TYPED_ELEMENT_HAS_REFCOUNT
+#if DO_CHECK
 		{
-			const FTypedHandleRefCount RefCount = RegisteredElementType->GetDataForElement(InOutElementOwner.GetId().GetElementId()).GetRefCount();
-			checkf(RefCount == 1, TEXT("Element is still externally referenced when being destroyed (ref-count: %d)!"), RefCount);
+			const FTypedElementInternalData& ElementData = RegisteredElementType->GetDataForElement(InOutElementOwner.GetId().GetElementId());
+			const FTypedElementRefCount RefCount = ElementData.GetRefCount();
+			checkf(RefCount <= 1, TEXT("Element is still externally referenced when being destroyed! %s"), *ElementData.GetRefString());
 		}
-#endif	// DO_CHECK && UE_TYPED_ELEMENT_HAS_REFCOUNT
+#endif	// DO_CHECK
 
 		RegisteredElementType->RemoveDataForElement(InOutElementOwner.GetId().GetElementId(), InOutElementOwner.Private_GetInternalData());
 		InOutElementOwner.Private_DestroyNoRef();
