Basic framework to allow UnrealEd to override the base Engine Typed Element interfaces

#rnx

diff --git a/Engine/Source/Editor/UnrealEd/Private/Elements/Actor/ActorElementEditorSelectionInterface.cpp b/Engine/Source/Editor/UnrealEd/Private/Elements/Actor/ActorElementEditorSelectionInterface.cpp
new file mode 100644
--- /dev/null
+++ b/Engine/Source/Editor/UnrealEd/Private/Elements/Actor/ActorElementEditorSelectionInterface.cpp
@@ -0,0 +1,5 @@
+// Copyright Epic Games, Inc. All Rights Reserved.
+
+#include "Elements/Actor/ActorElementEditorSelectionInterface.h"
+#include "Elements/Actor/ActorElementData.h"
+#include "GameFramework/Actor.h"
diff --git a/Engine/Source/Editor/UnrealEd/Private/Elements/Actor/ActorElementEditorSelectionInterface.h b/Engine/Source/Editor/UnrealEd/Private/Elements/Actor/ActorElementEditorSelectionInterface.h
new file mode 100644
--- /dev/null
+++ b/Engine/Source/Editor/UnrealEd/Private/Elements/Actor/ActorElementEditorSelectionInterface.h
@@ -0,0 +1,14 @@
+// Copyright Epic Games, Inc. All Rights Reserved.
+
+#pragma once
+
+#include "Elements/Actor/ActorElementSelectionInterface.h"
+#include "ActorElementEditorSelectionInterface.generated.h"
+
+UCLASS()
+class UActorElementEditorSelectionInterface : public UActorElementSelectionInterface
+{
+	GENERATED_BODY()
+
+public:
+};
diff --git a/Engine/Source/Editor/UnrealEd/Private/Elements/Component/ComponentElementEditorSelectionInterface.cpp b/Engine/Source/Editor/UnrealEd/Private/Elements/Component/ComponentElementEditorSelectionInterface.cpp
new file mode 100644
--- /dev/null
+++ b/Engine/Source/Editor/UnrealEd/Private/Elements/Component/ComponentElementEditorSelectionInterface.cpp
@@ -0,0 +1,5 @@
+// Copyright Epic Games, Inc. All Rights Reserved.
+
+#include "Elements/Component/ComponentElementEditorSelectionInterface.h"
+#include "Elements/Component/ComponentElementData.h"
+#include "Components/ActorComponent.h"
diff --git a/Engine/Source/Editor/UnrealEd/Private/Elements/Component/ComponentElementEditorSelectionInterface.h b/Engine/Source/Editor/UnrealEd/Private/Elements/Component/ComponentElementEditorSelectionInterface.h
new file mode 100644
--- /dev/null
+++ b/Engine/Source/Editor/UnrealEd/Private/Elements/Component/ComponentElementEditorSelectionInterface.h
@@ -0,0 +1,14 @@
+// Copyright Epic Games, Inc. All Rights Reserved.
+
+#pragma once
+
+#include "Elements/Component/ComponentElementSelectionInterface.h"
+#include "ComponentElementEditorSelectionInterface.generated.h"
+
+UCLASS()
+class UComponentElementEditorSelectionInterface : public UComponentElementSelectionInterface
+{
+	GENERATED_BODY()
+
+public:
+};
diff --git a/Engine/Source/Editor/UnrealEd/Private/Elements/EditorElements.cpp b/Engine/Source/Editor/UnrealEd/Private/Elements/EditorElements.cpp
new file mode 100644
--- /dev/null
+++ b/Engine/Source/Editor/UnrealEd/Private/Elements/EditorElements.cpp
@@ -0,0 +1,28 @@
+// Copyright Epic Games, Inc. All Rights Reserved.
+
+#include "EditorElements.h"
+#include "TypedElementRegistry.h"
+
+#include "Elements/Actor/ActorElementEditorSelectionInterface.h"
+
+#include "Elements/Component/ComponentElementEditorSelectionInterface.h"
+
+void RegisterEditorActorElements()
+{
+	UTypedElementRegistry* Registry = UTypedElementRegistry::GetInstance();
+
+	Registry->RegisterElementInterface<UTypedElementSelectionInterface>(NAME_Actor, NewObject<UActorElementEditorSelectionInterface>(), /*bAllowOverride*/true);
+}
+
+void RegisterEditorComponentElements()
+{
+	UTypedElementRegistry* Registry = UTypedElementRegistry::GetInstance();
+
+	Registry->RegisterElementInterface<UTypedElementSelectionInterface>(NAME_Components, NewObject<UComponentElementEditorSelectionInterface>(), /*bAllowOverride*/true);
+}
+
+void RegisterEditorElements()
+{
+	RegisterEditorActorElements();
+	RegisterEditorComponentElements();
+}
diff --git a/Engine/Source/Editor/UnrealEd/Private/Elements/EditorElements.h b/Engine/Source/Editor/UnrealEd/Private/Elements/EditorElements.h
new file mode 100644
--- /dev/null
+++ b/Engine/Source/Editor/UnrealEd/Private/Elements/EditorElements.h
@@ -0,0 +1,5 @@
+// Copyright Epic Games, Inc. All Rights Reserved.
+
+#pragma once
+
+void RegisterEditorElements();
diff --git a/Engine/Source/Editor/UnrealEd/Private/UnrealEdEngine.cpp b/Engine/Source/Editor/UnrealEd/Private/UnrealEdEngine.cpp
--- a/Engine/Source/Editor/UnrealEd/Private/UnrealEdEngine.cpp
+++ b/Engine/Source/Editor/UnrealEd/Private/UnrealEdEngine.cpp
@@ -32,6 +32,7 @@
 #include "EditorModeRegistry.h"
 #include "EditorModeManager.h"
 #include "EditorModes.h"
+#include "Elements/EditorElements.h"
 #include "UnrealEdMisc.h"
 #include "UnrealEdGlobals.h"
 
@@ -76,6 +77,8 @@ void UUnrealEdEngine::Init(IEngineLoop* InEngineLoop)
 {
 	Super::Init(InEngineLoop);
 
+	RegisterEditorElements();
+
 	// Display warnings to the user about disk space issues
 	ValidateFreeDiskSpace();
 
diff --git a/Engine/Source/Runtime/Engine/Private/Elements/Actor/ActorElementData.cpp b/Engine/Source/Runtime/Engine/Private/Elements/Actor/ActorElementData.cpp
--- a/Engine/Source/Runtime/Engine/Private/Elements/Actor/ActorElementData.cpp
+++ b/Engine/Source/Runtime/Engine/Private/Elements/Actor/ActorElementData.cpp
@@ -1,5 +1,5 @@
 // Copyright Epic Games, Inc. All Rights Reserved.
 
-#include "ActorElementData.h"
+#include "Elements/Actor/ActorElementData.h"
 
 UE_DEFINE_TYPED_ELEMENT_DATA_RTTI(FActorElementData);
diff --git a/Engine/Source/Runtime/Engine/Private/Elements/Actor/ActorElementData.h b/Engine/Source/Runtime/Engine/Private/Elements/Actor/ActorElementData.h
deleted file mode 100644
diff --git a/Engine/Source/Runtime/Engine/Private/Elements/Actor/ActorElementSelectionInterface.cpp b/Engine/Source/Runtime/Engine/Private/Elements/Actor/ActorElementSelectionInterface.cpp
--- a/Engine/Source/Runtime/Engine/Private/Elements/Actor/ActorElementSelectionInterface.cpp
+++ b/Engine/Source/Runtime/Engine/Private/Elements/Actor/ActorElementSelectionInterface.cpp
@@ -1,17 +1,37 @@
 // Copyright Epic Games, Inc. All Rights Reserved.
 
-#include "ActorElementSelectionInterface.h"
-#include "ActorElementData.h"
+#include "Elements/Actor/ActorElementSelectionInterface.h"
+#include "Elements/Actor/ActorElementData.h"
 #include "GameFramework/Actor.h"
-
-bool UActorElementSelectionInterface::IsValidSelection(const FTypedElementHandle& InElementHandle)
-{
-	// TODO: Validate that this actor is in a valid state to be selected?
-	return true;
-}
+#include "TypedElementList.h"
 
 UObject* UActorElementSelectionInterface::Legacy_GetSelectionObject(const FTypedElementHandle& InElementHandle)
 {
 	const FActorElementData* ActorData = InElementHandle.GetData<FActorElementData>();
 	return ActorData ? ActorData->Actor : nullptr;
 }
+
+int32 UActorElementSelectionInterface::GetNumSelectedActors(const FTypedElementList& InCurrentSelection)
+{
+	int32 NumSelected = 0;
+	for (const FTypedElementHandle& SelectedElement : InCurrentSelection)
+	{
+		if (SelectedElement.GetData<FActorElementData>(/*bSilent*/true))
+		{
+			++NumSelected;
+		}
+	}
+	return NumSelected;
+}
+
+bool UActorElementSelectionInterface::HasSelectedActors(const FTypedElementList& InCurrentSelection)
+{
+	for (const FTypedElementHandle& SelectedElement : InCurrentSelection)
+	{
+		if (SelectedElement.GetData<FActorElementData>(/*bSilent*/true))
+		{
+			return true;
+		}
+	}
+	return false;
+}
diff --git a/Engine/Source/Runtime/Engine/Private/Elements/Actor/ActorElementSelectionInterface.h b/Engine/Source/Runtime/Engine/Private/Elements/Actor/ActorElementSelectionInterface.h
deleted file mode 100644
diff --git a/Engine/Source/Runtime/Engine/Private/Elements/Component/ComponentElementData.cpp b/Engine/Source/Runtime/Engine/Private/Elements/Component/ComponentElementData.cpp
--- a/Engine/Source/Runtime/Engine/Private/Elements/Component/ComponentElementData.cpp
+++ b/Engine/Source/Runtime/Engine/Private/Elements/Component/ComponentElementData.cpp
@@ -1,5 +1,5 @@
 // Copyright Epic Games, Inc. All Rights Reserved.
 
-#include "ComponentElementData.h"
+#include "Elements/Component/ComponentElementData.h"
 
 UE_DEFINE_TYPED_ELEMENT_DATA_RTTI(FComponentElementData);
diff --git a/Engine/Source/Runtime/Engine/Private/Elements/Component/ComponentElementData.h b/Engine/Source/Runtime/Engine/Private/Elements/Component/ComponentElementData.h
deleted file mode 100644
diff --git a/Engine/Source/Runtime/Engine/Private/Elements/Component/ComponentElementSelectionInterface.cpp b/Engine/Source/Runtime/Engine/Private/Elements/Component/ComponentElementSelectionInterface.cpp
--- a/Engine/Source/Runtime/Engine/Private/Elements/Component/ComponentElementSelectionInterface.cpp
+++ b/Engine/Source/Runtime/Engine/Private/Elements/Component/ComponentElementSelectionInterface.cpp
@@ -1,17 +1,37 @@
 // Copyright Epic Games, Inc. All Rights Reserved.
 
-#include "ComponentElementSelectionInterface.h"
-#include "ComponentElementData.h"
+#include "Elements/Component/ComponentElementSelectionInterface.h"
+#include "Elements/Component/ComponentElementData.h"
 #include "Components/ActorComponent.h"
-
-bool UComponentElementSelectionInterface::IsValidSelection(const FTypedElementHandle& InElementHandle)
-{
-	// TODO: Validate that this component is in a valid state to be selected?
-	return true;
-}
+#include "TypedElementList.h"
 
 UObject* UComponentElementSelectionInterface::Legacy_GetSelectionObject(const FTypedElementHandle& InElementHandle)
 {
 	const FComponentElementData* ComponentData = InElementHandle.GetData<FComponentElementData>();
 	return ComponentData ? ComponentData->Component : nullptr;
 }
+
+int32 UComponentElementSelectionInterface::GetNumSelectedComponents(const FTypedElementList& InCurrentSelection)
+{
+	int32 NumSelected = 0;
+	for (const FTypedElementHandle& SelectedElement : InCurrentSelection)
+	{
+		if (SelectedElement.GetData<FComponentElementData>(/*bSilent*/true))
+		{
+			++NumSelected;
+		}
+	}
+	return NumSelected;
+}
+
+bool UComponentElementSelectionInterface::HasSelectedComponents(const FTypedElementList& InCurrentSelection)
+{
+	for (const FTypedElementHandle& SelectedElement : InCurrentSelection)
+	{
+		if (SelectedElement.GetData<FComponentElementData>(/*bSilent*/true))
+		{
+			return true;
+		}
+	}
+	return false;
+}
diff --git a/Engine/Source/Runtime/Engine/Private/Elements/Component/ComponentElementSelectionInterface.h b/Engine/Source/Runtime/Engine/Private/Elements/Component/ComponentElementSelectionInterface.h
deleted file mode 100644
diff --git a/Engine/Source/Runtime/Engine/Private/Elements/EngineElements.cpp b/Engine/Source/Runtime/Engine/Private/Elements/EngineElements.cpp
--- a/Engine/Source/Runtime/Engine/Private/Elements/EngineElements.cpp
+++ b/Engine/Source/Runtime/Engine/Private/Elements/EngineElements.cpp
@@ -5,13 +5,13 @@
 
 #include "TypedElementRegistry.h"
 
-#include "Actor/ActorElementData.h"
-#include "Actor/ActorElementSelectionInterface.h"
+#include "Elements/Actor/ActorElementData.h"
+#include "Elements/Actor/ActorElementSelectionInterface.h"
 
-#include "Component/ComponentElementData.h"
-#include "Component/ComponentElementSelectionInterface.h"
+#include "Elements/Component/ComponentElementData.h"
+#include "Elements/Component/ComponentElementSelectionInterface.h"
 
-void RegisterActorElements()
+void RegisterEngineActorElements()
 {
 	UTypedElementRegistry* Registry = UTypedElementRegistry::GetInstance();
 
@@ -19,7 +19,7 @@ void RegisterActorElements()
 	Registry->RegisterElementInterface<UTypedElementSelectionInterface>(NAME_Actor, NewObject<UActorElementSelectionInterface>());
 }
 
-void RegisterComponentElements()
+void RegisterEngineComponentElements()
 {
 	UTypedElementRegistry* Registry = UTypedElementRegistry::GetInstance();
 
@@ -33,6 +33,6 @@ void RegisterEngineElements()
 	FModuleManager::Get().LoadModuleChecked("TypedElementFramework");
 	FModuleManager::Get().LoadModuleChecked("TypedElementInterfaces");
 
-	RegisterActorElements();
-	RegisterComponentElements();
+	RegisterEngineActorElements();
+	RegisterEngineComponentElements();
 }
diff --git a/Engine/Source/Runtime/Engine/Public/Elements/Actor/ActorElementData.h b/Engine/Source/Runtime/Engine/Public/Elements/Actor/ActorElementData.h
new file mode 100644
--- /dev/null
+++ b/Engine/Source/Runtime/Engine/Public/Elements/Actor/ActorElementData.h
@@ -0,0 +1,17 @@
+// Copyright Epic Games, Inc. All Rights Reserved.
+
+#pragma once
+
+#include "TypedElementData.h"
+
+class AActor;
+
+/**
+ * Element data that represents an Actor.
+ */
+struct ENGINE_API FActorElementData
+{
+	UE_DECLARE_TYPED_ELEMENT_DATA_RTTI(FActorElementData);
+
+	AActor* Actor = nullptr;
+};
diff --git a/Engine/Source/Runtime/Engine/Public/Elements/Actor/ActorElementSelectionInterface.h b/Engine/Source/Runtime/Engine/Public/Elements/Actor/ActorElementSelectionInterface.h
new file mode 100644
--- /dev/null
+++ b/Engine/Source/Runtime/Engine/Public/Elements/Actor/ActorElementSelectionInterface.h
@@ -0,0 +1,20 @@
+// Copyright Epic Games, Inc. All Rights Reserved.
+
+#pragma once
+
+#include "TypedElementSelectionInterface.h"
+#include "ActorElementSelectionInterface.generated.h"
+
+class FTypedElementList;
+
+UCLASS()
+class ENGINE_API UActorElementSelectionInterface : public UTypedElementSelectionInterface
+{
+	GENERATED_BODY()
+
+public:
+	virtual UObject* Legacy_GetSelectionObject(const FTypedElementHandle& InElementHandle) override;
+
+	static int32 GetNumSelectedActors(const FTypedElementList& InCurrentSelection);
+	static bool HasSelectedActors(const FTypedElementList& InCurrentSelection);
+};
diff --git a/Engine/Source/Runtime/Engine/Public/Elements/Component/ComponentElementData.h b/Engine/Source/Runtime/Engine/Public/Elements/Component/ComponentElementData.h
new file mode 100644
--- /dev/null
+++ b/Engine/Source/Runtime/Engine/Public/Elements/Component/ComponentElementData.h
@@ -0,0 +1,17 @@
+// Copyright Epic Games, Inc. All Rights Reserved.
+
+#pragma once
+
+#include "TypedElementData.h"
+
+class UActorComponent;
+
+/**
+ * Element data that represents an Actor Component.
+ */
+struct ENGINE_API FComponentElementData
+{
+	UE_DECLARE_TYPED_ELEMENT_DATA_RTTI(FComponentElementData);
+
+	UActorComponent* Component = nullptr;
+};
diff --git a/Engine/Source/Runtime/Engine/Public/Elements/Component/ComponentElementSelectionInterface.h b/Engine/Source/Runtime/Engine/Public/Elements/Component/ComponentElementSelectionInterface.h
new file mode 100644
--- /dev/null
+++ b/Engine/Source/Runtime/Engine/Public/Elements/Component/ComponentElementSelectionInterface.h
@@ -0,0 +1,20 @@
+// Copyright Epic Games, Inc. All Rights Reserved.
+
+#pragma once
+
+#include "TypedElementSelectionInterface.h"
+#include "ComponentElementSelectionInterface.generated.h"
+
+class FTypedElementList;
+
+UCLASS()
+class ENGINE_API UComponentElementSelectionInterface : public UTypedElementSelectionInterface
+{
+	GENERATED_BODY()
+
+public:
+	virtual UObject* Legacy_GetSelectionObject(const FTypedElementHandle& InElementHandle) override;
+
+	static int32 GetNumSelectedComponents(const FTypedElementList& InCurrentSelection);
+	static bool HasSelectedComponents(const FTypedElementList& InCurrentSelection);
+};
diff --git a/Engine/Source/Runtime/TypedElementFramework/Private/TypedElementRegistry.cpp b/Engine/Source/Runtime/TypedElementFramework/Private/TypedElementRegistry.cpp
--- a/Engine/Source/Runtime/TypedElementFramework/Private/TypedElementRegistry.cpp
+++ b/Engine/Source/Runtime/TypedElementFramework/Private/TypedElementRegistry.cpp
@@ -88,14 +88,14 @@ void UTypedElementRegistry::RegisterElementTypeImpl(const FName InElementTypeNam
 	AddRegisteredElementType(MoveTemp(InRegisteredElementType));
 }
 
-void UTypedElementRegistry::RegisterElementInterfaceImpl(const FName InElementTypeName, UTypedElementInterface* InElementInterface, const TSubclassOf<UTypedElementInterface>& InBaseInterfaceType)
+void UTypedElementRegistry::RegisterElementInterfaceImpl(const FName InElementTypeName, UTypedElementInterface* InElementInterface, const TSubclassOf<UTypedElementInterface>& InBaseInterfaceType, const bool InAllowOverride)
 {
 	checkf(InElementInterface->IsA(InBaseInterfaceType), TEXT("Interface '%s' of type '%s' does not derive from '%s'!"), *InElementInterface->GetPathName(), *InElementInterface->GetClass()->GetName(), *InBaseInterfaceType->GetName());
 
 	FRegisteredElementType* RegisteredElementType = GetRegisteredElementTypeFromName(InElementTypeName);
 	checkf(RegisteredElementType, TEXT("Element type '%s' has not been registered!"), *InElementTypeName.ToString());
 
-	checkf(!RegisteredElementType->Interfaces.Contains(InBaseInterfaceType->GetFName()), TEXT("Element type '%s' has already registered an interface for '%s'!"), *InElementTypeName.ToString(), *InBaseInterfaceType->GetName());
+	checkf(InAllowOverride || !RegisteredElementType->Interfaces.Contains(InBaseInterfaceType->GetFName()), TEXT("Element type '%s' has already registered an interface for '%s'!"), *InElementTypeName.ToString(), *InBaseInterfaceType->GetName());
 	RegisteredElementType->Interfaces.Add(InBaseInterfaceType->GetFName(), InElementInterface);
 }
 
diff --git a/Engine/Source/Runtime/TypedElementFramework/Public/TypedElementHandle.h b/Engine/Source/Runtime/TypedElementFramework/Public/TypedElementHandle.h
--- a/Engine/Source/Runtime/TypedElementFramework/Public/TypedElementHandle.h
+++ b/Engine/Source/Runtime/TypedElementFramework/Public/TypedElementHandle.h
@@ -258,17 +258,23 @@ public:
 	 * @note This is not typically something you'd want to use outside of data access within an interface implementation.
 	 */
 	template <typename ElementDataType>
-	const ElementDataType* GetData() const
+	const ElementDataType* GetData(const bool bSilent = false) const
 	{
 		if (!DataPtr)
 		{
-			FFrame::KismetExecutionMessage(TEXT("Element handle data is null!"), ELogVerbosity::Error);
+			if (!bSilent)
+			{
+				FFrame::KismetExecutionMessage(TEXT("Element handle data is null!"), ELogVerbosity::Error);
+			}
 			return nullptr;
 		}
 
 		if (!IsDataOfType<ElementDataType>())
 		{
-			FFrame::KismetExecutionMessage(*FString::Printf(TEXT("Element handle data type is '%d', but '%d' (%s) was requested!"), Id.GetTypeId(), ElementDataType::StaticTypeId(), *ElementDataType::StaticTypeName().ToString()), ELogVerbosity::Error);
+			if (!bSilent)
+			{
+				FFrame::KismetExecutionMessage(*FString::Printf(TEXT("Element handle data type is '%d', but '%d' (%s) was requested!"), Id.GetTypeId(), ElementDataType::StaticTypeId(), *ElementDataType::StaticTypeName().ToString()), ELogVerbosity::Error);
+			}
 			return nullptr;
 		}
 
diff --git a/Engine/Source/Runtime/TypedElementFramework/Public/TypedElementRegistry.h b/Engine/Source/Runtime/TypedElementFramework/Public/TypedElementRegistry.h
--- a/Engine/Source/Runtime/TypedElementFramework/Public/TypedElementRegistry.h
+++ b/Engine/Source/Runtime/TypedElementFramework/Public/TypedElementRegistry.h
@@ -62,9 +62,9 @@ public:
 	 * Register that an element interface is supported for the given type, which must have previously been registered via RegisterElementType.
 	 */
 	template <typename BaseInterfaceType>
-	FORCEINLINE void RegisterElementInterface(const FName InElementTypeName, UTypedElementInterface* InElementInterface)
+	FORCEINLINE void RegisterElementInterface(const FName InElementTypeName, UTypedElementInterface* InElementInterface, const bool InAllowOverride = false)
 	{
-		RegisterElementInterfaceImpl(InElementTypeName, InElementInterface, BaseInterfaceType::StaticClass());
+		RegisterElementInterfaceImpl(InElementTypeName, InElementInterface, BaseInterfaceType::StaticClass(), InAllowOverride);
 	}
 
 	/**
@@ -285,7 +285,7 @@ private:
 	};
 
 	void RegisterElementTypeImpl(const FName InElementTypeName, TUniquePtr<FRegisteredElementType>&& InRegisteredElementType);
-	void RegisterElementInterfaceImpl(const FName InElementTypeName, UTypedElementInterface* InElementInterface, const TSubclassOf<UTypedElementInterface>& InBaseInterfaceType);
+	void RegisterElementInterfaceImpl(const FName InElementTypeName, UTypedElementInterface* InElementInterface, const TSubclassOf<UTypedElementInterface>& InBaseInterfaceType, const bool InAllowOverride);
 	UTypedElementInterface* GetElementInterfaceImpl(const FTypedElementId& InElementId, const TSubclassOf<UTypedElementInterface>& InBaseInterfaceType) const;
 
 	template <typename ElementDataType>
