Fixed typed element list iteration potentially skipping an element

The constructor would always increment the index when trying to find the first valid element, so the index and the current element would get out-of-sync. This resulted in the iteration ending early.

#rnx

diff --git a/Engine/Source/Runtime/TypedElementFramework/Public/TypedElementList.h b/Engine/Source/Runtime/TypedElementFramework/Public/TypedElementList.h
--- a/Engine/Source/Runtime/TypedElementFramework/Public/TypedElementList.h
+++ b/Engine/Source/Runtime/TypedElementFramework/Public/TypedElementList.h
@@ -45,7 +45,16 @@ public:
 		, InitialNum(InElementHandles.Num())
 		, Index(InIndex)
 	{
-		AdvanceIterator();
+		// Walk to the first valid element
+		while (Index < ElementHandles.Num())
+		{
+			TypedElementList_Private::GetElement(Registry, ElementHandles[Index], Element);
+			if (Element)
+			{
+				break;
+			}
+			++Index;
+		}
 	}
 
 	FORCEINLINE const TTypedElement<BaseInterfaceType>& operator*() const
@@ -53,9 +62,17 @@ public:
 		return Element;
 	}
 
-	FORCEINLINE TTypedElementListInterfaceIterator& operator++()
+	TTypedElementListInterfaceIterator& operator++()
 	{
-		AdvanceIterator();
+		// Walk to the next valid element
+		while (++Index < ElementHandles.Num())
+		{
+			TypedElementList_Private::GetElement(Registry, ElementHandles[Index], Element);
+			if (Element)
+			{
+				break;
+			}
+		}
 		return *this;
 	}
 
@@ -66,18 +83,6 @@ private:
 	int32 Index;
 	TTypedElement<BaseInterfaceType> Element;
 
-	void AdvanceIterator()
-	{
-		while (Index < ElementHandles.Num())
-		{
-			TypedElementList_Private::GetElement(Registry, ElementHandles[Index++], Element);
-			if (Element)
-			{
-				break;
-			}
-		}
-	}
-
 	FORCEINLINE friend bool operator!=(const TTypedElementListInterfaceIterator& Lhs, const TTypedElementListInterfaceIterator& Rhs)
 	{
 		// We only need to do the check in this operator, because no other operator will be
