Fixed leak of FAssetEditorModeManager instances

FEditorViewportClient would only unregister and destroy the mode manager if it had created it, and some places were passing in a new instance which was being leaked

#rnx

diff --git a/Engine/Plugins/Compositing/Composure/Source/ComposureLayersEditor/Private/CompositingViewportClient.cpp b/Engine/Plugins/Compositing/Composure/Source/ComposureLayersEditor/Private/CompositingViewportClient.cpp
--- a/Engine/Plugins/Compositing/Composure/Source/ComposureLayersEditor/Private/CompositingViewportClient.cpp
+++ b/Engine/Plugins/Compositing/Composure/Source/ComposureLayersEditor/Private/CompositingViewportClient.cpp
@@ -45,7 +45,7 @@ public:
  *****************************************************************************/
 
 FCompositingViewportClient::FCompositingViewportClient(TWeakObjectPtr<UEditorCompElementContainer> CompElements)
-	: FEditorViewportClient(new FAssetEditorModeManager())
+	: FEditorViewportClient(nullptr)
 	, ElementsContainerPtr(CompElements)
 {
 	SetViewModes(VMI_Unlit, VMI_Unlit);
diff --git a/Engine/Plugins/Experimental/ActorPalette/Source/ActorPalette/Private/ActorPaletteViewportClient.cpp b/Engine/Plugins/Experimental/ActorPalette/Source/ActorPalette/Private/ActorPaletteViewportClient.cpp
--- a/Engine/Plugins/Experimental/ActorPalette/Source/ActorPalette/Private/ActorPaletteViewportClient.cpp
+++ b/Engine/Plugins/Experimental/ActorPalette/Source/ActorPalette/Private/ActorPaletteViewportClient.cpp
@@ -26,7 +26,7 @@
 // FActorPaletteViewportClient
 
 FActorPaletteViewportClient::FActorPaletteViewportClient(int32 InTabIndex)
-	: FEditorViewportClient(new FAssetEditorModeManager(), nullptr, nullptr)
+	: FEditorViewportClient(nullptr, nullptr, nullptr)
 	, TabIndex(InTabIndex)
 {
 	SetViewModes(VMI_Lit, VMI_Lit);
diff --git a/Engine/Source/Editor/UnrealEd/Private/AssetEditorModeManager.cpp b/Engine/Source/Editor/UnrealEd/Private/AssetEditorModeManager.cpp
--- a/Engine/Source/Editor/UnrealEd/Private/AssetEditorModeManager.cpp
+++ b/Engine/Source/Editor/UnrealEd/Private/AssetEditorModeManager.cpp
@@ -28,16 +28,21 @@ FAssetEditorModeManager::~FAssetEditorModeManager()
 {
 	SetPreviewScene(nullptr);
 
-	ActorSet->SetElementList(nullptr);
-	ActorSet->RemoveFromRoot();
-	ActorSet = nullptr;
-
-	ObjectSet->RemoveFromRoot();
-	ObjectSet = nullptr;
-
-	ComponentSet->SetElementList(nullptr);
-	ComponentSet->RemoveFromRoot();
-	ComponentSet = nullptr;
+	// We may be destroyed after the UObject system has already shutdown, 
+	// which would mean that these instances will be garbage
+	if (UObjectInitialized())
+	{
+		ActorSet->SetElementList(nullptr);
+		ActorSet->RemoveFromRoot();
+		ActorSet = nullptr;
+
+		ObjectSet->RemoveFromRoot();
+		ObjectSet = nullptr;
+
+		ComponentSet->SetElementList(nullptr);
+		ComponentSet->RemoveFromRoot();
+		ComponentSet = nullptr;
+	}
 
 	SelectedElements.Reset();
 }
diff --git a/Engine/Source/Editor/UnrealEd/Private/EditorModeManager.cpp b/Engine/Source/Editor/UnrealEd/Private/EditorModeManager.cpp
--- a/Engine/Source/Editor/UnrealEd/Private/EditorModeManager.cpp
+++ b/Engine/Source/Editor/UnrealEd/Private/EditorModeManager.cpp
@@ -91,8 +91,13 @@ FEditorModeTools::~FEditorModeTools()
 	DeactivateAllModes();
 	RecycledScriptableModes.Empty();
 
-	InteractiveToolsContext->ShutdownContext();
-	InteractiveToolsContext = nullptr;
+	// We may be destroyed after the UObject system has already shutdown, 
+	// which would mean that this instances will be garbage
+	if (UObjectInitialized())
+	{
+		InteractiveToolsContext->ShutdownContext();
+		InteractiveToolsContext = nullptr;
+	}
 }
 
 void FEditorModeTools::LoadConfig(void)
diff --git a/Engine/Source/Editor/UnrealEd/Private/Selection.cpp b/Engine/Source/Editor/UnrealEd/Private/Selection.cpp
--- a/Engine/Source/Editor/UnrealEd/Private/Selection.cpp
+++ b/Engine/Source/Editor/UnrealEd/Private/Selection.cpp
@@ -728,6 +728,8 @@ void USelection::BeginDestroy()
 {
 	Super::BeginDestroy();
 
+	SelectionStore.Reset();
+
 	if (bOwnsSelectionAnnotation)
 	{
 		delete SelectionAnnotation;
diff --git a/Engine/Source/Runtime/TypedElementFramework/Private/TypedElementList.cpp b/Engine/Source/Runtime/TypedElementFramework/Private/TypedElementList.cpp
--- a/Engine/Source/Runtime/TypedElementFramework/Private/TypedElementList.cpp
+++ b/Engine/Source/Runtime/TypedElementFramework/Private/TypedElementList.cpp
@@ -75,13 +75,16 @@ void FTypedElementListLegacySync::ForceBatchOperationDirty()
 FTypedElementList::FTypedElementList(UTypedElementRegistry* InRegistry)
 	: Registry(InRegistry)
 {
-	Registry->Private_OnElementListCreated(this);
+	InRegistry->Private_OnElementListCreated(this);
 }
 
 FTypedElementList::~FTypedElementList()
 {
 	LegacySync.Reset();
-	Registry->Private_OnElementListDestroyed(this);
+	if (UTypedElementRegistry* RegistryPtr = Registry.Get())
+	{
+		RegistryPtr->Private_OnElementListDestroyed(this);
+	}
 }
 
 FTypedElementListPtr FTypedElementList::Private_CreateElementList(UTypedElementRegistry* InRegistry)
@@ -91,7 +94,7 @@ FTypedElementListPtr FTypedElementList::Private_CreateElementList(UTypedElementR
 
 FTypedElementListPtr FTypedElementList::Clone() const
 {
-	FTypedElementListPtr ClonedElementList = Private_CreateElementList(Registry);
+	FTypedElementListPtr ClonedElementList = Private_CreateElementList(Registry.Get());
 	ClonedElementList->ElementCombinedIds = ElementCombinedIds;
 	ClonedElementList->ElementHandles = ElementHandles;
 	return ClonedElementList;
diff --git a/Engine/Source/Runtime/TypedElementFramework/Public/TypedElementList.h b/Engine/Source/Runtime/TypedElementFramework/Public/TypedElementList.h
--- a/Engine/Source/Runtime/TypedElementFramework/Public/TypedElementList.h
+++ b/Engine/Source/Runtime/TypedElementFramework/Public/TypedElementList.h
@@ -260,7 +260,7 @@ public:
 	template <typename BaseInterfaceType>
 	FORCEINLINE TTypedElement<BaseInterfaceType> GetElement(const FTypedElementHandle& InElementHandle) const
 	{
-		return TypedElementList_Private::GetElement<BaseInterfaceType>(Registry, InElementHandle);
+		return TypedElementList_Private::GetElement<BaseInterfaceType>(Registry.Get(), InElementHandle);
 	}
 
 	/**
@@ -269,7 +269,7 @@ public:
 	template <typename BaseInterfaceType>
 	FORCEINLINE void GetElement(const FTypedElementHandle& InElementHandle, TTypedElement<BaseInterfaceType>& OutElement) const
 	{
-		TypedElementList_Private::GetElement(Registry, InElementHandle, OutElement);
+		TypedElementList_Private::GetElement(Registry.Get(), InElementHandle, OutElement);
 	}
 
 	/**
@@ -512,7 +512,7 @@ public:
 	template <typename BaseInterfaceType>
 	FORCEINLINE TypedElementList_Private::TTypedElementListInterfaceIteratorProxy<BaseInterfaceType> InterateInterface() const
 	{
-		return TypedElementList_Private::TTypedElementListInterfaceIteratorProxy<BaseInterfaceType>(Registry, ElementHandles);
+		return TypedElementList_Private::TTypedElementListInterfaceIteratorProxy<BaseInterfaceType>(Registry.Get(), ElementHandles);
 	}
 
 	/**
@@ -553,7 +553,10 @@ private:
 
 	void NoteListChanged(const EChangeType InChangeType, const FTypedElementHandle& InElementHandle = FTypedElementHandle());
 
-	UTypedElementRegistry* Registry = nullptr;
+	/**
+	 * Element registry this element list is associated with.
+	 */
+	TWeakObjectPtr<UTypedElementRegistry> Registry;
 
 	/**
 	 * Set of combined ID values that are currently present in this element list.
@@ -568,7 +571,7 @@ private:
 	TArray<FTypedElementHandle> ElementHandles;
 
 	/**
-	 * Delegate that is invoked whenever this element list has been changed
+	 * Delegate that is invoked whenever this element list has been changed.
 	 */
 	FOnChanged OnChangedDelegate;
 
