Improved scripting support for typed elements

 - Ported the typed element list (FTypedElementList) to be a UObject (UTypedElementList) so that it can be used with UFUNCTION arguments.
   - This also removes FTypedElementListPtr, and exposes UTypedElementList functions to scripting.
 - Exposed UTypedElementRegistry functions to scripting.
 - Added UTypedElementHandleLibrary to expose basic FTypedElementHandle functionality to scripting.
 - Exposed AActor and UActorComponent functions for getting an element handle.

#rnx

diff --git a/Engine/Source/Editor/LevelEditor/Private/SLevelEditor.cpp b/Engine/Source/Editor/LevelEditor/Private/SLevelEditor.cpp
--- a/Engine/Source/Editor/LevelEditor/Private/SLevelEditor.cpp
+++ b/Engine/Source/Editor/LevelEditor/Private/SLevelEditor.cpp
@@ -212,10 +212,11 @@ void SLevelEditor::Construct( const SLevelEditor::FArguments& InArgs)
 void SLevelEditor::Initialize( const TSharedRef<SDockTab>& OwnerTab, const TSharedRef<SWindow>& OwnerWindow )
 {
 	SelectedElements = UTypedElementRegistry::GetInstance()->CreateElementList();
+	SelectedElements->AddToRoot();
 
 	// Allow USelection to bridge to our selected element list
-	GUnrealEd->GetSelectedActors()->SetElementList(SelectedElements.Get());
-	GUnrealEd->GetSelectedComponents()->SetElementList(SelectedElements.Get());
+	GUnrealEd->GetSelectedActors()->SetElementList(SelectedElements);
+	GUnrealEd->GetSelectedComponents()->SetElementList(SelectedElements);
 
 	// Bind the level editor tab's label to the currently loaded level name string in the main frame
 	OwnerTab->SetLabel( TAttribute<FText>( this, &SLevelEditor::GetTabTitle) );
@@ -326,7 +327,12 @@ SLevelEditor::~SLevelEditor()
 		GUnrealEd->GetSelectedActors()->SetElementList(nullptr);
 		GUnrealEd->GetSelectedComponents()->SetElementList(nullptr);
 	}
-	SelectedElements.Reset();
+	if (UObjectInitialized())
+	{
+		SelectedElements->Empty();
+		SelectedElements->RemoveFromRoot();
+		SelectedElements = nullptr;
+	}
 }
 
 FText SLevelEditor::GetTabTitle() const
diff --git a/Engine/Source/Editor/LevelEditor/Private/SLevelEditor.h b/Engine/Source/Editor/LevelEditor/Private/SLevelEditor.h
--- a/Engine/Source/Editor/LevelEditor/Private/SLevelEditor.h
+++ b/Engine/Source/Editor/LevelEditor/Private/SLevelEditor.h
@@ -15,7 +15,6 @@
 #include "AssetThumbnail.h"
 #include "ILevelEditor.h"
 #include "LevelViewportTabContent.h"
-#include "TypedElementListPtr.h"
 
 class IAssetEditorInstance;
 class IDetailsView;
@@ -23,6 +22,7 @@ class SActorDetails;
 class SBorder;
 class SLevelEditorModeContent;
 class SLevelEditorToolBox;
+class UTypedElementList;
 
 /**
  * Unreal editor level editor Slate widget
@@ -259,7 +259,7 @@ private:
 	UWorld* World;
 
 	// The list of selected elements (also set on the global USelection for actors and components).
-	FTypedElementListPtr SelectedElements;
+	UTypedElementList* SelectedElements = nullptr;
 
 	// The box that holds the title bar messages.
 	TSharedPtr<SHorizontalBox> TtileBarMessageBox;
diff --git a/Engine/Source/Editor/UnrealEd/Private/AssetEditorModeManager.cpp b/Engine/Source/Editor/UnrealEd/Private/AssetEditorModeManager.cpp
--- a/Engine/Source/Editor/UnrealEd/Private/AssetEditorModeManager.cpp
+++ b/Engine/Source/Editor/UnrealEd/Private/AssetEditorModeManager.cpp
@@ -9,18 +9,19 @@
 // FAssetEditorModeManager
 
 FAssetEditorModeManager::FAssetEditorModeManager()
-	: PreviewScene(nullptr)
-	, SelectedElements(UTypedElementRegistry::GetInstance()->CreateElementList())
 {
+	SelectedElements = UTypedElementRegistry::GetInstance()->CreateElementList();
+	SelectedElements->AddToRoot();
+
 	ActorSet = USelection::CreateActorSelection(nullptr, GetTransientPackage(), NAME_None, RF_Transactional);
-	ActorSet->SetElementList(SelectedElements.Get());
+	ActorSet->SetElementList(SelectedElements);
 	ActorSet->AddToRoot();
 
 	ObjectSet = USelection::CreateObjectSelection(nullptr, GetTransientPackage(), NAME_None, RF_Transactional);
 	ObjectSet->AddToRoot();
 
 	ComponentSet = USelection::CreateComponentSelection(nullptr, GetTransientPackage(), NAME_None, RF_Transactional);
-	ComponentSet->SetElementList(SelectedElements.Get());
+	ComponentSet->SetElementList(SelectedElements);
 	ComponentSet->AddToRoot();
 }
 
@@ -42,9 +43,11 @@ FAssetEditorModeManager::~FAssetEditorModeManager()
 		ComponentSet->SetElementList(nullptr);
 		ComponentSet->RemoveFromRoot();
 		ComponentSet = nullptr;
-	}
 
-	SelectedElements.Reset();
+		SelectedElements->Empty();
+		SelectedElements->RemoveFromRoot();
+		SelectedElements = nullptr;
+	}
 }
 
 USelection* FAssetEditorModeManager::GetSelectedActors() const
diff --git a/Engine/Source/Editor/UnrealEd/Private/Selection.cpp b/Engine/Source/Editor/UnrealEd/Private/Selection.cpp
--- a/Engine/Source/Editor/UnrealEd/Private/Selection.cpp
+++ b/Engine/Source/Editor/UnrealEd/Private/Selection.cpp
@@ -56,7 +56,7 @@ public:
 	explicit FObjectSelectionStore(const UClass* InObjectBaseClass);
 
 	//~ ISelectionStore interface
-	virtual void SetElementList(FTypedElementList* InElementList) override;
+	virtual void SetElementList(UTypedElementList* InElementList) override;
 	virtual int32 GetNumObjects() const override;
 	virtual UObject* GetObjectAtIndex(const int32 InIndex) const override;
 	virtual bool IsObjectSelected(const UObject* InObject) const override;
@@ -84,7 +84,7 @@ FObjectSelectionStore::FObjectSelectionStore(const UClass* InObjectBaseClass)
 {
 }
 
-void FObjectSelectionStore::SetElementList(FTypedElementList* InElementList)
+void FObjectSelectionStore::SetElementList(UTypedElementList* InElementList)
 {
 	checkf(!UE_USE_ELEMENT_LIST_SELECTION, TEXT("This selection store does not support element lists! Only actor and component selections may currently use element lists!"));
 }
@@ -201,7 +201,7 @@ public:
 	virtual ~FElementSelectionStore();
 
 	//~ ISelectionStore interface
-	virtual void SetElementList(FTypedElementList* InElementList) override;
+	virtual void SetElementList(UTypedElementList* InElementList) override;
 	virtual int32 GetNumObjects() const override;
 	virtual UObject* GetObjectAtIndex(const int32 InIndex) const override;
 	virtual bool IsValidObjectToSelect(const UObject* InObject) const override;
@@ -217,11 +217,11 @@ public:
 	virtual FTypedElementHandle GetElementHandleForObject(const UObject* InObject, const bool bAllowCreate = true) const = 0;
 
 private:
-	void OnElementListSyncEvent(const FTypedElementList& InElementList, FTypedElementListLegacySync::ESyncType InSyncType, const FTypedElementHandle& InElementHandle, bool bIsWithinBatchOperation);
+	void OnElementListSyncEvent(const UTypedElementList* InElementList, FTypedElementListLegacySync::ESyncType InSyncType, const FTypedElementHandle& InElementHandle, bool bIsWithinBatchOperation);
 	UObject* GetObjectForElementHandle(const FTypedElementHandle& InElementHandle) const;
 
 	/** External list of elements that this selection instance is bridging. This will be null unless a level editor instance exists. */
-	FTypedElementList* ElementList = nullptr;
+	UTypedElementList* ElementList = nullptr;
 };
 
 FElementSelectionStore::FElementSelectionStore(const UClass* InObjectBaseClass)
@@ -237,7 +237,7 @@ FElementSelectionStore::~FElementSelectionStore()
 	}
 }
 
-void FElementSelectionStore::SetElementList(FTypedElementList* InElementList)
+void FElementSelectionStore::SetElementList(UTypedElementList* InElementList)
 {
 	if (ElementList)
 	{
@@ -368,7 +368,7 @@ void FElementSelectionStore::ForceBatchDirty()
 	}
 }
 
-void FElementSelectionStore::OnElementListSyncEvent(const FTypedElementList& InElementList, FTypedElementListLegacySync::ESyncType InSyncType, const FTypedElementHandle& InElementHandle, bool bIsWithinBatchOperation)
+void FElementSelectionStore::OnElementListSyncEvent(const UTypedElementList* InElementList, FTypedElementListLegacySync::ESyncType InSyncType, const FTypedElementHandle& InElementHandle, bool bIsWithinBatchOperation)
 {
 	if (Sink)
 	{
diff --git a/Engine/Source/Editor/UnrealEd/Public/AssetEditorModeManager.h b/Engine/Source/Editor/UnrealEd/Public/AssetEditorModeManager.h
--- a/Engine/Source/Editor/UnrealEd/Public/AssetEditorModeManager.h
+++ b/Engine/Source/Editor/UnrealEd/Public/AssetEditorModeManager.h
@@ -4,9 +4,9 @@
 
 #include "CoreMinimal.h"
 #include "EditorModeManager.h"
-#include "TypedElementListPtr.h"
 
 class FPreviewScene;
+class UTypedElementList;
 
 //////////////////////////////////////////////////////////////////////////
 // FAssetEditorModeManager
@@ -28,9 +28,9 @@ public:
 	FPreviewScene* GetPreviewScene() const;
 
 protected:
-	class USelection* ActorSet;
-	class USelection* ObjectSet;
-	class USelection* ComponentSet;
-	class FPreviewScene* PreviewScene;
-	FTypedElementListPtr SelectedElements;
+	class USelection* ActorSet = nullptr;
+	class USelection* ObjectSet = nullptr;
+	class USelection* ComponentSet = nullptr;
+	class FPreviewScene* PreviewScene = nullptr;
+	UTypedElementList* SelectedElements = nullptr;
 };
diff --git a/Engine/Source/Editor/UnrealEd/Public/Selection.h b/Engine/Source/Editor/UnrealEd/Public/Selection.h
--- a/Engine/Source/Editor/UnrealEd/Public/Selection.h
+++ b/Engine/Source/Editor/UnrealEd/Public/Selection.h
@@ -9,7 +9,7 @@
 #include "Components/ActorComponent.h"
 #include "Selection.generated.h"
 
-class FTypedElementList;
+class UTypedElementList;
 
 namespace Selection_Private
 {
@@ -49,7 +49,7 @@ public:
 	 * Set the element list instance for this store.
 	 * @note Does nothing if element list stores aren't enabled. Asserts for non-element list stores if they are!
 	 */
-	virtual void SetElementList(FTypedElementList* InElementList) = 0;
+	virtual void SetElementList(UTypedElementList* InElementList) = 0;
 
 	/**
 	 * Get the number of objects within the underlying store.
@@ -181,7 +181,7 @@ public:
 	 * Set the element list instance for this selection set.
 	 * @note Does nothing if element list stores aren't enabled. Asserts for non-element list stores if they are!
 	 */
-	void SetElementList(FTypedElementList* InElementList)
+	void SetElementList(UTypedElementList* InElementList)
 	{
 		SelectionStore->SetElementList(InElementList);
 	}
diff --git a/Engine/Source/Runtime/Engine/Classes/Components/ActorComponent.h b/Engine/Source/Runtime/Engine/Classes/Components/ActorComponent.h
--- a/Engine/Source/Runtime/Engine/Classes/Components/ActorComponent.h
+++ b/Engine/Source/Runtime/Engine/Classes/Components/ActorComponent.h
@@ -540,6 +540,7 @@ public:
 	TTypedElementOwner<struct FComponentElementData> CreateComponentElement() const;
 	void DestroyComponentElement(TTypedElementOwner<struct FComponentElementData>& InOutComponentElement) const;
 #if WITH_EDITOR
+	UFUNCTION(BlueprintCallable, BlueprintPure=false, Category="TypedElementFramework|Component")
 	FTypedElementHandle AcquireEditorElementHandle(const bool bAllowCreate = true) const;
 #endif
 
diff --git a/Engine/Source/Runtime/Engine/Classes/GameFramework/Actor.h b/Engine/Source/Runtime/Engine/Classes/GameFramework/Actor.h
--- a/Engine/Source/Runtime/Engine/Classes/GameFramework/Actor.h
+++ b/Engine/Source/Runtime/Engine/Classes/GameFramework/Actor.h
@@ -3312,6 +3312,7 @@ public:
 	TTypedElementOwner<struct FActorElementData> CreateActorElement() const;
 	void DestroyActorElement(TTypedElementOwner<struct FActorElementData>& InOutActorElement) const;
 #if WITH_EDITOR
+	UFUNCTION(BlueprintCallable, BlueprintPure=false, Category="TypedElementFramework|Actor")
 	FTypedElementHandle AcquireEditorElementHandle(const bool bAllowCreate = true) const;
 #endif
 
diff --git a/Engine/Source/Runtime/Engine/Private/Elements/Actor/ActorElementSelectionInterface.cpp b/Engine/Source/Runtime/Engine/Private/Elements/Actor/ActorElementSelectionInterface.cpp
--- a/Engine/Source/Runtime/Engine/Private/Elements/Actor/ActorElementSelectionInterface.cpp
+++ b/Engine/Source/Runtime/Engine/Private/Elements/Actor/ActorElementSelectionInterface.cpp
@@ -11,10 +11,10 @@ UObject* UActorElementSelectionInterface::Legacy_GetSelectionObject(const FTyped
 	return ActorData ? ActorData->Actor : nullptr;
 }
 
-int32 UActorElementSelectionInterface::GetNumSelectedActors(const FTypedElementList& InCurrentSelection)
+int32 UActorElementSelectionInterface::GetNumSelectedActors(const UTypedElementList* InCurrentSelection)
 {
 	int32 NumSelected = 0;
-	for (const FTypedElementHandle& SelectedElement : InCurrentSelection)
+	for (const FTypedElementHandle& SelectedElement : *InCurrentSelection)
 	{
 		if (SelectedElement.GetData<FActorElementData>(/*bSilent*/true))
 		{
@@ -24,9 +24,9 @@ int32 UActorElementSelectionInterface::GetNumSelectedActors(const FTypedElementL
 	return NumSelected;
 }
 
-bool UActorElementSelectionInterface::HasSelectedActors(const FTypedElementList& InCurrentSelection)
+bool UActorElementSelectionInterface::HasSelectedActors(const UTypedElementList* InCurrentSelection)
 {
-	for (const FTypedElementHandle& SelectedElement : InCurrentSelection)
+	for (const FTypedElementHandle& SelectedElement : *InCurrentSelection)
 	{
 		if (SelectedElement.GetData<FActorElementData>(/*bSilent*/true))
 		{
diff --git a/Engine/Source/Runtime/Engine/Private/Elements/Component/ComponentElementSelectionInterface.cpp b/Engine/Source/Runtime/Engine/Private/Elements/Component/ComponentElementSelectionInterface.cpp
--- a/Engine/Source/Runtime/Engine/Private/Elements/Component/ComponentElementSelectionInterface.cpp
+++ b/Engine/Source/Runtime/Engine/Private/Elements/Component/ComponentElementSelectionInterface.cpp
@@ -11,10 +11,10 @@ UObject* UComponentElementSelectionInterface::Legacy_GetSelectionObject(const FT
 	return ComponentData ? ComponentData->Component : nullptr;
 }
 
-int32 UComponentElementSelectionInterface::GetNumSelectedComponents(const FTypedElementList& InCurrentSelection)
+int32 UComponentElementSelectionInterface::GetNumSelectedComponents(const UTypedElementList* InCurrentSelection)
 {
 	int32 NumSelected = 0;
-	for (const FTypedElementHandle& SelectedElement : InCurrentSelection)
+	for (const FTypedElementHandle& SelectedElement : *InCurrentSelection)
 	{
 		if (SelectedElement.GetData<FComponentElementData>(/*bSilent*/true))
 		{
@@ -24,9 +24,9 @@ int32 UComponentElementSelectionInterface::GetNumSelectedComponents(const FTyped
 	return NumSelected;
 }
 
-bool UComponentElementSelectionInterface::HasSelectedComponents(const FTypedElementList& InCurrentSelection)
+bool UComponentElementSelectionInterface::HasSelectedComponents(const UTypedElementList* InCurrentSelection)
 {
-	for (const FTypedElementHandle& SelectedElement : InCurrentSelection)
+	for (const FTypedElementHandle& SelectedElement : *InCurrentSelection)
 	{
 		if (SelectedElement.GetData<FComponentElementData>(/*bSilent*/true))
 		{
diff --git a/Engine/Source/Runtime/Engine/Private/UnrealClient.cpp b/Engine/Source/Runtime/Engine/Private/UnrealClient.cpp
--- a/Engine/Source/Runtime/Engine/Private/UnrealClient.cpp
+++ b/Engine/Source/Runtime/Engine/Private/UnrealClient.cpp
@@ -1842,9 +1842,9 @@ FTypedElementHandle FViewport::GetElementHandleAtPoint(int32 X, int32 Y)
 	return FTypedElementHandle();
 }
 
-void FViewport::GetElementHandlesInRect(FIntRect InRect, FTypedElementList& OutElementHandles)
+void FViewport::GetElementHandlesInRect(FIntRect InRect, UTypedElementList* OutElementHandles)
 {
-	GetElementHandlesInRectImpl(InRect, OutElementHandles);
+	GetElementHandlesInRectImpl(InRect, *OutElementHandles);
 }
 
 void FViewport::GetElementHandlesInRect(FIntRect InRect, TSet<FTypedElementHandle>& OutElementHandles)
diff --git a/Engine/Source/Runtime/Engine/Public/Elements/Actor/ActorElementSelectionInterface.h b/Engine/Source/Runtime/Engine/Public/Elements/Actor/ActorElementSelectionInterface.h
--- a/Engine/Source/Runtime/Engine/Public/Elements/Actor/ActorElementSelectionInterface.h
+++ b/Engine/Source/Runtime/Engine/Public/Elements/Actor/ActorElementSelectionInterface.h
@@ -5,7 +5,7 @@
 #include "TypedElementSelectionInterface.h"
 #include "ActorElementSelectionInterface.generated.h"
 
-class FTypedElementList;
+class UTypedElementList;
 
 UCLASS()
 class ENGINE_API UActorElementSelectionInterface : public UTypedElementSelectionInterface
@@ -15,6 +15,6 @@ class ENGINE_API UActorElementSelectionInterface : public UTypedElementSelection
 public:
 	virtual UObject* Legacy_GetSelectionObject(const FTypedElementHandle& InElementHandle) override;
 
-	static int32 GetNumSelectedActors(const FTypedElementList& InCurrentSelection);
-	static bool HasSelectedActors(const FTypedElementList& InCurrentSelection);
+	static int32 GetNumSelectedActors(const UTypedElementList* InCurrentSelection);
+	static bool HasSelectedActors(const UTypedElementList* InCurrentSelection);
 };
diff --git a/Engine/Source/Runtime/Engine/Public/Elements/Component/ComponentElementSelectionInterface.h b/Engine/Source/Runtime/Engine/Public/Elements/Component/ComponentElementSelectionInterface.h
--- a/Engine/Source/Runtime/Engine/Public/Elements/Component/ComponentElementSelectionInterface.h
+++ b/Engine/Source/Runtime/Engine/Public/Elements/Component/ComponentElementSelectionInterface.h
@@ -5,7 +5,7 @@
 #include "TypedElementSelectionInterface.h"
 #include "ComponentElementSelectionInterface.generated.h"
 
-class FTypedElementList;
+class UTypedElementList;
 
 UCLASS()
 class ENGINE_API UComponentElementSelectionInterface : public UTypedElementSelectionInterface
@@ -15,6 +15,6 @@ class ENGINE_API UComponentElementSelectionInterface : public UTypedElementSelec
 public:
 	virtual UObject* Legacy_GetSelectionObject(const FTypedElementHandle& InElementHandle) override;
 
-	static int32 GetNumSelectedComponents(const FTypedElementList& InCurrentSelection);
-	static bool HasSelectedComponents(const FTypedElementList& InCurrentSelection);
+	static int32 GetNumSelectedComponents(const UTypedElementList* InCurrentSelection);
+	static bool HasSelectedComponents(const UTypedElementList* InCurrentSelection);
 };
diff --git a/Engine/Source/Runtime/Engine/Public/UnrealClient.h b/Engine/Source/Runtime/Engine/Public/UnrealClient.h
--- a/Engine/Source/Runtime/Engine/Public/UnrealClient.h
+++ b/Engine/Source/Runtime/Engine/Public/UnrealClient.h
@@ -20,7 +20,7 @@ class FCanvas;
 class FViewport;
 class FViewportClient;
 class UModel;
-class FTypedElementList;
+class UTypedElementList;
 
 /**
  * A render target.
@@ -502,7 +502,7 @@ public:
 	 * InRect must be entirely within the viewport's client area.
 	 * If the hit proxies are not cached, this will call ViewportClient->Draw with a hit-testing canvas.
 	 */
-	ENGINE_API void GetElementHandlesInRect(FIntRect InRect, FTypedElementList& OutElementHandles);
+	ENGINE_API void GetElementHandlesInRect(FIntRect InRect, UTypedElementList* OutElementHandles);
 	ENGINE_API void GetElementHandlesInRect(FIntRect InRect, TSet<FTypedElementHandle>& OutElementHandles);
 
 	/**
diff --git a/Engine/Source/Runtime/TypedElementFramework/Private/TypedElementFrameworkTests.cpp b/Engine/Source/Runtime/TypedElementFramework/Private/TypedElementFrameworkTests.cpp
--- a/Engine/Source/Runtime/TypedElementFramework/Private/TypedElementFrameworkTests.cpp
+++ b/Engine/Source/Runtime/TypedElementFramework/Private/TypedElementFrameworkTests.cpp
@@ -109,7 +109,7 @@ bool FTypedElementRegistrySmokeTest::RunTest(const FString& Parameters)
 	TestInterfaceAccess(TypedElement1.AcquireHandle());
 	TestInterfaceAccess(UntypedElement1.AcquireHandle());
 
-	FTypedElementListPtr ElementList = Registry->CreateElementList();
+	UTypedElementList* ElementList = Registry->CreateElementList();
 	ElementList->Add(TypedElement1);
 	ElementList->Add(TypedElement2);
 	ElementList->Add(TypedElement3);
@@ -127,7 +127,8 @@ bool FTypedElementRegistrySmokeTest::RunTest(const FString& Parameters)
 		TestInterfaceAHandle(ElementHandle);
 	}
 
-	ElementList.Reset();
+	ElementList->Empty();
+	ElementList = nullptr;
 
 	Registry->DestroyElement(TypedElement1);
 	Registry->DestroyElement(TypedElement2);
@@ -157,7 +158,7 @@ bool FTypedElementRegistryPerfTest::RunTest(const FString& Parameters)
 
 	TArray<TTypedElementOwner<FTestTypedElementData>> TypedOwnerHandles;
 	TArray<FTypedElementOwner> UntypedOwnerHandles;
-	FTypedElementListPtr ElementList = Registry->CreateElementList();
+	UTypedElementList* ElementList = Registry->CreateElementList();
 
 	// Create typed handles
 	{
@@ -220,7 +221,8 @@ bool FTypedElementRegistryPerfTest::RunTest(const FString& Parameters)
 	// Clear the element list
 	{
 		FScopedDurationTimeLogger Timer(FString::Printf(TEXT("Reset %d elements in list"), ElementList->Num()));
-		ElementList.Reset();
+		ElementList->Empty();
+		ElementList = nullptr;
 	}
 
 	// Destroy typed handles
diff --git a/Engine/Source/Runtime/TypedElementFramework/Private/TypedElementList.cpp b/Engine/Source/Runtime/TypedElementFramework/Private/TypedElementList.cpp
--- a/Engine/Source/Runtime/TypedElementFramework/Private/TypedElementList.cpp
+++ b/Engine/Source/Runtime/TypedElementFramework/Private/TypedElementList.cpp
@@ -14,9 +14,10 @@ void GetElementImpl(const UTypedElementRegistry* InRegistry, const FTypedElement
 } // namespace TypedElementList_Private
 
 
-FTypedElementListLegacySync::FTypedElementListLegacySync(const FTypedElementList& InElementList)
+FTypedElementListLegacySync::FTypedElementListLegacySync(const UTypedElementList* InElementList)
 	: ElementList(InElementList)
 {
+	checkf(ElementList, TEXT("ElementList was null!"));
 }
 
 FTypedElementListLegacySync::FOnSyncEvent& FTypedElementListLegacySync::OnSyncEvent()
@@ -72,35 +73,47 @@ void FTypedElementListLegacySync::ForceBatchOperationDirty()
 }
 
 
-FTypedElementList::FTypedElementList(UTypedElementRegistry* InRegistry)
-	: Registry(InRegistry)
+UTypedElementList* UTypedElementList::Private_CreateElementList(UTypedElementRegistry* InRegistry)
 {
+	UTypedElementList* ElementList = NewObject<UTypedElementList>();
+	ElementList->Initialize(InRegistry);
+	return ElementList;
+}
+
+void UTypedElementList::Initialize(UTypedElementRegistry* InRegistry)
+{
+	checkf(!Registry.Get(), TEXT("Initialize has already been called!"));
+	Registry = InRegistry;
+	checkf(InRegistry, TEXT("Registry is null!"));
 	InRegistry->Private_OnElementListCreated(this);
 }
 
-FTypedElementList::~FTypedElementList()
+void UTypedElementList::BeginDestroy()
 {
+	Super::BeginDestroy();
+
 	LegacySync.Reset();
 	if (UTypedElementRegistry* RegistryPtr = Registry.Get())
 	{
 		RegistryPtr->Private_OnElementListDestroyed(this);
+		Registry = nullptr;
 	}
 }
 
-FTypedElementListPtr FTypedElementList::Private_CreateElementList(UTypedElementRegistry* InRegistry)
-{
-	return FTypedElementListPtr(new FTypedElementList(InRegistry));
-}
-
-FTypedElementListPtr FTypedElementList::Clone() const
+UTypedElementList* UTypedElementList::Clone() const
 {
-	FTypedElementListPtr ClonedElementList = Private_CreateElementList(Registry.Get());
+	UTypedElementList* ClonedElementList = Private_CreateElementList(Registry.Get());
 	ClonedElementList->ElementCombinedIds = ElementCombinedIds;
 	ClonedElementList->ElementHandles = ElementHandles;
 	return ClonedElementList;
 }
 
-bool FTypedElementList::AddElementImpl(FTypedElementHandle&& InElementHandle)
+UTypedElementInterface* UTypedElementList::GetElementInterface(const FTypedElementHandle& InElementHandle, const TSubclassOf<UTypedElementInterface>& InBaseInterfaceType) const
+{
+	return Registry->GetElementInterface(InElementHandle, InBaseInterfaceType);
+}
+
+bool UTypedElementList::AddElementImpl(FTypedElementHandle&& InElementHandle)
 {
 	if (!InElementHandle)
 	{
@@ -119,7 +132,7 @@ bool FTypedElementList::AddElementImpl(FTypedElementHandle&& InElementHandle)
 	return !bAlreadyAdded;
 }
 
-bool FTypedElementList::RemoveElementImpl(const FTypedElementId& InElementId)
+bool UTypedElementList::RemoveElementImpl(const FTypedElementId& InElementId)
 {
 	if (!InElementId)
 	{
@@ -145,7 +158,7 @@ bool FTypedElementList::RemoveElementImpl(const FTypedElementId& InElementId)
 	return bRemoved;
 }
 
-int32 FTypedElementList::RemoveAllElementsImpl(TFunctionRef<bool(const FTypedElementHandle&)> InPredicate)
+int32 UTypedElementList::RemoveAllElementsImpl(TFunctionRef<bool(const FTypedElementHandle&)> InPredicate)
 {
 	if (LegacySync)
 	{
@@ -175,32 +188,32 @@ int32 FTypedElementList::RemoveAllElementsImpl(TFunctionRef<bool(const FTypedEle
 	return RemovedCount;
 }
 
-bool FTypedElementList::ContainsElementImpl(const FTypedElementId& InElementId) const
+bool UTypedElementList::ContainsElementImpl(const FTypedElementId& InElementId) const
 {
 	return InElementId 
 		&& ElementCombinedIds.Contains(InElementId.GetCombinedId());
 }
 
-FTypedElementListLegacySync& FTypedElementList::Legacy_GetSync()
+FTypedElementListLegacySync& UTypedElementList::Legacy_GetSync()
 {
 	if (!LegacySync)
 	{
-		LegacySync = MakeUnique<FTypedElementListLegacySync>(*this);
+		LegacySync = MakeUnique<FTypedElementListLegacySync>(this);
 	}
 	return *LegacySync;
 }
 
-void FTypedElementList::NotifyPendingChanges()
+void UTypedElementList::NotifyPendingChanges()
 {
 	if (bHasPendingNotify)
 	{
 		bHasPendingNotify = false;
-		OnChangedDelegate.Broadcast(*this);
+		OnChangedDelegate.Broadcast(this);
 		check(!bHasPendingNotify); // This should still be false after emitting the notification!
 	}
 }
 
-void FTypedElementList::NoteListChanged(const EChangeType InChangeType, const FTypedElementHandle& InElementHandle)
+void UTypedElementList::NoteListChanged(const EChangeType InChangeType, const FTypedElementHandle& InElementHandle)
 {
 	bHasPendingNotify = true;
 
diff --git a/Engine/Source/Runtime/TypedElementFramework/Private/TypedElementRegistry.cpp b/Engine/Source/Runtime/TypedElementFramework/Private/TypedElementRegistry.cpp
--- a/Engine/Source/Runtime/TypedElementFramework/Private/TypedElementRegistry.cpp
+++ b/Engine/Source/Runtime/TypedElementFramework/Private/TypedElementRegistry.cpp
@@ -144,9 +144,9 @@ FTypedElementHandle UTypedElementRegistry::GetElementHandle(const FTypedElementI
 	return ElementHandle;
 }
 
-FTypedElementListPtr UTypedElementRegistry::CreateElementList(TArrayView<const FTypedElementId> InElementIds)
+UTypedElementList* UTypedElementRegistry::CreateElementList(TArrayView<const FTypedElementId> InElementIds)
 {
-	FTypedElementListPtr ElementList = CreateElementList();
+	UTypedElementList* ElementList = CreateElementList();
 
 	for (const FTypedElementId& ElementId : InElementIds)
 	{
@@ -159,9 +159,9 @@ FTypedElementListPtr UTypedElementRegistry::CreateElementList(TArrayView<const F
 	return ElementList;
 }
 
-FTypedElementListPtr UTypedElementRegistry::CreateElementList(TArrayView<const FTypedElementHandle> InElementHandles)
+UTypedElementList* UTypedElementRegistry::CreateElementList(TArrayView<const FTypedElementHandle> InElementHandles)
 {
-	FTypedElementListPtr ElementList = CreateElementList();
+	UTypedElementList* ElementList = CreateElementList();
 	ElementList->Append(InElementHandles);
 	return ElementList;
 }
@@ -170,7 +170,7 @@ void UTypedElementRegistry::NotifyElementListPendingChanges()
 {
 	FReadScopeLock ActiveElementListsLock(ActiveElementListsRW);
 
-	for (FTypedElementList* ActiveElementList : ActiveElementLists)
+	for (UTypedElementList* ActiveElementList : ActiveElementLists)
 	{
 		ActiveElementList->NotifyPendingChanges();
 	}
diff --git a/Engine/Source/Runtime/TypedElementFramework/Public/TypedElementHandle.h b/Engine/Source/Runtime/TypedElementFramework/Public/TypedElementHandle.h
--- a/Engine/Source/Runtime/TypedElementFramework/Public/TypedElementHandle.h
+++ b/Engine/Source/Runtime/TypedElementFramework/Public/TypedElementHandle.h
@@ -88,7 +88,7 @@ public:
 	/**
 	 * Access the combined value of this element ID.
 	 * @note You typically don't want to store this directly as the element ID could be re-used.
-	 *       It is primarily useful as a secondary cache where something is keeping a reference to an element ID or element handle (eg, how FTypedElementList uses it internally).
+	 *       It is primarily useful as a secondary cache where something is keeping a reference to an element ID or element handle (eg, how UTypedElementList uses it internally).
 	 */
 	FORCEINLINE FTypedHandleCombinedId GetCombinedId() const
 	{
@@ -867,3 +867,47 @@ private:
 #endif	// UE_TYPED_ELEMENT_HAS_REFTRACKING
 };
 using FTypedElementOwner = TTypedElementOwner<void>;
+
+/** Script exposure for FTypedElementHandle. */
+UCLASS()
+class UTypedElementHandleLibrary : public UObject
+{
+	GENERATED_BODY()
+
+public:
+	/**
+	 * Has this handle been initialized to a valid element?
+	 */
+	UFUNCTION(BlueprintPure, Category="TypedElementFramework|Handle", meta=(ScriptMethod, ScriptOperator="bool"))
+	static bool IsSet(const FTypedElementHandle& ElementHandle)
+	{
+		return ElementHandle.IsSet();
+	}
+
+	/**
+	 * Release this handle and set it back to an empty state.
+	 */
+	UFUNCTION(BlueprintCallable, Category="TypedElementFramework|Handle", meta=(ScriptMethod))
+	static void Release(UPARAM(ref) FTypedElementHandle& ElementHandle)
+	{
+		ElementHandle.Release();
+	}
+
+	/**
+	 * Are these two handles equal?
+	 */
+	UFUNCTION(BlueprintPure, Category="TypedElementFramework|Handle", meta=(DisplayName="Equal (TypedElementHandle)", CompactNodeTitle="==", Keywords="== equal", ScriptMethod, ScriptOperator="=="))
+	static bool Equal(const FTypedElementHandle& LHS, const FTypedElementHandle& RHS)
+	{
+		return LHS == RHS;
+	}
+
+	/**
+	 * Are these two handles not equal?
+	 */
+	UFUNCTION(BlueprintPure, Category="TypedElementFramework|Handle", meta=(DisplayName="NotEqual (TypedElementHandle)", CompactNodeTitle="!=", Keywords="!= not equal", ScriptMethod, ScriptOperator="!="))
+	static bool NotEqual(const FTypedElementHandle& LHS, const FTypedElementHandle& RHS)
+	{
+		return LHS != RHS;
+	}
+};
diff --git a/Engine/Source/Runtime/TypedElementFramework/Public/TypedElementList.h b/Engine/Source/Runtime/TypedElementFramework/Public/TypedElementList.h
--- a/Engine/Source/Runtime/TypedElementFramework/Public/TypedElementList.h
+++ b/Engine/Source/Runtime/TypedElementFramework/Public/TypedElementList.h
@@ -5,10 +5,13 @@
 #include "CoreMinimal.h"
 #include "Delegates/Delegate.h"
 #include "Delegates/DelegateCombinations.h"
+#include "Templates/SubclassOf.h"
 #include "Containers/ArrayView.h"
+#include "Templates/UniquePtr.h"
 #include "TypedElementHandle.h"
-#include "TypedElementListPtr.h"
+#include "TypedElementList.generated.h"
 
+class UTypedElementList;
 class UTypedElementRegistry;
 
 namespace TypedElementList_Private
@@ -171,11 +174,11 @@ public:
 		BatchComplete,
 	};
 	
-	FTypedElementListLegacySync(const FTypedElementList& InElementList);
+	FTypedElementListLegacySync(const UTypedElementList* InElementList);
 
 	void Private_EmitSyncEvent(const ESyncType InSyncType, const FTypedElementHandle& InElementHandle = FTypedElementHandle());
 
-	DECLARE_EVENT_FourParams(FTypedElementListLegacySync, FOnSyncEvent, const FTypedElementList& /*InElementList*/, ESyncType /*InSyncType*/, const FTypedElementHandle& /*InElementHandle*/, bool /*bIsWithinBatchOperation*/);
+	DECLARE_EVENT_FourParams(FTypedElementListLegacySync, FOnSyncEvent, const UTypedElementList* /*InElementList*/, ESyncType /*InSyncType*/, const FTypedElementHandle& /*InElementHandle*/, bool /*bIsWithinBatchOperation*/);
 	FOnSyncEvent& OnSyncEvent();
 
 	bool IsRunningBatchOperation() const;
@@ -185,7 +188,7 @@ public:
 	void ForceBatchOperationDirty();
 
 private:
-	const FTypedElementList& ElementList;
+	const UTypedElementList* ElementList;
 
 	FOnSyncEvent OnSyncEventDelegate;
 
@@ -197,30 +200,32 @@ private:
  * A list of element handles.
  * Provides high-level access to groups of elements, including accessing elements that implement specific interfaces.
  */
-class TYPEDELEMENTFRAMEWORK_API FTypedElementList
+UCLASS(Transient)
+class TYPEDELEMENTFRAMEWORK_API UTypedElementList : public UObject
 {
-public:
-	FTypedElementList(const FTypedElementList&) = delete;
-	FTypedElementList& operator=(const FTypedElementList&) = delete;
-
-	FTypedElementList(FTypedElementList&&) = delete;
-	FTypedElementList& operator=(FTypedElementList&&) = delete;
+	GENERATED_BODY()
 
-	~FTypedElementList();
+public:
+	//~ UObject interface
+	virtual void BeginDestroy() override;
 
-	static FTypedElementListPtr Private_CreateElementList(UTypedElementRegistry* InRegistry);
+	/**
+	 * Internal function used by the element registry to create an element list instance.
+	 */
+	static UTypedElementList* Private_CreateElementList(UTypedElementRegistry* InRegistry);
 
 	/**
 	 * Clone this list instance.
 	 * @note Only copies elements; does not copy any bindings!
 	 */
-	FTypedElementListPtr Clone() const;
+	UFUNCTION(BlueprintCallable, BlueprintPure=false, Category="TypedElementFramework|List")
+	UTypedElementList* Clone() const;
 
 	/**
 	 * Get the element handle at the given index.
 	 * @note Use IsValidIndex to test for validity.
 	 */
-	FORCEINLINE const FTypedElementHandle& operator[](const int32 InIndex) const
+	FORCEINLINE FTypedElementHandle operator[](const int32 InIndex) const
 	{
 		return GetElementHandleAt(InIndex);
 	}
@@ -229,7 +234,8 @@ public:
 	 * Get the element handle at the given index.
 	 * @note Use IsValidIndex to test for validity.
 	 */
-	FORCEINLINE const FTypedElementHandle& GetElementHandleAt(const int32 InIndex) const
+	UFUNCTION(BlueprintPure, Category="TypedElementFramework|List")
+	FORCEINLINE FTypedElementHandle GetElementHandleAt(const int32 InIndex) const
 	{
 		return ElementHandles[InIndex];
 	}
@@ -272,9 +278,25 @@ public:
 		TypedElementList_Private::GetElement(Registry.Get(), InElementHandle, OutElement);
 	}
 
+	/**
+	 * Get the element interface from the given handle.
+	 */
+	template <typename BaseInterfaceType>
+	BaseInterfaceType* GetElementInterface(const FTypedElementHandle& InElementHandle) const
+	{
+		return static_cast<BaseInterfaceType*>(GetElementInterface(InElementHandle, BaseInterfaceType::StaticClass()));
+	}
+
+	/**
+	 * Get the element interface from the given handle.
+	 */
+	UFUNCTION(BlueprintPure, Category="TypedElementFramework|List")
+	UTypedElementInterface* GetElementInterface(const FTypedElementHandle& InElementHandle, const TSubclassOf<UTypedElementInterface>& InBaseInterfaceType) const;
+
 	/**
 	 * Is the given index a valid entry within this element list?
 	 */
+	UFUNCTION(BlueprintPure, Category="TypedElementFramework|List")
 	FORCEINLINE bool IsValidIndex(const int32 InIndex) const
 	{
 		return ElementHandles.IsValidIndex(InIndex);
@@ -283,6 +305,7 @@ public:
 	/**
 	 * Get the number of entries within this element list.
 	 */
+	UFUNCTION(BlueprintPure, Category="TypedElementFramework|List")
 	FORCEINLINE int32 Num() const
 	{
 		return ElementHandles.Num();
@@ -291,6 +314,7 @@ public:
 	/**
 	 * Shrink this element list storage to avoid slack.
 	 */
+	UFUNCTION(BlueprintCallable, Category="TypedElementFramework|List")
 	FORCEINLINE void Shrink()
 	{
 		ElementCombinedIds.Shrink();
@@ -300,6 +324,7 @@ public:
 	/**
 	 * Pre-allocate enough memory in this element list to store the given number of entries.
 	 */
+	UFUNCTION(BlueprintCallable, Category="TypedElementFramework|List")
 	FORCEINLINE void Reserve(const int32 InSize)
 	{
 		ElementCombinedIds.Reserve(InSize);
@@ -309,6 +334,7 @@ public:
 	/**
 	 * Remove all entries from this element list, potentially leaving space allocated for the given number of entries.
 	 */
+	UFUNCTION(BlueprintCallable, Category="TypedElementFramework|List")
 	FORCEINLINE void Empty(const int32 InSlack = 0)
 	{
 		ElementCombinedIds.Empty(InSlack);
@@ -319,6 +345,7 @@ public:
 	/**
 	 * Remove all entries from this element list, preserving existing allocations.
 	 */
+	UFUNCTION(BlueprintCallable, Category="TypedElementFramework|List")
 	FORCEINLINE void Reset()
 	{
 		ElementCombinedIds.Reset();
@@ -337,6 +364,7 @@ public:
 	/**
 	 * Does this element list contain an entry for the given element handle?
 	 */
+	UFUNCTION(BlueprintPure, Category="TypedElementFramework|List")
 	FORCEINLINE bool Contains(const FTypedElementHandle& InElementHandle) const
 	{
 		return ContainsElementImpl(InElementHandle.GetId());
@@ -355,6 +383,7 @@ public:
 	 * Add the given element handle to this element list, if it isn't already in the list.
 	 * @return True if the element handle was added, false if it is already in the list.
 	 */
+	UFUNCTION(BlueprintCallable, Category="TypedElementFramework|List")
 	FORCEINLINE bool Add(const FTypedElementHandle& InElementHandle)
 	{
 		return AddElementImpl(CopyTemp(InElementHandle));
@@ -379,6 +408,15 @@ public:
 		return AddElementImpl(InElementOwner.AcquireHandle());
 	}
 
+	/**
+	 * Append the given element handles to this element list, for any that already in the list.
+	 */
+	UFUNCTION(BlueprintCallable, Category="TypedElementFramework|List")
+	FORCEINLINE void Append(const TArray<FTypedElementHandle>& InElementHandles)
+	{
+		Append(MakeArrayView(InElementHandles));
+	}
+	
 	/**
 	 * Append the given element handles to this element list, for any that already in the list.
 	 */
@@ -446,6 +484,7 @@ public:
 	 * Remove the given element handle from this element list, if it is in the list.
 	 * @return True if the element handle was removed, false if it isn't in the list.
 	 */
+	UFUNCTION(BlueprintCallable, Category="TypedElementFramework|List")
 	FORCEINLINE bool Remove(const FTypedElementHandle& InElementHandle)
 	{
 		return RemoveElementImpl(InElementHandle.GetId());
@@ -489,7 +528,7 @@ public:
 	 * Access the delegate that is invoked whenever this element list has been changed.
 	 * @note This is called automatically at the end of each frame, but can also be manually invoked by NotifyPendingChanges.
 	 */
-	DECLARE_EVENT_OneParam(FTypedElementList, FOnChanged, const FTypedElementList& /*InElementList*/);
+	DECLARE_EVENT_OneParam(UTypedElementList, FOnChanged, const UTypedElementList* /*InElementList*/);
 	FOnChanged& OnChanged()
 	{
 		return OnChangedDelegate;
@@ -544,7 +583,7 @@ private:
 		Cleared,
 	};
 
-	explicit FTypedElementList(UTypedElementRegistry* InRegistry);
+	void Initialize(UTypedElementRegistry* InRegistry);
 
 	bool AddElementImpl(FTypedElementHandle&& InElementHandle);
 	bool RemoveElementImpl(const FTypedElementId& InElementId);
diff --git a/Engine/Source/Runtime/TypedElementFramework/Public/TypedElementListPtr.h b/Engine/Source/Runtime/TypedElementFramework/Public/TypedElementListPtr.h
deleted file mode 100644
diff --git a/Engine/Source/Runtime/TypedElementFramework/Public/TypedElementRegistry.h b/Engine/Source/Runtime/TypedElementFramework/Public/TypedElementRegistry.h
--- a/Engine/Source/Runtime/TypedElementFramework/Public/TypedElementRegistry.h
+++ b/Engine/Source/Runtime/TypedElementFramework/Public/TypedElementRegistry.h
@@ -6,7 +6,6 @@
 #include "Misc/ScopeRWLock.h"
 #include "Containers/ArrayView.h"
 #include "Containers/SortedMap.h"
-#include "Templates/SubclassOf.h"
 #include "TypedElementHandle.h"
 #include "TypedElementList.h"
 #include "TypedElementRegistry.generated.h"
@@ -14,7 +13,7 @@
 /**
  * Registry of element types and their associated interfaces, along with the elements that represent their instances.
  */
-UCLASS()
+UCLASS(Transient)
 class TYPEDELEMENTFRAMEWORK_API UTypedElementRegistry : public UObject
 {
 	GENERATED_BODY()
@@ -39,6 +38,7 @@ public:
 	/**
 	 * Get the singleton instance of the registry used in most cases.
 	 */
+	UFUNCTION(BlueprintPure, Category="TypedElementFramework|Registry")
 	static UTypedElementRegistry* GetInstance();
 
 	/**
@@ -85,6 +85,15 @@ public:
 		return static_cast<BaseInterfaceType*>(GetElementInterfaceImpl(InElementHandle.GetId(), BaseInterfaceType::StaticClass()));
 	}
 
+	/**
+	 * Get the element interface supported by the given handle, or null if there is no support for this interface.
+	 */
+	UFUNCTION(BlueprintPure, Category="TypedElementFramework|Registry")
+	FORCEINLINE UTypedElementInterface* GetElementInterface(const FTypedElementHandle& InElementHandle, const TSubclassOf<UTypedElementInterface>& InBaseInterfaceType) const
+	{
+		return GetElementInterfaceImpl(InElementHandle.GetId(), InBaseInterfaceType);
+	}
+
 	/**
 	 * Create an element that doesn't require any additional payload data.
 	 * @note The associated handle ID should be something that can externally be used to uniquely identify this element, until DestroyElementHandle is called on this handle.
@@ -176,26 +185,27 @@ public:
 	/**
 	 * Create an empty list of elements associated with this registry.
 	 */
-	FORCEINLINE FTypedElementListPtr CreateElementList()
+	UFUNCTION(BlueprintCallable, Category="TypedElementFramework|Registry")
+	FORCEINLINE UTypedElementList* CreateElementList()
 	{
-		return FTypedElementList::Private_CreateElementList(this);
+		return UTypedElementList::Private_CreateElementList(this);
 	}
 
 	/**
 	 * Create an empty list of elements associated with this registry, populated from the given minimal IDs that are valid.
 	 */
-	FTypedElementListPtr CreateElementList(TArrayView<const FTypedElementId> InElementIds);
+	UTypedElementList* CreateElementList(TArrayView<const FTypedElementId> InElementIds);
 
 	/**
 	 * Create an empty list of elements associated with this registry, populated from the given handles that are valid.
 	 */
-	FTypedElementListPtr CreateElementList(TArrayView<const FTypedElementHandle> InElementHandles);
+	UTypedElementList* CreateElementList(TArrayView<const FTypedElementHandle> InElementHandles);
 
 	/**
 	 * Create an empty list of elements associated with this registry, populated from the given owners that are valid.
 	 */
 	template <typename ElementDataType>
-	FORCEINLINE FTypedElementListPtr CreateElementList(const TArray<TTypedElementOwner<ElementDataType>>& InElementOwners)
+	FORCEINLINE UTypedElementList* CreateElementList(const TArray<TTypedElementOwner<ElementDataType>>& InElementOwners)
 	{
 		return CreateElementList(MakeArrayView(InElementOwners));
 	}
@@ -204,26 +214,26 @@ public:
 	 * Create an empty list of elements associated with this registry, populated from the given owners that are valid.
 	 */
 	template <typename ElementDataType>
-	FTypedElementListPtr CreateElementList(TArrayView<const TTypedElementOwner<ElementDataType>> InElementOwners)
+	UTypedElementList* CreateElementList(TArrayView<const TTypedElementOwner<ElementDataType>> InElementOwners)
 	{
-		FTypedElementListPtr ElementList = CreateElementList();
+		UTypedElementList* ElementList = CreateElementList();
 		ElementList->Append(InElementOwners);
 		return ElementList;
 	}
 
-	void Private_OnElementListCreated(FTypedElementList* InElementList)
+	void Private_OnElementListCreated(UTypedElementList* InElementList)
 	{
 		FWriteScopeLock ActiveElementListsLock(ActiveElementListsRW);
 		ActiveElementLists.Add(InElementList);
 	}
 
-	void Private_OnElementListDestroyed(FTypedElementList* InElementList)
+	void Private_OnElementListDestroyed(UTypedElementList* InElementList)
 	{
 		FWriteScopeLock ActiveElementListsLock(ActiveElementListsRW);
 		ActiveElementLists.Remove(InElementList);
 	}
 
-	// Note: Access for FTypedElementList
+	// Note: Access for UTypedElementList
 	FORCEINLINE void Private_GetElementImpl(const FTypedElementHandle& InElementHandle, const UClass* InBaseInterfaceType, FTypedElement& OutElement) const
 	{
 		GetElementImpl(InElementHandle, InBaseInterfaceType, OutElement);
@@ -393,5 +403,5 @@ private:
 	TSortedMap<FName, FTypedHandleTypeId, FDefaultAllocator, FNameFastLess> RegisteredElementTypesNameToId;
 
 	mutable FRWLock ActiveElementListsRW;
-	TSet<FTypedElementList*> ActiveElementLists;
+	TSet<UTypedElementList*> ActiveElementLists;
 };
diff --git a/Engine/Source/Runtime/TypedElementInterfaces/Public/TypedElementSelectionInterface.h b/Engine/Source/Runtime/TypedElementInterfaces/Public/TypedElementSelectionInterface.h
--- a/Engine/Source/Runtime/TypedElementInterfaces/Public/TypedElementSelectionInterface.h
+++ b/Engine/Source/Runtime/TypedElementInterfaces/Public/TypedElementSelectionInterface.h
@@ -19,7 +19,7 @@ public:
 
 	/**
 	 * Retreive the object instance that should be selected via this handle.
-	 * @note This exists to allow the legacy USelection to bridge to an FTypedElementList instance. It should not be used in new code!
+	 * @note This exists to allow the legacy USelection to bridge to an UTypedElementList instance. It should not be used in new code!
 	 */
 	UFUNCTION(BlueprintCallable, Category="TypedElementInterfaces|Selection")
 	virtual UObject* Legacy_GetSelectionObject(const FTypedElementHandle& InElementHandle) { return nullptr; }
