Fixed typo

#rnx

diff --git a/Engine/Source/Runtime/TypedElementFramework/Private/TypedElementFrameworkTests.cpp b/Engine/Source/Runtime/TypedElementFramework/Private/TypedElementFrameworkTests.cpp
--- a/Engine/Source/Runtime/TypedElementFramework/Private/TypedElementFrameworkTests.cpp
+++ b/Engine/Source/Runtime/TypedElementFramework/Private/TypedElementFrameworkTests.cpp
@@ -122,7 +122,7 @@ bool FTypedElementRegistrySmokeTest::RunTest(const FString& Parameters)
 		TestInterfaceAccess(ElementHandle);
 	}
 	
-	for (const TTypedElement<UTestTypedElementInterfaceA>& ElementHandle : ElementList->InterateInterface<UTestTypedElementInterfaceA>())
+	for (const TTypedElement<UTestTypedElementInterfaceA>& ElementHandle : ElementList->IterateInterface<UTestTypedElementInterfaceA>())
 	{
 		TestInterfaceAHandle(ElementHandle);
 	}
@@ -213,7 +213,7 @@ bool FTypedElementRegistryPerfTest::RunTest(const FString& Parameters)
 	// Iterate all elements that implement an interface
 	{
 		FScopedDurationTimeLogger Timer(FString::Printf(TEXT("Iterating %d elements in list"), ElementList->Num()));
-		for (const TTypedElement<UTestTypedElementInterfaceA>& ElementHandle : ElementList->InterateInterface<UTestTypedElementInterfaceA>())
+		for (const TTypedElement<UTestTypedElementInterfaceA>& ElementHandle : ElementList->IterateInterface<UTestTypedElementInterfaceA>())
 		{
 		}
 	}
diff --git a/Engine/Source/Runtime/TypedElementFramework/Public/TypedElementList.h b/Engine/Source/Runtime/TypedElementFramework/Public/TypedElementList.h
--- a/Engine/Source/Runtime/TypedElementFramework/Public/TypedElementList.h
+++ b/Engine/Source/Runtime/TypedElementFramework/Public/TypedElementList.h
@@ -583,7 +583,7 @@ public:
 	 * Get an iterator for elements that implement the given element interface.
 	 */
 	template <typename BaseInterfaceType>
-	FORCEINLINE TypedElementList_Private::TTypedElementListInterfaceIteratorProxy<BaseInterfaceType> InterateInterface() const
+	FORCEINLINE TypedElementList_Private::TTypedElementListInterfaceIteratorProxy<BaseInterfaceType> IterateInterface() const
 	{
 		return TypedElementList_Private::TTypedElementListInterfaceIteratorProxy<BaseInterfaceType>(Registry.Get(), ElementHandles);
 	}
