Allow default GetTypedElementOwnerStoreKeyDebugString function to work with any key type

#rnx

diff --git a/Engine/Source/Runtime/TypedElementFramework/Public/TypedElementOwnerStore.h b/Engine/Source/Runtime/TypedElementFramework/Public/TypedElementOwnerStore.h
--- a/Engine/Source/Runtime/TypedElementFramework/Public/TypedElementOwnerStore.h
+++ b/Engine/Source/Runtime/TypedElementFramework/Public/TypedElementOwnerStore.h
@@ -91,7 +91,8 @@ using TTypedElementOwnerScopedAccess = TTypedElementOwnerScopedAccessImpl<const
 template <typename ElementDataType>
 using TTypedElementOwnerScopedMutableAccess = TTypedElementOwnerScopedAccessImpl<TTypedElementOwner<ElementDataType>>;
 
-inline FString GetTypedElementOwnerStoreKeyDebugString(const void* InKey)
+template <typename KeyDataType>
+inline FString GetTypedElementOwnerStoreKeyDebugString(const KeyDataType& InKey)
 {
 	return TEXT("<GetTypedElementOwnerStoreKeyDebugString not implemented for key type>");
 }
