Fixed truncated typed element reference logging

The callstacks were too long to be logged by a single assert message.

#rnx

diff --git a/Engine/Source/Runtime/TypedElementFramework/Public/TypedElementData.h b/Engine/Source/Runtime/TypedElementFramework/Public/TypedElementData.h
--- a/Engine/Source/Runtime/TypedElementFramework/Public/TypedElementData.h
+++ b/Engine/Source/Runtime/TypedElementFramework/Public/TypedElementData.h
@@ -44,29 +44,17 @@ public:
 	FTypedElementReference(FTypedElementReference&&) = default;
 	FTypedElementReference& operator=(FTypedElementReference&&) = default;
 
-	void AppendString(FString& OutStr) const
+	void LogReference() const
 	{
 		ANSICHAR CallstackText[4096];
 		for (uint32 CallstackIndex = TypedHandleRefTrackingSkipCount; CallstackIndex < CallstackDepth; ++CallstackIndex)
 		{
-			if (CallstackIndex > TypedHandleRefTrackingSkipCount)
-			{
-				OutStr.Append(LINE_TERMINATOR);
-			}
-
 			CallstackText[0] = 0;
 			FPlatformStackWalk::ProgramCounterToHumanReadableString(CallstackIndex - TypedHandleRefTrackingSkipCount, Callstack[CallstackIndex], CallstackText, UE_ARRAY_COUNT(CallstackText));
-			OutStr.Append(CallstackText);
+			UE_LOG(LogCore, Error, TEXT("%s"), ANSI_TO_TCHAR(CallstackText));
 		}
 	}
 
-	FString ToString() const
-	{
-		FString Str;
-		AppendString(Str);
-		return Str;
-	}
-
 private:
 	uint64 Callstack[TypedHandleRefTrackingDepth + TypedHandleRefTrackingSkipCount];
 	uint32 CallstackDepth = 0;
@@ -143,32 +131,24 @@ public:
 #endif	// UE_TYPED_ELEMENT_HAS_REFCOUNTING
 	}
 
-	void AppendRefString(FString& OutStr) const
+	void LogReferences() const
 	{
-#if UE_TYPED_ELEMENT_HAS_REFCOUNTING
-		OutStr.Appendf(TEXT("Ref-count: %d"), FPlatformAtomics::AtomicRead(&RefCount));
-#endif	// UE_TYPED_ELEMENT_HAS_REFCOUNTING
 #if UE_TYPED_ELEMENT_HAS_REFTRACKING
 		{
 			FScopeLock ReferencesLock(&ReferencesCS);
-			OutStr.Append(LINE_TERMINATOR TEXT("References:") LINE_TERMINATOR);
+			UE_LOG(LogCore, Error, TEXT("External Element References:"));
 			for (const FTypedElementReference& Reference : References)
 			{
-				OutStr.Append(TEXT("-----------------------------------------------") LINE_TERMINATOR);
-				Reference.AppendString(OutStr);
+				UE_LOG(LogCore, Error, TEXT("-----------------------------------------------"));
+				Reference.LogReference();
 			}
-			OutStr.Append(LINE_TERMINATOR TEXT("==============================================="));
+			UE_LOG(LogCore, Error, TEXT("==============================================="));
 		}
+#else	// UE_TYPED_ELEMENT_HAS_REFTRACKING
+		UE_LOG(LogCore, Error, TEXT("UE_TYPED_ELEMENT_HAS_REFTRACKING is disabled. Enable it and recompile to see reference tracking."));
 #endif	// UE_TYPED_ELEMENT_HAS_REFTRACKING
 	}
 
-	FString GetRefString() const
-	{
-		FString Str;
-		AppendRefString(Str);
-		return Str;
-	}
-
 	virtual const void* GetUntypedData() const
 	{
 		return nullptr;
diff --git a/Engine/Source/Runtime/TypedElementFramework/Public/TypedElementLimits.h b/Engine/Source/Runtime/TypedElementFramework/Public/TypedElementLimits.h
--- a/Engine/Source/Runtime/TypedElementFramework/Public/TypedElementLimits.h
+++ b/Engine/Source/Runtime/TypedElementFramework/Public/TypedElementLimits.h
@@ -24,7 +24,7 @@ constexpr SIZE_T TypedHandleMaxTypeId = ((SIZE_T)1 << TypedHandleTypeIdBits) - 1
 constexpr SIZE_T TypedHandleMaxElementId = ((SIZE_T)1 << TypedHandleElementIdBits) - 1;
 
 constexpr SIZE_T TypedHandleRefTrackingSkipCount = 1;
-constexpr SIZE_T TypedHandleRefTrackingDepth = 15;
+constexpr SIZE_T TypedHandleRefTrackingDepth = 31;
 
 using FTypedHandleTypeId = uint8;
 using FTypedHandleElementId = int32;
diff --git a/Engine/Source/Runtime/TypedElementFramework/Public/TypedElementRegistry.h b/Engine/Source/Runtime/TypedElementFramework/Public/TypedElementRegistry.h
--- a/Engine/Source/Runtime/TypedElementFramework/Public/TypedElementRegistry.h
+++ b/Engine/Source/Runtime/TypedElementFramework/Public/TypedElementRegistry.h
@@ -337,7 +337,13 @@ private:
 		{
 			const FTypedElementInternalData& ElementData = RegisteredElementType->GetDataForElement(InOutElementOwner.GetId().GetElementId());
 			const FTypedElementRefCount RefCount = ElementData.GetRefCount();
-			checkf(RefCount <= 1, TEXT("Element is still externally referenced when being destroyed! %s"), *ElementData.GetRefString());
+
+			const bool bIsExternallyReferenced = RefCount > 1;
+			if (bIsExternallyReferenced)
+			{
+				ElementData.LogReferences();
+				UE_LOG(LogCore, Fatal, TEXT("Element is still externally referenced when being destroyed! Ref-count: %d; see above for reference information (if available)."), RefCount);
+			}
 		}
 #endif	// DO_CHECK
 
