Made typed element includes more consistent

 - Moved all the "framework" code under "Elements/Framework".
 - Moved all the "interface" code under "Elements/Interfaces".

This lets the code be split over multiple modules (including the Engine and Editor) while retaining consistent include paths with things like "Elements/Actor", "Elements/Component", etc, as well as letting new interfaces and frameworky things be added outside of the core typed element modules while still maintaining a coherent include scheme.

diff --git a/Engine/Source/Editor/AssetPlacementEdMode/Public/AssetPlacementEdMode.h b/Engine/Source/Editor/AssetPlacementEdMode/Public/AssetPlacementEdMode.h
--- a/Engine/Source/Editor/AssetPlacementEdMode/Public/AssetPlacementEdMode.h
+++ b/Engine/Source/Editor/AssetPlacementEdMode/Public/AssetPlacementEdMode.h
@@ -4,7 +4,7 @@
 
 #include "CoreMinimal.h"
 #include "Tools/UEdMode.h"
-#include "TypedElementHandle.h"
+#include "Elements/Framework/TypedElementHandle.h"
 
 #include "AssetPlacementEdMode.generated.h"
 
diff --git a/Engine/Source/Editor/EditorFramework/Private/Elements/Framework/TypedElementSelectionSet.cpp b/Engine/Source/Editor/EditorFramework/Private/Elements/Framework/TypedElementSelectionSet.cpp
new file mode 100644
--- /dev/null
+++ b/Engine/Source/Editor/EditorFramework/Private/Elements/Framework/TypedElementSelectionSet.cpp
@@ -0,0 +1,245 @@
+// Copyright Epic Games, Inc. All Rights Reserved.
+
+#include "Elements/Framework/TypedElementSelectionSet.h"
+#include "Elements/Framework/TypedElementRegistry.h"
+
+UTypedElementSelectionSet::UTypedElementSelectionSet()
+{
+	if (!HasAnyFlags(RF_ClassDefaultObject))
+	{
+		ElementList = UTypedElementRegistry::GetInstance()->CreateElementList();
+		ElementList->OnPreChange().AddUObject(this, &UTypedElementSelectionSet::OnElementListPreChange);
+		ElementList->OnChanged().AddUObject(this, &UTypedElementSelectionSet::OnElementListChanged);
+	}
+}
+
+bool UTypedElementSelectionSet::Modify(bool bAlwaysMarkDirty)
+{
+	if (GUndo && CanModify())
+	{
+		for (const TTypedElement<UTypedElementSelectionInterface>& SelectionElement : ElementList->IterateInterface<UTypedElementSelectionInterface>())
+		{
+			if (SelectionElement.ShouldPreventTransactions())
+			{
+				return false;
+			}
+		}
+
+		return Super::Modify(bAlwaysMarkDirty);
+	}
+
+	return false;
+}
+
+void UTypedElementSelectionSet::Serialize(FArchive& Ar)
+{
+	checkf(Ar.IsTransacting(), TEXT("UTypedElementSelectionSet can only be serialized for transactions!"));
+
+	if (Ar.IsSaving())
+	{
+		FTypedHandleTypeId ElementTypeId = 0;
+
+		for (const TTypedElement<UTypedElementSelectionInterface>& SelectionElement : ElementList->IterateInterface<UTypedElementSelectionInterface>())
+		{
+			ElementTypeId = SelectionElement.GetId().GetTypeId();
+			Ar << ElementTypeId;
+
+			SelectionElement.WriteTransactedElement(Ar);
+		}
+
+		// End of the list
+		ElementTypeId = 0;
+		Ar << ElementTypeId;
+	}
+	else if (Ar.IsLoading())
+	{
+		TArray<FTypedElementHandle, TInlineAllocator<256>> SelectedElements;
+
+		{
+			UTypedElementRegistry* Registry = UTypedElementRegistry::GetInstance();
+
+			FTypedHandleTypeId ElementTypeId = 0;
+			for (;;)
+			{
+				Ar << ElementTypeId;
+				if (!ElementTypeId)
+				{
+					// End of the list
+					break;
+				}
+
+				UTypedElementSelectionInterface* ElementTypeSelectionInterface = Registry->GetElementInterface<UTypedElementSelectionInterface>(ElementTypeId);
+				checkf(ElementTypeSelectionInterface, TEXT("Failed to find selection interface for a previously transacted element type!"));
+				FTypedElementHandle SelectedElement = ElementTypeSelectionInterface->ReadTransactedElement(Ar);
+				if (SelectedElement)
+				{
+					SelectedElements.Add(MoveTemp(SelectedElement));
+				}
+			}
+		}
+
+		{
+			FTypedElementListLegacySyncScopedBatch LegacySyncBatch(ElementList, /*bNotify*/false);
+			TGuardValue<bool> GuardIsRestoringFromTransaction(bIsRestoringFromTransaction, true);
+
+			// TODO: Work out the intersection of the before and after state instead of clearing and reselecting?
+
+			const FTypedElementSelectionOptions SelectionOptions = FTypedElementSelectionOptions()
+				.SetAllowHidden(true)
+				.SetAllowGroups(false)
+				.SetWarnIfLocked(false);
+
+			ClearSelection(SelectionOptions);
+			SelectElements(SelectedElements, SelectionOptions);
+		}
+	}
+}
+
+bool UTypedElementSelectionSet::IsElementSelected(const FTypedElementHandle& InElementHandle, const FTypedElementIsSelectedOptions InSelectionOptions) const
+{
+	FTypedElementSelectionSetElement SelectionSetElement = ResolveSelectionSetElement(InElementHandle);
+	return SelectionSetElement && SelectionSetElement.IsElementSelected(InSelectionOptions);
+}
+
+bool UTypedElementSelectionSet::CanSelectElement(const FTypedElementHandle& InElementHandle, const FTypedElementSelectionOptions InSelectionOptions) const
+{
+	FTypedElementSelectionSetElement SelectionSetElement = ResolveSelectionSetElement(InElementHandle);
+	return SelectionSetElement && SelectionSetElement.CanSelectElement(InSelectionOptions);
+}
+
+bool UTypedElementSelectionSet::CanDeselectElement(const FTypedElementHandle& InElementHandle, const FTypedElementSelectionOptions InSelectionOptions) const
+{
+	FTypedElementSelectionSetElement SelectionSetElement = ResolveSelectionSetElement(InElementHandle);
+	return SelectionSetElement && SelectionSetElement.CanDeselectElement(InSelectionOptions);
+}
+
+bool UTypedElementSelectionSet::SelectElement(const FTypedElementHandle& InElementHandle, const FTypedElementSelectionOptions InSelectionOptions)
+{
+	FTypedElementSelectionSetElement SelectionSetElement = ResolveSelectionSetElement(InElementHandle);
+	return SelectionSetElement && SelectionSetElement.CanSelectElement(InSelectionOptions) && SelectionSetElement.SelectElement(InSelectionOptions);
+}
+
+bool UTypedElementSelectionSet::SelectElements(const TArray<FTypedElementHandle>& InElementHandles, const FTypedElementSelectionOptions InSelectionOptions)
+{
+	return SelectElements(MakeArrayView(InElementHandles), InSelectionOptions);
+}
+
+bool UTypedElementSelectionSet::SelectElements(TArrayView<const FTypedElementHandle> InElementHandles, const FTypedElementSelectionOptions InSelectionOptions)
+{
+	FTypedElementListLegacySyncScopedBatch LegacySyncBatch(ElementList, InSelectionOptions.AllowLegacyNotifications());
+
+	bool bSelectionChanged = false;
+
+	for (const FTypedElementHandle& ElementHandle : InElementHandles)
+	{
+		bSelectionChanged |= SelectElement(ElementHandle, InSelectionOptions);
+	}
+
+	return bSelectionChanged;
+}
+
+bool UTypedElementSelectionSet::DeselectElement(const FTypedElementHandle& InElementHandle, const FTypedElementSelectionOptions InSelectionOptions)
+{
+	FTypedElementSelectionSetElement SelectionSetElement = ResolveSelectionSetElement(InElementHandle);
+	return SelectionSetElement && SelectionSetElement.CanDeselectElement(InSelectionOptions) && SelectionSetElement.DeselectElement(InSelectionOptions);
+}
+
+bool UTypedElementSelectionSet::DeselectElements(const TArray<FTypedElementHandle>& InElementHandles, const FTypedElementSelectionOptions InSelectionOptions)
+{
+	return DeselectElements(MakeArrayView(InElementHandles), InSelectionOptions);
+}
+
+bool UTypedElementSelectionSet::DeselectElements(TArrayView<const FTypedElementHandle> InElementHandles, const FTypedElementSelectionOptions InSelectionOptions)
+{
+	FTypedElementListLegacySyncScopedBatch LegacySyncBatch(ElementList, InSelectionOptions.AllowLegacyNotifications());
+
+	bool bSelectionChanged = false;
+
+	for (const FTypedElementHandle& ElementHandle : InElementHandles)
+	{
+		bSelectionChanged |= DeselectElement(ElementHandle, InSelectionOptions);
+	}
+
+	return bSelectionChanged;
+}
+
+bool UTypedElementSelectionSet::ClearSelection(const FTypedElementSelectionOptions InSelectionOptions)
+{
+	FTypedElementListLegacySyncScopedBatch LegacySyncBatch(ElementList, InSelectionOptions.AllowLegacyNotifications());
+
+	bool bSelectionChanged = false;
+
+	// Run deselection via the selection interface where possible
+	{
+		// Take a copy of the currently selected elements to avoid mutating the selection set while iterating
+		TArray<FTypedElementHandle, TInlineAllocator<8>> ElementsCopy;
+		ElementsCopy.Reserve(ElementList->Num());
+		for (const FTypedElementHandle& ElementHandle : *ElementList)
+		{
+			ElementsCopy.Emplace(ElementHandle);
+		}
+
+		for (const FTypedElementHandle& ElementHandle : ElementsCopy)
+		{
+			bSelectionChanged |= DeselectElement(ElementHandle, InSelectionOptions);
+		}
+	}
+
+	// TODO: BSP surfaces?
+
+	// If anything remains in the selection set after processing elements that implement this interface, just clear it
+	if (ElementList->Num() > 0)
+	{
+		bSelectionChanged = true;
+		ElementList->Reset();
+	}
+
+	return bSelectionChanged;
+}
+
+bool UTypedElementSelectionSet::AllowSelectionModifiers(const FTypedElementHandle& InElementHandle) const
+{
+	FTypedElementSelectionSetElement SelectionSetElement = ResolveSelectionSetElement(InElementHandle);
+	return SelectionSetElement && SelectionSetElement.AllowSelectionModifiers();
+}
+
+FTypedElementHandle UTypedElementSelectionSet::GetSelectionElement(const FTypedElementHandle& InElementHandle, const ETypedElementSelectionMethod InSelectionMethod) const
+{
+	FTypedElementSelectionSetElement SelectionSetElement = ResolveSelectionSetElement(InElementHandle);
+	return SelectionSetElement ? SelectionSetElement.GetSelectionElement(InSelectionMethod) : FTypedElementHandle();
+}
+
+void UTypedElementSelectionSet::RegisterAssetEditorSelectionProxy(const FName InElementTypeName, UTypedElementAssetEditorSelectionProxy* InAssetEditorSelectionProxy)
+{
+	const FTypedHandleTypeId ElementTypeId = UTypedElementRegistry::GetInstance()->GetRegisteredElementTypeId(InElementTypeName);
+	checkf(ElementTypeId > 0, TEXT("Element type '%s' has not been registered!"), *InElementTypeName.ToString());
+	RegisteredAssetEditorSelectionProxies[ElementTypeId - 1] = InAssetEditorSelectionProxy;
+}
+
+FTypedElementSelectionSetElement UTypedElementSelectionSet::ResolveSelectionSetElement(const FTypedElementHandle& InElementHandle) const
+{
+	if (InElementHandle)
+	{
+		UTypedElementAssetEditorSelectionProxy* AssetEditorSelectionProxy = RegisteredAssetEditorSelectionProxies[InElementHandle.GetId().GetTypeId() - 1];
+		return FTypedElementSelectionSetElement(ElementList->GetElement<UTypedElementSelectionInterface>(InElementHandle), ElementList, AssetEditorSelectionProxy ? AssetEditorSelectionProxy : GetMutableDefault<UTypedElementAssetEditorSelectionProxy>());
+	}
+	return FTypedElementSelectionSetElement();
+}
+
+void UTypedElementSelectionSet::OnElementListPreChange(const UTypedElementList* InElementList)
+{
+	check(InElementList == ElementList);
+	OnPreChangeDelegate.Broadcast(this);
+
+	if (!bIsRestoringFromTransaction)
+	{
+		// Track the pre-change state for undo/redo
+		Modify();
+	}
+}
+
+void UTypedElementSelectionSet::OnElementListChanged(const UTypedElementList* InElementList)
+{
+	check(InElementList == ElementList);
+	OnChangedDelegate.Broadcast(this);
+}
diff --git a/Engine/Source/Editor/EditorFramework/Private/Elements/TypedElementSelectionSet.cpp b/Engine/Source/Editor/EditorFramework/Private/Elements/TypedElementSelectionSet.cpp
deleted file mode 100644
diff --git a/Engine/Source/Editor/EditorFramework/Public/Elements/Framework/TypedElementSelectionSet.h b/Engine/Source/Editor/EditorFramework/Public/Elements/Framework/TypedElementSelectionSet.h
new file mode 100644
--- /dev/null
+++ b/Engine/Source/Editor/EditorFramework/Public/Elements/Framework/TypedElementSelectionSet.h
@@ -0,0 +1,259 @@
+// Copyright Epic Games, Inc. All Rights Reserved.
+
+#pragma once
+
+#include "Elements/Framework/TypedElementList.h"
+#include "Elements/Interfaces/TypedElementSelectionInterface.h"
+#include "TypedElementSelectionSet.generated.h"
+
+/**
+ * Proxy type used to allow asset editors (such as the level editor) to override the base behavior of element selection,
+ * by injecting extra pre/post selection logic around the call into the selection interface for an element type.
+ */
+UCLASS()
+class EDITORFRAMEWORK_API UTypedElementAssetEditorSelectionProxy : public UObject
+{
+	GENERATED_BODY()
+
+public:
+	//~ See UTypedElementSelectionInterface for API docs
+	virtual bool IsElementSelected(const TTypedElement<UTypedElementSelectionInterface>& InElementSelectionHandle, const UTypedElementList* InSelectionSet, const FTypedElementIsSelectedOptions& InSelectionOptions) { return InElementSelectionHandle.IsElementSelected(InSelectionSet, InSelectionOptions); }
+	virtual bool CanSelectElement(const TTypedElement<UTypedElementSelectionInterface>& InElementSelectionHandle, const FTypedElementSelectionOptions& InSelectionOptions) { return InElementSelectionHandle.CanSelectElement(InSelectionOptions); }
+	virtual bool CanDeselectElement(const TTypedElement<UTypedElementSelectionInterface>& InElementSelectionHandle, const FTypedElementSelectionOptions& InSelectionOptions) { return InElementSelectionHandle.CanDeselectElement(InSelectionOptions); }
+	virtual bool SelectElement(const TTypedElement<UTypedElementSelectionInterface>& InElementSelectionHandle, UTypedElementList* InSelectionSet, const FTypedElementSelectionOptions& InSelectionOptions) { return InElementSelectionHandle.SelectElement(InSelectionSet, InSelectionOptions); }
+	virtual bool DeselectElement(const TTypedElement<UTypedElementSelectionInterface>& InElementSelectionHandle, UTypedElementList* InSelectionSet, const FTypedElementSelectionOptions& InSelectionOptions) { return InElementSelectionHandle.DeselectElement(InSelectionSet, InSelectionOptions); }
+	virtual bool AllowSelectionModifiers(const TTypedElement<UTypedElementSelectionInterface>& InElementSelectionHandle, const UTypedElementList* InSelectionSet) { return InElementSelectionHandle.AllowSelectionModifiers(InSelectionSet); }
+	virtual FTypedElementHandle GetSelectionElement(const TTypedElement<UTypedElementSelectionInterface>& InElementSelectionHandle, const UTypedElementList* InCurrentSelection, const ETypedElementSelectionMethod InSelectionMethod) { return InElementSelectionHandle.GetSelectionElement(InCurrentSelection, InSelectionMethod); }
+};
+
+/**
+ * Utility to hold a typed element handle and its associated selection interface and asset editor selection proxy.
+ */
+struct EDITORFRAMEWORK_API FTypedElementSelectionSetElement
+{
+public:
+	FTypedElementSelectionSetElement() = default;
+
+	FTypedElementSelectionSetElement(TTypedElement<UTypedElementSelectionInterface> InElementSelectionHandle, UTypedElementList* InElementList, UTypedElementAssetEditorSelectionProxy* InAssetEditorSelectionProxy)
+		: ElementSelectionHandle(MoveTemp(InElementSelectionHandle))
+		, ElementList(InElementList)
+		, AssetEditorSelectionProxy(InAssetEditorSelectionProxy)
+	{
+	}
+
+	FTypedElementSelectionSetElement(const FTypedElementSelectionSetElement&) = default;
+	FTypedElementSelectionSetElement& operator=(const FTypedElementSelectionSetElement&) = default;
+
+	FTypedElementSelectionSetElement(FTypedElementSelectionSetElement&&) = default;
+	FTypedElementSelectionSetElement& operator=(FTypedElementSelectionSetElement&&) = default;
+
+	FORCEINLINE explicit operator bool() const
+	{
+		return IsSet();
+	}
+
+	FORCEINLINE bool IsSet() const
+	{
+		return ElementSelectionHandle.IsSet()
+			&& ElementList
+			&& AssetEditorSelectionProxy;
+	}
+
+	//~ See UTypedElementSelectionInterface for API docs
+	bool IsElementSelected(const FTypedElementIsSelectedOptions& InSelectionOptions) const { return AssetEditorSelectionProxy->IsElementSelected(ElementSelectionHandle, ElementList, InSelectionOptions); }
+	bool CanSelectElement(const FTypedElementSelectionOptions& InSelectionOptions) const { return AssetEditorSelectionProxy->CanSelectElement(ElementSelectionHandle, InSelectionOptions); }
+	bool CanDeselectElement(const FTypedElementSelectionOptions& InSelectionOptions) const { return AssetEditorSelectionProxy->CanDeselectElement(ElementSelectionHandle, InSelectionOptions); }
+	bool SelectElement(const FTypedElementSelectionOptions& InSelectionOptions) const { return AssetEditorSelectionProxy->SelectElement(ElementSelectionHandle, ElementList, InSelectionOptions); }
+	bool DeselectElement(const FTypedElementSelectionOptions& InSelectionOptions) const { return AssetEditorSelectionProxy->DeselectElement(ElementSelectionHandle, ElementList, InSelectionOptions); }
+	bool AllowSelectionModifiers() const { return AssetEditorSelectionProxy->AllowSelectionModifiers(ElementSelectionHandle, ElementList); }
+	FTypedElementHandle GetSelectionElement(const ETypedElementSelectionMethod InSelectionMethod) const { return AssetEditorSelectionProxy->GetSelectionElement(ElementSelectionHandle, ElementList, InSelectionMethod); }
+
+private:
+	TTypedElement<UTypedElementSelectionInterface> ElementSelectionHandle;
+	UTypedElementList* ElementList = nullptr;
+	UTypedElementAssetEditorSelectionProxy* AssetEditorSelectionProxy = nullptr;
+};
+
+/**
+ * A wrapper around an element list that ensures mutation goes via the selection 
+ * interfaces, as well as providing some utilities for batching operations.
+ */
+UCLASS(Transient)
+class EDITORFRAMEWORK_API UTypedElementSelectionSet : public UObject
+{
+	GENERATED_BODY()
+
+public:
+	UTypedElementSelectionSet();
+
+	//~ UObject interface
+	virtual bool Modify(bool bAlwaysMarkDirty = true) override;
+	virtual void Serialize(FArchive& Ar) override;
+
+	/**
+	 * Test to see whether the given element is currently considered selected.
+	 */
+	UFUNCTION(BlueprintPure, Category="TypedElementFramework|Selection")
+	bool IsElementSelected(const FTypedElementHandle& InElementHandle, const FTypedElementIsSelectedOptions InSelectionOptions) const;
+
+	/**
+	 * Test to see whether the given element can be selected.
+	 */
+	UFUNCTION(BlueprintPure, Category="TypedElementFramework|Selection")
+	bool CanSelectElement(const FTypedElementHandle& InElementHandle, const FTypedElementSelectionOptions InSelectionOptions) const;
+
+	/**
+	 * Test to see whether the given element can be deselected.
+	 */
+	UFUNCTION(BlueprintPure, Category="TypedElementFramework|Selection")
+	bool CanDeselectElement(const FTypedElementHandle& InElementHandle, const FTypedElementSelectionOptions InSelectionOptions) const;
+
+	/**
+	 * Attempt to select the given element.
+	 * @return True if the selection was changed, false otherwise.
+	 */
+	UFUNCTION(BlueprintCallable, Category="TypedElementFramework|Selection")
+	bool SelectElement(const FTypedElementHandle& InElementHandle, const FTypedElementSelectionOptions InSelectionOptions);
+
+	/**
+	 * Attempt to select the given elements.
+	 * @return True if the selection was changed, false otherwise.
+	 */
+	UFUNCTION(BlueprintCallable, Category="TypedElementFramework|Selection")
+	bool SelectElements(const TArray<FTypedElementHandle>& InElementHandles, const FTypedElementSelectionOptions InSelectionOptions);
+	bool SelectElements(TArrayView<const FTypedElementHandle> InElementHandles, const FTypedElementSelectionOptions InSelectionOptions);
+
+	/**
+	 * Attempt to deselect the given element.
+	 * @return True if the selection was changed, false otherwise.
+	 */
+	UFUNCTION(BlueprintCallable, Category="TypedElementFramework|Selection")
+	bool DeselectElement(const FTypedElementHandle& InElementHandle, const FTypedElementSelectionOptions InSelectionOptions);
+
+	/**
+	 * Attempt to deselect the given elements.
+	 * @return True if the selection was changed, false otherwise.
+	 */
+	UFUNCTION(BlueprintCallable, Category="TypedElementFramework|Selection")
+	bool DeselectElements(const TArray<FTypedElementHandle>& InElementHandles, const FTypedElementSelectionOptions InSelectionOptions);
+	bool DeselectElements(TArrayView<const FTypedElementHandle> InElementHandles, const FTypedElementSelectionOptions InSelectionOptions);
+
+	/**
+	 * Clear the current selection.
+	 * @return True if the selection was changed, false otherwise.
+	 */
+	UFUNCTION(BlueprintCallable, Category="TypedElementFramework|Selection")
+	bool ClearSelection(const FTypedElementSelectionOptions InSelectionOptions);
+
+	/**
+	 * Test to see whether selection modifiers (Ctrl or Shift) are allowed while selecting this element.
+	 */
+	UFUNCTION(BlueprintPure, Category="TypedElementFramework|Selection")
+	bool AllowSelectionModifiers(const FTypedElementHandle& InElementHandle) const;
+
+	/**
+	 * Given an element, return the element that should actually perform a selection operation.
+	 */
+	UFUNCTION(BlueprintPure, Category="TypedElementFramework|Selection")
+	FTypedElementHandle GetSelectionElement(const FTypedElementHandle& InElementHandle, const ETypedElementSelectionMethod InSelectionMethod) const;
+
+	/**
+	 * Access the delegate that is invoked whenever this element list is potentially about to change.
+	 * @note This may be called even if no actual change happens, though once a change does happen it won't be called again until after the next call to NotifyPendingChanges.
+	 */
+	DECLARE_EVENT_OneParam(UTypedElementSelectionSet, FOnPreChange, const UTypedElementSelectionSet* /*InElementSelectionSet*/);
+	FOnPreChange& OnPreChange()
+	{
+		return OnPreChangeDelegate;
+	}
+
+	/**
+	 * Access the delegate that is invoked whenever the underlying element list has been changed.
+	 * @note This is called automatically at the end of each frame, but can also be manually invoked by NotifyPendingChanges.
+	 */
+	DECLARE_EVENT_OneParam(UTypedElementSelectionSet, FOnChanged, const UTypedElementSelectionSet* /*InElementSelectionSet*/);
+	FOnChanged& OnChanged()
+	{
+		return OnChangedDelegate;
+	}
+
+	/**
+	 * Invoke the delegate called whenever the underlying element list has been changed.
+	 */
+	void NotifyPendingChanges()
+	{
+		ElementList->NotifyPendingChanges();
+	}
+
+	/**
+	 * Clear whether there are pending changes for OnChangedDelegate to notify for, without emitting a notification.
+	 */
+	void ClearPendingChanges()
+	{
+		ElementList->ClearPendingChanges();
+	}
+
+	/**
+	 * Access the interface to allow external systems (such as USelection) to receive immediate sync notifications as the underlying element list is changed.
+	 * This exists purely as a bridging mechanism and shouldn't be relied on for new code. It is lazily created as needed.
+	 */
+	FTypedElementListLegacySync& Legacy_GetElementListSync()
+	{
+		return ElementList->Legacy_GetSync();
+	}
+
+	/**
+	 * Access the interface to allow external systems (such as USelection) to receive immediate sync notifications as the underlying element list is changed.
+	 * This exists purely as a bridging mechanism and shouldn't be relied on for new code. This will return null if no legacy sync has been created for this instance.
+	 */
+	FTypedElementListLegacySync* Legacy_GetElementListSyncPtr()
+	{
+		return ElementList->Legacy_GetSyncPtr();
+	}
+
+	/**
+	 * Get the underlying element list holding the selection state.
+	 */
+	const UTypedElementList* GetElementList() const
+	{
+		return ElementList;
+	}
+
+	/**
+	 * Register an asset editor selection proxy for the given named element type.
+	 */
+	void RegisterAssetEditorSelectionProxy(const FName InElementTypeName, UTypedElementAssetEditorSelectionProxy* InAssetEditorSelectionProxy);
+
+private:
+	/**
+	 * Attempt to resolve the selection interface and asset editor selection proxy for the given element, if any.
+	 */
+	FTypedElementSelectionSetElement ResolveSelectionSetElement(const FTypedElementHandle& InElementHandle) const;
+
+	/**
+	 * Proxy the internal OnPreChange event from the underlying element list.
+	 */
+	void OnElementListPreChange(const UTypedElementList* InElementList);
+
+	/**
+	 * Proxy the internal OnChanged event from the underlying element list.
+	 */
+	void OnElementListChanged(const UTypedElementList* InElementList);
+
+	/** Underlying element list holding the selection state. */
+	UPROPERTY()
+	UTypedElementList* ElementList = nullptr;
+
+	/** Array of registered asset editor selection proxies, indexed by ElementTypeId-1. */
+	UPROPERTY()
+	UTypedElementAssetEditorSelectionProxy* RegisteredAssetEditorSelectionProxies[TypedHandleMaxTypeId - 1];
+
+	/** Delegate that is invoked whenever the underlying element list is potentially about to change. */
+	FOnPreChange OnPreChangeDelegate;
+
+	/** Delegate that is invoked whenever the underlying element list has been changed. */
+	FOnChanged OnChangedDelegate;
+
+	/** Set when we are currently restoring the selection state from a transaction (undo/redo) */
+	bool bIsRestoringFromTransaction = false;
+};
diff --git a/Engine/Source/Editor/EditorFramework/Public/Elements/TypedElementSelectionSet.h b/Engine/Source/Editor/EditorFramework/Public/Elements/TypedElementSelectionSet.h
deleted file mode 100644
diff --git a/Engine/Source/Editor/LevelEditor/Private/Elements/Actor/ActorElementLevelEditorSelectionProxy.cpp b/Engine/Source/Editor/LevelEditor/Private/Elements/Actor/ActorElementLevelEditorSelectionProxy.cpp
--- a/Engine/Source/Editor/LevelEditor/Private/Elements/Actor/ActorElementLevelEditorSelectionProxy.cpp
+++ b/Engine/Source/Editor/LevelEditor/Private/Elements/Actor/ActorElementLevelEditorSelectionProxy.cpp
@@ -4,8 +4,9 @@
 #include "Elements/Actor/ActorElementData.h"
 #include "GameFramework/Actor.h"
 #include "Components/ActorComponent.h"
-#include "TypedElementList.h"
-#include "Elements/EngineElementsLibrary.h"
+
+#include "Elements/Framework/TypedElementList.h"
+#include "Elements/Framework/EngineElementsLibrary.h"
 
 #include "Elements/Actor/ActorElementSelectionInterface.h"
 #include "Elements/Component/ComponentElementSelectionInterface.h"
diff --git a/Engine/Source/Editor/LevelEditor/Private/Elements/Actor/ActorElementLevelEditorSelectionProxy.h b/Engine/Source/Editor/LevelEditor/Private/Elements/Actor/ActorElementLevelEditorSelectionProxy.h
--- a/Engine/Source/Editor/LevelEditor/Private/Elements/Actor/ActorElementLevelEditorSelectionProxy.h
+++ b/Engine/Source/Editor/LevelEditor/Private/Elements/Actor/ActorElementLevelEditorSelectionProxy.h
@@ -2,7 +2,7 @@
 
 #pragma once
 
-#include "Elements/TypedElementSelectionSet.h"
+#include "Elements/Framework/TypedElementSelectionSet.h"
 #include "ActorElementLevelEditorSelectionProxy.generated.h"
 
 class AGroupActor;
diff --git a/Engine/Source/Editor/LevelEditor/Private/Elements/Component/ComponentElementLevelEditorSelectionProxy.cpp b/Engine/Source/Editor/LevelEditor/Private/Elements/Component/ComponentElementLevelEditorSelectionProxy.cpp
--- a/Engine/Source/Editor/LevelEditor/Private/Elements/Component/ComponentElementLevelEditorSelectionProxy.cpp
+++ b/Engine/Source/Editor/LevelEditor/Private/Elements/Component/ComponentElementLevelEditorSelectionProxy.cpp
@@ -4,8 +4,9 @@
 #include "Elements/Component/ComponentElementData.h"
 #include "GameFramework/Actor.h"
 #include "Components/ActorComponent.h"
-#include "TypedElementList.h"
-#include "Elements/EngineElementsLibrary.h"
+
+#include "Elements/Framework/TypedElementList.h"
+#include "Elements/Framework/EngineElementsLibrary.h"
 
 #include "Elements/Actor/ActorElementSelectionInterface.h"
 #include "Elements/Component/ComponentElementSelectionInterface.h"
diff --git a/Engine/Source/Editor/LevelEditor/Private/Elements/Component/ComponentElementLevelEditorSelectionProxy.h b/Engine/Source/Editor/LevelEditor/Private/Elements/Component/ComponentElementLevelEditorSelectionProxy.h
--- a/Engine/Source/Editor/LevelEditor/Private/Elements/Component/ComponentElementLevelEditorSelectionProxy.h
+++ b/Engine/Source/Editor/LevelEditor/Private/Elements/Component/ComponentElementLevelEditorSelectionProxy.h
@@ -2,7 +2,7 @@
 
 #pragma once
 
-#include "Elements/TypedElementSelectionSet.h"
+#include "Elements/Framework/TypedElementSelectionSet.h"
 #include "ComponentElementLevelEditorSelectionProxy.generated.h"
 
 UCLASS()
diff --git a/Engine/Source/Editor/LevelEditor/Private/SLevelEditor.cpp b/Engine/Source/Editor/LevelEditor/Private/SLevelEditor.cpp
--- a/Engine/Source/Editor/LevelEditor/Private/SLevelEditor.cpp
+++ b/Engine/Source/Editor/LevelEditor/Private/SLevelEditor.cpp
@@ -58,7 +58,7 @@
 #include "SourceCodeNavigation.h"
 #include "Editor/EnvironmentLightingViewer/Public/EnvironmentLightingModule.h"
 #include "Misc/MessageDialog.h"
-#include "Elements/TypedElementSelectionSet.h"
+#include "Elements/Framework/TypedElementSelectionSet.h"
 #include "Elements/Actor/ActorElementLevelEditorSelectionProxy.h"
 #include "Elements/Component/ComponentElementLevelEditorSelectionProxy.h"
 
diff --git a/Engine/Source/Editor/UnrealEd/Private/AssetEditorModeManager.cpp b/Engine/Source/Editor/UnrealEd/Private/AssetEditorModeManager.cpp
--- a/Engine/Source/Editor/UnrealEd/Private/AssetEditorModeManager.cpp
+++ b/Engine/Source/Editor/UnrealEd/Private/AssetEditorModeManager.cpp
@@ -3,7 +3,7 @@
 #include "AssetEditorModeManager.h"
 #include "Engine/Selection.h"
 #include "PreviewScene.h"
-#include "Elements/TypedElementSelectionSet.h"
+#include "Elements/Framework/TypedElementSelectionSet.h"
 
 //////////////////////////////////////////////////////////////////////////
 // FAssetEditorModeManager
diff --git a/Engine/Source/Editor/UnrealEd/Private/EditorEngine.cpp b/Engine/Source/Editor/UnrealEd/Private/EditorEngine.cpp
--- a/Engine/Source/Editor/UnrealEd/Private/EditorEngine.cpp
+++ b/Engine/Source/Editor/UnrealEd/Private/EditorEngine.cpp
@@ -129,7 +129,7 @@
 
 #include "Editor/ActorPositioning.h"
 
-#include "Elements/TypedElementSelectionSet.h"
+#include "Elements/Framework/TypedElementSelectionSet.h"
 #include "Elements/Object/ObjectElementEditorSelectionProxy.h"
 
 #include "IDirectoryWatcher.h"
diff --git a/Engine/Source/Editor/UnrealEd/Private/EditorSelectUtils.cpp b/Engine/Source/Editor/UnrealEd/Private/EditorSelectUtils.cpp
--- a/Engine/Source/Editor/UnrealEd/Private/EditorSelectUtils.cpp
+++ b/Engine/Source/Editor/UnrealEd/Private/EditorSelectUtils.cpp
@@ -27,9 +27,9 @@
 #include "Logging/MessageLog.h"
 #include "ActorGroupingUtils.h"
 #include "Subsystems/BrushEditingSubsystem.h"
-#include "TypedElementList.h"
-#include "Elements/TypedElementSelectionSet.h"
-#include "Elements/EngineElementsLibrary.h"
+#include "Elements/Framework/TypedElementList.h"
+#include "Elements/Framework/TypedElementSelectionSet.h"
+#include "Elements/Framework/EngineElementsLibrary.h"
 
 #define LOCTEXT_NAMESPACE "EditorSelectUtils"
 
diff --git a/Engine/Source/Editor/UnrealEd/Private/Elements/Actor/ActorElementEditorSelectionInterface.cpp b/Engine/Source/Editor/UnrealEd/Private/Elements/Actor/ActorElementEditorSelectionInterface.cpp
--- a/Engine/Source/Editor/UnrealEd/Private/Elements/Actor/ActorElementEditorSelectionInterface.cpp
+++ b/Engine/Source/Editor/UnrealEd/Private/Elements/Actor/ActorElementEditorSelectionInterface.cpp
@@ -4,8 +4,8 @@
 #include "Elements/Actor/ActorElementData.h"
 #include "GameFramework/Actor.h"
 
-#include "TypedElementList.h"
-#include "Elements/EngineElementsLibrary.h"
+#include "Elements/Framework/TypedElementList.h"
+#include "Elements/Framework/EngineElementsLibrary.h"
 #include "Elements/Object/ObjectElementEditorSelectionInterface.h"
 
 bool UActorElementEditorSelectionInterface::IsElementSelected(const FTypedElementHandle& InElementHandle, const UTypedElementList* InSelectionSet, const FTypedElementIsSelectedOptions& InSelectionOptions)
diff --git a/Engine/Source/Editor/UnrealEd/Private/Elements/Component/ComponentElementEditorSelectionInterface.cpp b/Engine/Source/Editor/UnrealEd/Private/Elements/Component/ComponentElementEditorSelectionInterface.cpp
--- a/Engine/Source/Editor/UnrealEd/Private/Elements/Component/ComponentElementEditorSelectionInterface.cpp
+++ b/Engine/Source/Editor/UnrealEd/Private/Elements/Component/ComponentElementEditorSelectionInterface.cpp
@@ -5,8 +5,8 @@
 #include "Components/SceneComponent.h"
 #include "GameFramework/Actor.h"
 
-#include "TypedElementList.h"
-#include "Elements/EngineElementsLibrary.h"
+#include "Elements/Framework/TypedElementList.h"
+#include "Elements/Framework/EngineElementsLibrary.h"
 #include "Elements/Object/ObjectElementEditorSelectionInterface.h"
 
 bool UComponentElementEditorSelectionInterface::IsElementSelected(const FTypedElementHandle& InElementHandle, const UTypedElementList* InSelectionSet, const FTypedElementIsSelectedOptions& InSelectionOptions)
diff --git a/Engine/Source/Editor/UnrealEd/Private/Elements/EditorElements.cpp b/Engine/Source/Editor/UnrealEd/Private/Elements/EditorElements.cpp
deleted file mode 100644
diff --git a/Engine/Source/Editor/UnrealEd/Private/Elements/Framework/EditorElements.cpp b/Engine/Source/Editor/UnrealEd/Private/Elements/Framework/EditorElements.cpp
new file mode 100644
--- /dev/null
+++ b/Engine/Source/Editor/UnrealEd/Private/Elements/Framework/EditorElements.cpp
@@ -0,0 +1,42 @@
+// Copyright Epic Games, Inc. All Rights Reserved.
+
+#include "Elements/Framework/EditorElements.h"
+#include "Elements/Framework/TypedElementRegistry.h"
+
+#include "Elements/Object/ObjectElementEditorSelectionInterface.h"
+
+#include "Elements/Actor/ActorElementEditorSelectionInterface.h"
+
+#include "Elements/Component/ComponentElementEditorSelectionInterface.h"
+
+FSimpleMulticastDelegate OnRegisterEditorElementsDelegate;
+
+void RegisterEditorObjectElements()
+{
+	UTypedElementRegistry* Registry = UTypedElementRegistry::GetInstance();
+
+	Registry->RegisterElementInterface<UTypedElementSelectionInterface>(NAME_Object, NewObject<UObjectElementEditorSelectionInterface>(), /*bAllowOverride*/true);
+}
+
+void RegisterEditorActorElements()
+{
+	UTypedElementRegistry* Registry = UTypedElementRegistry::GetInstance();
+
+	Registry->RegisterElementInterface<UTypedElementSelectionInterface>(NAME_Actor, NewObject<UActorElementEditorSelectionInterface>(), /*bAllowOverride*/true);
+}
+
+void RegisterEditorComponentElements()
+{
+	UTypedElementRegistry* Registry = UTypedElementRegistry::GetInstance();
+
+	Registry->RegisterElementInterface<UTypedElementSelectionInterface>(NAME_Components, NewObject<UComponentElementEditorSelectionInterface>(), /*bAllowOverride*/true);
+}
+
+void RegisterEditorElements()
+{
+	RegisterEditorObjectElements();
+	RegisterEditorActorElements();
+	RegisterEditorComponentElements();
+
+	OnRegisterEditorElementsDelegate.Broadcast();
+}
diff --git a/Engine/Source/Editor/UnrealEd/Private/Elements/Object/ObjectElementEditorSelectionInterface.cpp b/Engine/Source/Editor/UnrealEd/Private/Elements/Object/ObjectElementEditorSelectionInterface.cpp
--- a/Engine/Source/Editor/UnrealEd/Private/Elements/Object/ObjectElementEditorSelectionInterface.cpp
+++ b/Engine/Source/Editor/UnrealEd/Private/Elements/Object/ObjectElementEditorSelectionInterface.cpp
@@ -4,7 +4,7 @@
 #include "Elements/Object/ObjectElementData.h"
 
 #include "UObject/Package.h"
-#include "Elements/EngineElementsLibrary.h"
+#include "Elements/Framework/EngineElementsLibrary.h"
 
 bool UObjectElementEditorSelectionInterface::ShouldPreventTransactions(const FTypedElementHandle& InElementHandle)
 {
diff --git a/Engine/Source/Editor/UnrealEd/Private/Elements/Object/ObjectElementEditorSelectionProxy.h b/Engine/Source/Editor/UnrealEd/Private/Elements/Object/ObjectElementEditorSelectionProxy.h
--- a/Engine/Source/Editor/UnrealEd/Private/Elements/Object/ObjectElementEditorSelectionProxy.h
+++ b/Engine/Source/Editor/UnrealEd/Private/Elements/Object/ObjectElementEditorSelectionProxy.h
@@ -2,7 +2,7 @@
 
 #pragma once
 
-#include "Elements/TypedElementSelectionSet.h"
+#include "Elements/Framework/TypedElementSelectionSet.h"
 #include "ObjectElementEditorSelectionProxy.generated.h"
 
 UCLASS()
diff --git a/Engine/Source/Editor/UnrealEd/Private/LevelViewportClickHandlers.cpp b/Engine/Source/Editor/UnrealEd/Private/LevelViewportClickHandlers.cpp
--- a/Engine/Source/Editor/UnrealEd/Private/LevelViewportClickHandlers.cpp
+++ b/Engine/Source/Editor/UnrealEd/Private/LevelViewportClickHandlers.cpp
@@ -30,7 +30,7 @@
 #include "StaticLightingSystem/StaticLightingPrivate.h"
 #include "LightMap.h"
 #include "Subsystems/BrushEditingSubsystem.h"
-#include "Elements/TypedElementSelectionSet.h"
+#include "Elements/Framework/TypedElementSelectionSet.h"
 
 #define LOCTEXT_NAMESPACE "ClickHandlers"
 
diff --git a/Engine/Source/Editor/UnrealEd/Private/Selection.cpp b/Engine/Source/Editor/UnrealEd/Private/Selection.cpp
--- a/Engine/Source/Editor/UnrealEd/Private/Selection.cpp
+++ b/Engine/Source/Editor/UnrealEd/Private/Selection.cpp
@@ -3,8 +3,8 @@
 #include "Selection.h"
 #include "UObject/Package.h"
 #include "GameFramework/Actor.h"
-#include "TypedElementObjectInterface.h"
-#include "Elements/EngineElementsLibrary.h"
+#include "Elements/Framework/EngineElementsLibrary.h"
+#include "Elements/Interfaces/TypedElementObjectInterface.h"
 
 DEFINE_LOG_CATEGORY_STATIC(LogSelection, Log, All);
 
diff --git a/Engine/Source/Editor/UnrealEd/Private/UnrealEdEngine.cpp b/Engine/Source/Editor/UnrealEd/Private/UnrealEdEngine.cpp
--- a/Engine/Source/Editor/UnrealEd/Private/UnrealEdEngine.cpp
+++ b/Engine/Source/Editor/UnrealEd/Private/UnrealEdEngine.cpp
@@ -32,7 +32,7 @@
 #include "EditorModeRegistry.h"
 #include "EditorModeManager.h"
 #include "EditorModes.h"
-#include "Elements/EditorElements.h"
+#include "Elements/Framework/EditorElements.h"
 #include "UnrealEdMisc.h"
 #include "UnrealEdGlobals.h"
 
diff --git a/Engine/Source/Editor/UnrealEd/Public/Elements/EditorElements.h b/Engine/Source/Editor/UnrealEd/Public/Elements/EditorElements.h
deleted file mode 100644
diff --git a/Engine/Source/Editor/UnrealEd/Public/Elements/Framework/EditorElements.h b/Engine/Source/Editor/UnrealEd/Public/Elements/Framework/EditorElements.h
new file mode 100644
--- /dev/null
+++ b/Engine/Source/Editor/UnrealEd/Public/Elements/Framework/EditorElements.h
@@ -0,0 +1,9 @@
+// Copyright Epic Games, Inc. All Rights Reserved.
+
+#pragma once
+
+#include "Delegates/Delegate.h"
+
+extern UNREALED_API FSimpleMulticastDelegate OnRegisterEditorElementsDelegate;
+
+void RegisterEditorElements();
diff --git a/Engine/Source/Editor/UnrealEd/Public/Selection.h b/Engine/Source/Editor/UnrealEd/Public/Selection.h
--- a/Engine/Source/Editor/UnrealEd/Public/Selection.h
+++ b/Engine/Source/Editor/UnrealEd/Public/Selection.h
@@ -7,7 +7,7 @@
 #include "UObject/UObjectGlobals.h"
 #include "UObject/Object.h"
 #include "Components/ActorComponent.h"
-#include "Elements/TypedElementSelectionSet.h"
+#include "Elements/Framework/TypedElementSelectionSet.h"
 #include "Selection.generated.h"
 
 class ISelectionElementBridge;
diff --git a/Engine/Source/Runtime/Engine/Private/Actor.cpp b/Engine/Source/Runtime/Engine/Private/Actor.cpp
--- a/Engine/Source/Runtime/Engine/Private/Actor.cpp
+++ b/Engine/Source/Runtime/Engine/Private/Actor.cpp
@@ -4,7 +4,7 @@
 #include "EngineDefines.h"
 #include "EngineStats.h"
 #include "EngineGlobals.h"
-#include "Elements/EngineElementsLibrary.h"
+#include "Elements/Framework/EngineElementsLibrary.h"
 #include "GameFramework/DamageType.h"
 #include "TimerManager.h"
 #include "GameFramework/Pawn.h"
diff --git a/Engine/Source/Runtime/Engine/Private/Components/ActorComponent.cpp b/Engine/Source/Runtime/Engine/Private/Components/ActorComponent.cpp
--- a/Engine/Source/Runtime/Engine/Private/Components/ActorComponent.cpp
+++ b/Engine/Source/Runtime/Engine/Private/Components/ActorComponent.cpp
@@ -13,7 +13,7 @@
 #include "Components/PrimitiveComponent.h"
 #include "AI/NavigationSystemBase.h"
 #include "Engine/BlueprintGeneratedClass.h"
-#include "Elements/EngineElementsLibrary.h"
+#include "Elements/Framework/EngineElementsLibrary.h"
 #include "ContentStreaming.h"
 #include "ComponentReregisterContext.h"
 #include "Engine/AssetUserData.h"
diff --git a/Engine/Source/Runtime/Engine/Private/Elements/Actor/ActorElementSelectionInterface.cpp b/Engine/Source/Runtime/Engine/Private/Elements/Actor/ActorElementSelectionInterface.cpp
--- a/Engine/Source/Runtime/Engine/Private/Elements/Actor/ActorElementSelectionInterface.cpp
+++ b/Engine/Source/Runtime/Engine/Private/Elements/Actor/ActorElementSelectionInterface.cpp
@@ -3,7 +3,7 @@
 #include "Elements/Actor/ActorElementSelectionInterface.h"
 #include "Elements/Actor/ActorElementData.h"
 #include "GameFramework/Actor.h"
-#include "TypedElementList.h"
+#include "Elements/Framework/TypedElementList.h"
 
 int32 UActorElementSelectionInterface::GetNumSelectedActors(const UTypedElementList* InCurrentSelection)
 {
diff --git a/Engine/Source/Runtime/Engine/Private/Elements/Component/ComponentElementSelectionInterface.cpp b/Engine/Source/Runtime/Engine/Private/Elements/Component/ComponentElementSelectionInterface.cpp
--- a/Engine/Source/Runtime/Engine/Private/Elements/Component/ComponentElementSelectionInterface.cpp
+++ b/Engine/Source/Runtime/Engine/Private/Elements/Component/ComponentElementSelectionInterface.cpp
@@ -3,7 +3,7 @@
 #include "Elements/Component/ComponentElementSelectionInterface.h"
 #include "Elements/Component/ComponentElementData.h"
 #include "Components/ActorComponent.h"
-#include "TypedElementList.h"
+#include "Elements/Framework/TypedElementList.h"
 
 int32 UComponentElementSelectionInterface::GetNumSelectedComponents(const UTypedElementList* InCurrentSelection)
 {
diff --git a/Engine/Source/Runtime/Engine/Private/Elements/EngineElements.cpp b/Engine/Source/Runtime/Engine/Private/Elements/EngineElements.cpp
deleted file mode 100644
diff --git a/Engine/Source/Runtime/Engine/Private/Elements/EngineElementsLibrary.cpp b/Engine/Source/Runtime/Engine/Private/Elements/EngineElementsLibrary.cpp
deleted file mode 100644
diff --git a/Engine/Source/Runtime/Engine/Private/Elements/Framework/EngineElements.cpp b/Engine/Source/Runtime/Engine/Private/Elements/Framework/EngineElements.cpp
new file mode 100644
--- /dev/null
+++ b/Engine/Source/Runtime/Engine/Private/Elements/Framework/EngineElements.cpp
@@ -0,0 +1,60 @@
+// Copyright Epic Games, Inc. All Rights Reserved.
+
+#include "Elements/Framework/EngineElements.h"
+#include "Elements/Framework/TypedElementRegistry.h"
+
+#include "Elements/Object/ObjectElementData.h"
+#include "Elements/Object/ObjectElementObjectInterface.h"
+#include "Elements/Object/ObjectElementSelectionInterface.h"
+
+#include "Elements/Actor/ActorElementData.h"
+#include "Elements/Actor/ActorElementObjectInterface.h"
+#include "Elements/Actor/ActorElementSelectionInterface.h"
+
+#include "Elements/Component/ComponentElementData.h"
+#include "Elements/Component/ComponentElementObjectInterface.h"
+#include "Elements/Component/ComponentElementSelectionInterface.h"
+
+#include "Modules/ModuleManager.h"
+
+FSimpleMulticastDelegate OnRegisterEngineElementsDelegate;
+
+void RegisterEngineObjectElements()
+{
+	UTypedElementRegistry* Registry = UTypedElementRegistry::GetInstance();
+
+	Registry->RegisterElementType<FObjectElementData>(NAME_Object);
+	Registry->RegisterElementInterface<UTypedElementObjectInterface>(NAME_Object, NewObject<UObjectElementObjectInterface>());
+	Registry->RegisterElementInterface<UTypedElementSelectionInterface>(NAME_Object, NewObject<UObjectElementSelectionInterface>());
+}
+
+void RegisterEngineActorElements()
+{
+	UTypedElementRegistry* Registry = UTypedElementRegistry::GetInstance();
+
+	Registry->RegisterElementType<FActorElementData>(NAME_Actor);
+	Registry->RegisterElementInterface<UTypedElementObjectInterface>(NAME_Actor, NewObject<UActorElementObjectInterface>());
+	Registry->RegisterElementInterface<UTypedElementSelectionInterface>(NAME_Actor, NewObject<UActorElementSelectionInterface>());
+}
+
+void RegisterEngineComponentElements()
+{
+	UTypedElementRegistry* Registry = UTypedElementRegistry::GetInstance();
+
+	Registry->RegisterElementType<FComponentElementData>(NAME_Components);
+	Registry->RegisterElementInterface<UTypedElementObjectInterface>(NAME_Components, NewObject<UComponentElementObjectInterface>());
+	Registry->RegisterElementInterface<UTypedElementSelectionInterface>(NAME_Components, NewObject<UComponentElementSelectionInterface>());
+}
+
+void RegisterEngineElements()
+{
+	// Ensure the framework and base interfaces are also loaded
+	FModuleManager::Get().LoadModuleChecked("TypedElementFramework");
+	FModuleManager::Get().LoadModuleChecked("TypedElementInterfaces");
+
+	RegisterEngineObjectElements();
+	RegisterEngineActorElements();
+	RegisterEngineComponentElements();
+
+	OnRegisterEngineElementsDelegate.Broadcast();
+}
diff --git a/Engine/Source/Runtime/Engine/Private/Elements/Framework/EngineElementsLibrary.cpp b/Engine/Source/Runtime/Engine/Private/Elements/Framework/EngineElementsLibrary.cpp
new file mode 100644
--- /dev/null
+++ b/Engine/Source/Runtime/Engine/Private/Elements/Framework/EngineElementsLibrary.cpp
@@ -0,0 +1,209 @@
+// Copyright Epic Games, Inc. All Rights Reserved.
+
+#include "Elements/Framework/EngineElementsLibrary.h"
+#include "Elements/Framework/TypedElementRegistry.h"
+#include "Elements/Framework/TypedElementOwnerStore.h"
+
+#include "Elements/Object/ObjectElementData.h"
+#include "UObject/Object.h"
+
+#include "Elements/Actor/ActorElementData.h"
+#include "GameFramework/Actor.h"
+
+#include "Elements/Component/ComponentElementData.h"
+#include "Components/ActorComponent.h"
+
+namespace EngineElementsLibraryUtil
+{
+template <typename ObjectClass, typename ElementDataType>
+TTypedElementOwner<ElementDataType> CreateTypedElement(const ObjectClass* Object, const FName ElementTypeName)
+{
+	UTypedElementRegistry* Registry = UTypedElementRegistry::GetInstance();
+
+	TTypedElementOwner<ElementDataType> TypedElement;
+	if (ensureMsgf(Registry, TEXT("Typed element was requested for '%s' before the registry was available! This usually means that NewObject was used instead of CreateDefaultSubobject during CDO construction."), *Object->GetPathName()))
+	{
+		TypedElement = Registry->CreateElement<ElementDataType>(ElementTypeName);
+	}
+
+	return TypedElement;
+}
+
+#if WITH_EDITOR
+template <typename ObjectClass, typename ElementDataType>
+void CreateEditorTypedElement(const ObjectClass* Object, TTypedElementOwnerStore<ElementDataType, const ObjectClass*>& ElementOwnerStore, TFunctionRef<TTypedElementOwner<ElementDataType>(const ObjectClass*)> CreateElement)
+{
+	if (GIsEditor && !Object->HasAnyFlags(RF_ClassDefaultObject | RF_ArchetypeObject))
+	{
+		ElementOwnerStore.RegisterElementOwner(Object, CreateElement(Object));
+	}
+}
+
+template <typename ObjectClass, typename ElementDataType>
+void DestroyEditorTypedElement(const ObjectClass* Object, TTypedElementOwnerStore<ElementDataType, const ObjectClass*>& ElementOwnerStore, TFunctionRef<void(const ObjectClass*, TTypedElementOwner<ElementDataType>&)> DestroyElement)
+{
+	if (TTypedElementOwner<ElementDataType> EditorElement = ElementOwnerStore.UnregisterElementOwner(Object))
+	{
+		DestroyElement(Object, EditorElement);
+	}
+}
+
+template <typename ObjectClass, typename ElementDataType>
+FTypedElementHandle AcquireEditorTypedElementHandle(const ObjectClass* Object, TTypedElementOwnerStore<ElementDataType, const ObjectClass*>& ElementOwnerStore, TFunctionRef<TTypedElementOwner<ElementDataType>(const ObjectClass*)> CreateElement, const bool bAllowCreate)
+{
+	if (GIsEditor)
+	{
+		TTypedElementOwnerScopedAccess<ElementDataType> EditorElement = bAllowCreate
+			? ElementOwnerStore.FindOrRegisterElementOwner(Object, [Object, &CreateElement](){ return CreateElement(Object); })
+			: ElementOwnerStore.FindElementOwner(Object);
+		if (EditorElement)
+		{
+			return EditorElement->AcquireHandle();
+		}
+	}
+
+	return FTypedElementHandle();
+}
+#endif
+} // namespace EngineElementsLibraryUtil
+
+#if WITH_EDITOR
+TTypedElementOwnerStore<FObjectElementData, const UObject*> GObjectElementOwnerStore;
+TTypedElementOwnerStore<FActorElementData, const AActor*> GActorElementOwnerStore;
+TTypedElementOwnerStore<FComponentElementData, const UActorComponent*> GComponentElementOwnerStore;
+#endif
+
+UEngineElementsLibrary::UEngineElementsLibrary()
+{
+#if WITH_EDITOR
+	if (HasAnyFlags(RF_ClassDefaultObject))
+	{
+		// UObject exists inside CoreUObject, so it cannot call through directly to clean-up any element handles that have been created
+		// Instead we rely on this GC hook to clean-up any element handles for unreachable objects prior to them being destroyed
+		FCoreUObjectDelegates::PreGarbageCollectConditionalBeginDestroy.AddStatic(&UEngineElementsLibrary::DestroyUnreachableEditorObjectElements);
+	}
+#endif
+}
+
+TTypedElementOwner<FObjectElementData> UEngineElementsLibrary::CreateObjectElement(const UObject* InObject)
+{
+	TTypedElementOwner<FObjectElementData> ObjectElement = EngineElementsLibraryUtil::CreateTypedElement<UObject, FObjectElementData>(InObject, NAME_Object);
+	if (ObjectElement)
+	{
+		ObjectElement.GetDataChecked().Object = const_cast<UObject*>(InObject);
+	}
+	return ObjectElement;
+}
+
+void UEngineElementsLibrary::DestroyObjectElement(const UObject* InObject, TTypedElementOwner<FObjectElementData>& InOutObjectElement)
+{
+	if (InOutObjectElement)
+	{
+		checkf(InOutObjectElement.GetDataChecked().Object == InObject, TEXT("Object element was not for this object instance! %s"), *InObject->GetPathName());
+		UTypedElementRegistry::GetInstance()->DestroyElement(InOutObjectElement);
+	}
+}
+
+#if WITH_EDITOR
+void UEngineElementsLibrary::CreateEditorObjectElement(const UObject* Object)
+{
+	EngineElementsLibraryUtil::CreateEditorTypedElement<UObject, FObjectElementData>(Object, GObjectElementOwnerStore, &UEngineElementsLibrary::CreateObjectElement);
+}
+
+void UEngineElementsLibrary::DestroyEditorObjectElement(const UObject* Object)
+{
+	EngineElementsLibraryUtil::DestroyEditorTypedElement<UObject, FObjectElementData>(Object, GObjectElementOwnerStore, &UEngineElementsLibrary::DestroyObjectElement);
+}
+
+void UEngineElementsLibrary::DestroyUnreachableEditorObjectElements()
+{
+	auto IsObjectElementUnreachable = [](const TTypedElementOwner<FObjectElementData>& EditorElement)
+	{
+		return EditorElement.GetDataChecked().Object->IsUnreachable();
+	};
+
+	auto DestroyUnreachableObjectElement = [Registry = UTypedElementRegistry::GetInstance()](TTypedElementOwner<FObjectElementData>&& EditorElement)
+	{
+		Registry->DestroyElement(EditorElement);
+	};
+
+	GObjectElementOwnerStore.UnregisterElementOwners(IsObjectElementUnreachable, DestroyUnreachableObjectElement);
+}
+
+FTypedElementHandle UEngineElementsLibrary::AcquireEditorObjectElementHandle(const UObject* Object, const bool bAllowCreate)
+{
+	return EngineElementsLibraryUtil::AcquireEditorTypedElementHandle<UObject, FObjectElementData>(Object, GObjectElementOwnerStore, &UEngineElementsLibrary::CreateObjectElement, bAllowCreate);
+}
+#endif
+
+TTypedElementOwner<FActorElementData> UEngineElementsLibrary::CreateActorElement(const AActor* InActor)
+{
+	TTypedElementOwner<FActorElementData> ActorElement = EngineElementsLibraryUtil::CreateTypedElement<AActor, FActorElementData>(InActor, NAME_Actor);
+	if (ActorElement)
+	{
+		ActorElement.GetDataChecked().Actor = const_cast<AActor*>(InActor);
+	}
+	return ActorElement;
+}
+
+void UEngineElementsLibrary::DestroyActorElement(const AActor* InActor, TTypedElementOwner<FActorElementData>& InOutActorElement)
+{
+	if (InOutActorElement)
+	{
+		checkf(InOutActorElement.GetDataChecked().Actor == InActor, TEXT("Actor element was not for this actor instance! %s"), *InActor->GetPathName());
+		UTypedElementRegistry::GetInstance()->DestroyElement(InOutActorElement);
+	}
+}
+
+#if WITH_EDITOR
+void UEngineElementsLibrary::CreateEditorActorElement(const AActor* Actor)
+{
+	EngineElementsLibraryUtil::CreateEditorTypedElement<AActor, FActorElementData>(Actor, GActorElementOwnerStore, &UEngineElementsLibrary::CreateActorElement);
+}
+
+void UEngineElementsLibrary::DestroyEditorActorElement(const AActor* Actor)
+{
+	EngineElementsLibraryUtil::DestroyEditorTypedElement<AActor, FActorElementData>(Actor, GActorElementOwnerStore, &UEngineElementsLibrary::DestroyActorElement);
+}
+
+FTypedElementHandle UEngineElementsLibrary::AcquireEditorActorElementHandle(const AActor* Actor, const bool bAllowCreate)
+{
+	return EngineElementsLibraryUtil::AcquireEditorTypedElementHandle<AActor, FActorElementData>(Actor, GActorElementOwnerStore, &UEngineElementsLibrary::CreateActorElement, bAllowCreate);
+}
+#endif
+
+TTypedElementOwner<FComponentElementData> UEngineElementsLibrary::CreateComponentElement(const UActorComponent* InComponent)
+{
+	TTypedElementOwner<FComponentElementData> ComponentElement = EngineElementsLibraryUtil::CreateTypedElement<UActorComponent, FComponentElementData>(InComponent, NAME_Components);
+	if (ComponentElement)
+	{
+		ComponentElement.GetDataChecked().Component = const_cast<UActorComponent*>(InComponent);
+	}
+	return ComponentElement;
+}
+
+void UEngineElementsLibrary::DestroyComponentElement(const UActorComponent* InComponent, TTypedElementOwner<FComponentElementData>& InOutComponentElement)
+{
+	if (InOutComponentElement)
+	{
+		checkf(InOutComponentElement.GetDataChecked().Component == InComponent, TEXT("Component element was not for this component instance! %s"), *InComponent->GetPathName());
+		UTypedElementRegistry::GetInstance()->DestroyElement(InOutComponentElement);
+	}
+}
+
+#if WITH_EDITOR
+void UEngineElementsLibrary::CreateEditorComponentElement(const UActorComponent* Component)
+{
+	EngineElementsLibraryUtil::CreateEditorTypedElement<UActorComponent, FComponentElementData>(Component, GComponentElementOwnerStore, &UEngineElementsLibrary::CreateComponentElement);
+}
+
+void UEngineElementsLibrary::DestroyEditorComponentElement(const UActorComponent* Component)
+{
+	EngineElementsLibraryUtil::DestroyEditorTypedElement<UActorComponent, FComponentElementData>(Component, GComponentElementOwnerStore, &UEngineElementsLibrary::DestroyComponentElement);
+}
+
+FTypedElementHandle UEngineElementsLibrary::AcquireEditorComponentElementHandle(const UActorComponent* Component, const bool bAllowCreate)
+{
+	return EngineElementsLibraryUtil::AcquireEditorTypedElementHandle<UActorComponent, FComponentElementData>(Component, GComponentElementOwnerStore, &UEngineElementsLibrary::CreateComponentElement, bAllowCreate);
+}
+#endif
diff --git a/Engine/Source/Runtime/Engine/Private/EngineUtils.cpp b/Engine/Source/Runtime/Engine/Private/EngineUtils.cpp
--- a/Engine/Source/Runtime/Engine/Private/EngineUtils.cpp
+++ b/Engine/Source/Runtime/Engine/Private/EngineUtils.cpp
@@ -11,7 +11,7 @@
 #include "Misc/EngineVersion.h"
 #include "GameFramework/PlayerController.h"
 #include "GameFramework/LightWeightInstanceSubsystem.h"
-#include "Elements/EngineElementsLibrary.h"
+#include "Elements/Framework/EngineElementsLibrary.h"
 #include "EngineGlobals.h"
 #include "Components/PrimitiveComponent.h"
 #include "Engine/Engine.h"
diff --git a/Engine/Source/Runtime/Engine/Private/UnrealClient.cpp b/Engine/Source/Runtime/Engine/Private/UnrealClient.cpp
--- a/Engine/Source/Runtime/Engine/Private/UnrealClient.cpp
+++ b/Engine/Source/Runtime/Engine/Private/UnrealClient.cpp
@@ -26,7 +26,7 @@
 #include "EngineModule.h"
 #include "Performance/EnginePerformanceTargets.h"
 #include "Templates/UniquePtr.h"
-#include "TypedElementList.h"
+#include "Elements/Framework/TypedElementList.h"
 #include "EngineUtils.h"
 
 DEFINE_LOG_CATEGORY_STATIC(LogClient, Log, All);
diff --git a/Engine/Source/Runtime/Engine/Private/UnrealEngine.cpp b/Engine/Source/Runtime/Engine/Private/UnrealEngine.cpp
--- a/Engine/Source/Runtime/Engine/Private/UnrealEngine.cpp
+++ b/Engine/Source/Runtime/Engine/Private/UnrealEngine.cpp
@@ -105,7 +105,7 @@ UnrealEngine.cpp: Implements the UEngine class and helpers.
 #include "IImageWrapper.h"
 #include "IImageWrapperModule.h"
 #include "EngineAnalytics.h"
-#include "Elements/EngineElements.h"
+#include "Elements/Framework/EngineElements.h"
 #include "TickTaskManagerInterface.h"
 #include "Net/NetworkProfiler.h"
 #include "ProfilingDebugging/MallocProfiler.h"
diff --git a/Engine/Source/Runtime/Engine/Public/Elements/Actor/ActorElementData.h b/Engine/Source/Runtime/Engine/Public/Elements/Actor/ActorElementData.h
--- a/Engine/Source/Runtime/Engine/Public/Elements/Actor/ActorElementData.h
+++ b/Engine/Source/Runtime/Engine/Public/Elements/Actor/ActorElementData.h
@@ -2,7 +2,7 @@
 
 #pragma once
 
-#include "TypedElementData.h"
+#include "Elements/Framework/TypedElementData.h"
 
 class AActor;
 
diff --git a/Engine/Source/Runtime/Engine/Public/Elements/Actor/ActorElementObjectInterface.h b/Engine/Source/Runtime/Engine/Public/Elements/Actor/ActorElementObjectInterface.h
--- a/Engine/Source/Runtime/Engine/Public/Elements/Actor/ActorElementObjectInterface.h
+++ b/Engine/Source/Runtime/Engine/Public/Elements/Actor/ActorElementObjectInterface.h
@@ -2,7 +2,7 @@
 
 #pragma once
 
-#include "TypedElementObjectInterface.h"
+#include "Elements/Interfaces/TypedElementObjectInterface.h"
 #include "ActorElementObjectInterface.generated.h"
 
 UCLASS()
diff --git a/Engine/Source/Runtime/Engine/Public/Elements/Actor/ActorElementSelectionInterface.h b/Engine/Source/Runtime/Engine/Public/Elements/Actor/ActorElementSelectionInterface.h
--- a/Engine/Source/Runtime/Engine/Public/Elements/Actor/ActorElementSelectionInterface.h
+++ b/Engine/Source/Runtime/Engine/Public/Elements/Actor/ActorElementSelectionInterface.h
@@ -2,7 +2,7 @@
 
 #pragma once
 
-#include "TypedElementSelectionInterface.h"
+#include "Elements/Interfaces/TypedElementSelectionInterface.h"
 #include "ActorElementSelectionInterface.generated.h"
 
 UCLASS()
diff --git a/Engine/Source/Runtime/Engine/Public/Elements/Component/ComponentElementData.h b/Engine/Source/Runtime/Engine/Public/Elements/Component/ComponentElementData.h
--- a/Engine/Source/Runtime/Engine/Public/Elements/Component/ComponentElementData.h
+++ b/Engine/Source/Runtime/Engine/Public/Elements/Component/ComponentElementData.h
@@ -2,7 +2,7 @@
 
 #pragma once
 
-#include "TypedElementData.h"
+#include "Elements/Framework/TypedElementData.h"
 
 class UActorComponent;
 
diff --git a/Engine/Source/Runtime/Engine/Public/Elements/Component/ComponentElementObjectInterface.h b/Engine/Source/Runtime/Engine/Public/Elements/Component/ComponentElementObjectInterface.h
--- a/Engine/Source/Runtime/Engine/Public/Elements/Component/ComponentElementObjectInterface.h
+++ b/Engine/Source/Runtime/Engine/Public/Elements/Component/ComponentElementObjectInterface.h
@@ -2,7 +2,7 @@
 
 #pragma once
 
-#include "TypedElementObjectInterface.h"
+#include "Elements/Interfaces/TypedElementObjectInterface.h"
 #include "ComponentElementObjectInterface.generated.h"
 
 UCLASS()
diff --git a/Engine/Source/Runtime/Engine/Public/Elements/Component/ComponentElementSelectionInterface.h b/Engine/Source/Runtime/Engine/Public/Elements/Component/ComponentElementSelectionInterface.h
--- a/Engine/Source/Runtime/Engine/Public/Elements/Component/ComponentElementSelectionInterface.h
+++ b/Engine/Source/Runtime/Engine/Public/Elements/Component/ComponentElementSelectionInterface.h
@@ -2,7 +2,7 @@
 
 #pragma once
 
-#include "TypedElementSelectionInterface.h"
+#include "Elements/Interfaces/TypedElementSelectionInterface.h"
 #include "ComponentElementSelectionInterface.generated.h"
 
 UCLASS()
diff --git a/Engine/Source/Runtime/Engine/Public/Elements/EngineElements.h b/Engine/Source/Runtime/Engine/Public/Elements/EngineElements.h
deleted file mode 100644
diff --git a/Engine/Source/Runtime/Engine/Public/Elements/EngineElementsLibrary.h b/Engine/Source/Runtime/Engine/Public/Elements/EngineElementsLibrary.h
deleted file mode 100644
diff --git a/Engine/Source/Runtime/Engine/Public/Elements/Framework/EngineElements.h b/Engine/Source/Runtime/Engine/Public/Elements/Framework/EngineElements.h
new file mode 100644
--- /dev/null
+++ b/Engine/Source/Runtime/Engine/Public/Elements/Framework/EngineElements.h
@@ -0,0 +1,9 @@
+// Copyright Epic Games, Inc. All Rights Reserved.
+
+#pragma once
+
+#include "Delegates/Delegate.h"
+
+extern ENGINE_API FSimpleMulticastDelegate OnRegisterEngineElementsDelegate;
+
+void RegisterEngineElements();
diff --git a/Engine/Source/Runtime/Engine/Public/Elements/Framework/EngineElementsLibrary.h b/Engine/Source/Runtime/Engine/Public/Elements/Framework/EngineElementsLibrary.h
new file mode 100644
--- /dev/null
+++ b/Engine/Source/Runtime/Engine/Public/Elements/Framework/EngineElementsLibrary.h
@@ -0,0 +1,57 @@
+// Copyright Epic Games, Inc. All Rights Reserved.
+
+#pragma once
+
+#include "Elements/Framework/TypedElementHandle.h"
+#include "Kismet/BlueprintFunctionLibrary.h"
+#include "EngineElementsLibrary.generated.h"
+
+class UObject;
+struct FObjectElementData;
+
+class AActor;
+struct FActorElementData;
+
+class UActorComponent;
+struct FComponentElementData;
+
+UCLASS()
+class ENGINE_API UEngineElementsLibrary : public UBlueprintFunctionLibrary
+{
+	GENERATED_BODY()
+
+public:
+	UEngineElementsLibrary();
+
+	static TTypedElementOwner<FObjectElementData> CreateObjectElement(const UObject* InObject);
+	static void DestroyObjectElement(const UObject* InObject, TTypedElementOwner<FObjectElementData>& InOutObjectElement);
+#if WITH_EDITOR
+	static void CreateEditorObjectElement(const UObject* Object);
+	static void DestroyEditorObjectElement(const UObject* Object);
+	UFUNCTION(BlueprintCallable, BlueprintPure=false, Category="TypedElementFramework|Object", meta=(ScriptMethod="AcquireEditorElementHandle"))
+	static FTypedElementHandle AcquireEditorObjectElementHandle(const UObject* Object, const bool bAllowCreate = true);
+#endif
+
+	static TTypedElementOwner<FActorElementData> CreateActorElement(const AActor* InActor);
+	static void DestroyActorElement(const AActor* InActor, TTypedElementOwner<FActorElementData>& InOutActorElement);
+#if WITH_EDITOR
+	static void CreateEditorActorElement(const AActor* Actor);
+	static void DestroyEditorActorElement(const AActor* Actor);
+	UFUNCTION(BlueprintCallable, Category="TypedElementFramework|Actor", meta=(ScriptMethod="AcquireEditorElementHandle"))
+	static FTypedElementHandle AcquireEditorActorElementHandle(const AActor* Actor, const bool bAllowCreate = true);
+#endif
+
+	static TTypedElementOwner<FComponentElementData> CreateComponentElement(const UActorComponent* InComponent);
+	static void DestroyComponentElement(const UActorComponent* InComponent, TTypedElementOwner<FComponentElementData>& InOutComponentElement);
+#if WITH_EDITOR
+	static void CreateEditorComponentElement(const UActorComponent* Component);
+	static void DestroyEditorComponentElement(const UActorComponent* Component);
+	UFUNCTION(BlueprintCallable, BlueprintPure=false, Category="TypedElementFramework|Component", meta=(ScriptMethod="AcquireEditorElementHandle"))
+	static FTypedElementHandle AcquireEditorComponentElementHandle(const UActorComponent* Component, const bool bAllowCreate = true);
+#endif
+
+private:
+#if WITH_EDITOR
+	static void DestroyUnreachableEditorObjectElements();
+#endif
+};
diff --git a/Engine/Source/Runtime/Engine/Public/Elements/Object/ObjectElementData.h b/Engine/Source/Runtime/Engine/Public/Elements/Object/ObjectElementData.h
--- a/Engine/Source/Runtime/Engine/Public/Elements/Object/ObjectElementData.h
+++ b/Engine/Source/Runtime/Engine/Public/Elements/Object/ObjectElementData.h
@@ -2,7 +2,7 @@
 
 #pragma once
 
-#include "TypedElementData.h"
+#include "Elements/Framework/TypedElementData.h"
 
 class UObject;
 
diff --git a/Engine/Source/Runtime/Engine/Public/Elements/Object/ObjectElementObjectInterface.h b/Engine/Source/Runtime/Engine/Public/Elements/Object/ObjectElementObjectInterface.h
--- a/Engine/Source/Runtime/Engine/Public/Elements/Object/ObjectElementObjectInterface.h
+++ b/Engine/Source/Runtime/Engine/Public/Elements/Object/ObjectElementObjectInterface.h
@@ -2,7 +2,7 @@
 
 #pragma once
 
-#include "TypedElementObjectInterface.h"
+#include "Elements/Interfaces/TypedElementObjectInterface.h"
 #include "ObjectElementObjectInterface.generated.h"
 
 UCLASS()
diff --git a/Engine/Source/Runtime/Engine/Public/Elements/Object/ObjectElementSelectionInterface.h b/Engine/Source/Runtime/Engine/Public/Elements/Object/ObjectElementSelectionInterface.h
--- a/Engine/Source/Runtime/Engine/Public/Elements/Object/ObjectElementSelectionInterface.h
+++ b/Engine/Source/Runtime/Engine/Public/Elements/Object/ObjectElementSelectionInterface.h
@@ -2,7 +2,7 @@
 
 #pragma once
 
-#include "TypedElementSelectionInterface.h"
+#include "Elements/Interfaces/TypedElementSelectionInterface.h"
 #include "ObjectElementSelectionInterface.generated.h"
 
 UCLASS()
diff --git a/Engine/Source/Runtime/Engine/Public/HitProxies.h b/Engine/Source/Runtime/Engine/Public/HitProxies.h
--- a/Engine/Source/Runtime/Engine/Public/HitProxies.h
+++ b/Engine/Source/Runtime/Engine/Public/HitProxies.h
@@ -10,7 +10,7 @@
 #include "UObject/ObjectMacros.h"
 #include "Templates/RefCounting.h"
 #include "GenericPlatform/ICursor.h"
-#include "TypedElementHandle.h"
+#include "Elements/Framework/TypedElementHandle.h"
 
 class FReferenceCollector;
 
diff --git a/Engine/Source/Runtime/TypedElementFramework/Private/Elements/Framework/TypedElementList.cpp b/Engine/Source/Runtime/TypedElementFramework/Private/Elements/Framework/TypedElementList.cpp
new file mode 100644
--- /dev/null
+++ b/Engine/Source/Runtime/TypedElementFramework/Private/Elements/Framework/TypedElementList.cpp
@@ -0,0 +1,281 @@
+// Copyright Epic Games, Inc. All Rights Reserved.
+
+#include "Elements/Framework/TypedElementList.h"
+#include "Elements/Framework/TypedElementRegistry.h"
+
+namespace TypedElementList_Private
+{
+
+void GetElementImpl(const UTypedElementRegistry* InRegistry, const FTypedElementHandle& InElementHandle, const UClass* InBaseInterfaceType, FTypedElement& OutElement)
+{
+	InRegistry->Private_GetElementImpl(InElementHandle, InBaseInterfaceType, OutElement);
+}
+
+} // namespace TypedElementList_Private
+
+
+FTypedElementListLegacySync::FTypedElementListLegacySync(const UTypedElementList* InElementList)
+	: ElementList(InElementList)
+{
+	checkf(ElementList, TEXT("ElementList was null!"));
+}
+
+FTypedElementListLegacySync::FOnSyncEvent& FTypedElementListLegacySync::OnSyncEvent()
+{
+	return OnSyncEventDelegate;
+}
+
+void FTypedElementListLegacySync::Private_EmitSyncEvent(const ESyncType InSyncType, const FTypedElementHandle& InElementHandle)
+{
+	const bool bIsWithinBatchOperation = IsRunningBatchOperation();
+	bBatchOperationIsDirty |= bIsWithinBatchOperation;
+	OnSyncEventDelegate.Broadcast(ElementList, InSyncType, InElementHandle, bIsWithinBatchOperation);
+}
+
+bool FTypedElementListLegacySync::IsRunningBatchOperation() const
+{
+	return NumOpenBatchOperations > 0;
+}
+
+void FTypedElementListLegacySync::BeginBatchOperation()
+{
+	++NumOpenBatchOperations;
+}
+
+void FTypedElementListLegacySync::EndBatchOperation(const bool InNotify)
+{
+	checkf(NumOpenBatchOperations > 0, TEXT("Batch operation underflow!"));
+
+	if (--NumOpenBatchOperations == 0)
+	{
+		const bool bNotifyChange = bBatchOperationIsDirty && InNotify;
+		bBatchOperationIsDirty = false;
+
+		if (bNotifyChange)
+		{
+			Private_EmitSyncEvent(ESyncType::BatchComplete);
+			check(!bBatchOperationIsDirty); // This should still be false after emitting the notification!
+		}
+	}
+}
+
+bool FTypedElementListLegacySync::IsBatchOperationDirty() const
+{
+	return bBatchOperationIsDirty;
+}
+
+void FTypedElementListLegacySync::ForceBatchOperationDirty()
+{
+	if (NumOpenBatchOperations > 0)
+	{
+		bBatchOperationIsDirty = true;
+	}
+}
+
+
+FTypedElementListLegacySyncScopedBatch::FTypedElementListLegacySyncScopedBatch(UTypedElementList* InElementList, const bool InNotify)
+	: ElementListLegacySync(InElementList->Legacy_GetSyncPtr())
+	, bNotify(InNotify)
+{
+	if (ElementListLegacySync)
+	{
+		ElementListLegacySync->BeginBatchOperation();
+	}
+}
+
+FTypedElementListLegacySyncScopedBatch::~FTypedElementListLegacySyncScopedBatch()
+{
+	if (ElementListLegacySync)
+	{
+		ElementListLegacySync->EndBatchOperation(bNotify);
+	}
+}
+
+UTypedElementList* UTypedElementList::Private_CreateElementList(UTypedElementRegistry* InRegistry)
+{
+	UTypedElementList* ElementList = NewObject<UTypedElementList>();
+	ElementList->Initialize(InRegistry);
+	return ElementList;
+}
+
+void UTypedElementList::Initialize(UTypedElementRegistry* InRegistry)
+{
+	checkf(!Registry.Get(), TEXT("Initialize has already been called!"));
+	Registry = InRegistry;
+	checkf(InRegistry, TEXT("Registry is null!"));
+	InRegistry->Private_OnElementListCreated(this);
+}
+
+void UTypedElementList::BeginDestroy()
+{
+	Super::BeginDestroy();
+
+	LegacySync.Reset();
+	if (UTypedElementRegistry* RegistryPtr = Registry.Get())
+	{
+		RegistryPtr->Private_OnElementListDestroyed(this);
+		Registry = nullptr;
+	}
+}
+
+UTypedElementList* UTypedElementList::Clone() const
+{
+	UTypedElementList* ClonedElementList = Private_CreateElementList(Registry.Get());
+	ClonedElementList->ElementCombinedIds = ElementCombinedIds;
+	ClonedElementList->ElementHandles = ElementHandles;
+	return ClonedElementList;
+}
+
+UTypedElementInterface* UTypedElementList::GetElementInterface(const FTypedElementHandle& InElementHandle, const TSubclassOf<UTypedElementInterface>& InBaseInterfaceType) const
+{
+	return Registry->GetElementInterface(InElementHandle, InBaseInterfaceType);
+}
+
+bool UTypedElementList::AddElementImpl(FTypedElementHandle&& InElementHandle)
+{
+	if (!InElementHandle)
+	{
+		return false;
+	}
+
+	NoteListMayChange();
+
+	bool bAlreadyAdded = false;
+	ElementCombinedIds.Add(InElementHandle.GetId().GetCombinedId(), &bAlreadyAdded);
+
+	if (!bAlreadyAdded)
+	{
+		const FTypedElementHandle& AddedElementHandle = ElementHandles.Add_GetRef(MoveTemp(InElementHandle));
+		NoteListChanged(EChangeType::Added, AddedElementHandle);
+	}
+
+	return !bAlreadyAdded;
+}
+
+bool UTypedElementList::RemoveElementImpl(const FTypedElementId& InElementId)
+{
+	if (!InElementId)
+	{
+		return false;
+	}
+
+	NoteListMayChange();
+
+	const bool bRemoved = ElementCombinedIds.Remove(InElementId.GetCombinedId()) > 0;
+
+	if (bRemoved)
+	{
+		const int32 ElementHandleIndexToRemove = ElementHandles.IndexOfByPredicate([&InElementId](const FTypedElementHandle& InElementHandle)
+		{
+			return InElementHandle.GetId() == InElementId;
+		});
+		checkSlow(ElementHandleIndexToRemove != INDEX_NONE);
+
+		FTypedElementHandle RemovedElementHandle = MoveTemp(ElementHandles[ElementHandleIndexToRemove]);
+		ElementHandles.RemoveAt(ElementHandleIndexToRemove, 1, /*bAllowShrinking*/false);
+
+		NoteListChanged(EChangeType::Removed, RemovedElementHandle);
+	}
+
+	return bRemoved;
+}
+
+int32 UTypedElementList::RemoveAllElementsImpl(TFunctionRef<bool(const FTypedElementHandle&)> InPredicate)
+{
+	int32 RemovedCount = 0;
+
+	if (ElementHandles.Num() > 0)
+	{
+		FTypedElementListLegacySyncScopedBatch LegacySyncBatch(this);
+
+		NoteListMayChange();
+
+		for (int32 Index = ElementHandles.Num() - 1; Index >= 0; --Index)
+		{
+			if (InPredicate(ElementHandles[Index]))
+			{
+				FTypedElementHandle RemovedElementHandle = MoveTemp(ElementHandles[Index]);
+				ElementCombinedIds.Remove(RemovedElementHandle.GetId().GetCombinedId());
+				ElementHandles.RemoveAt(Index, 1, /*bAllowShrinking*/false);
+
+				NoteListChanged(EChangeType::Removed, RemovedElementHandle);
+
+				++RemovedCount;
+			}
+		}
+	}
+
+	return RemovedCount;
+}
+
+bool UTypedElementList::ContainsElementImpl(const FTypedElementId& InElementId) const
+{
+	return InElementId 
+		&& ElementCombinedIds.Contains(InElementId.GetCombinedId());
+}
+
+FTypedElementListLegacySync& UTypedElementList::Legacy_GetSync()
+{
+	if (!LegacySync)
+	{
+		LegacySync = MakeUnique<FTypedElementListLegacySync>(this);
+	}
+	return *LegacySync;
+}
+
+FTypedElementListLegacySync* UTypedElementList::Legacy_GetSyncPtr()
+{
+	return LegacySync.Get();
+}
+
+void UTypedElementList::NotifyPendingChanges()
+{
+	if (bHasPendingNotify)
+	{
+		bHasPendingNotify = false;
+		OnChangedDelegate.Broadcast(this);
+		check(!bHasPendingNotify); // This should still be false after emitting the notification!
+	}
+}
+
+void UTypedElementList::ClearPendingChanges()
+{
+	bHasPendingNotify = false;
+}
+
+void UTypedElementList::NoteListMayChange()
+{
+	if (!bHasPendingNotify)
+	{
+		OnPreChangeDelegate.Broadcast(this);
+	}
+}
+
+void UTypedElementList::NoteListChanged(const EChangeType InChangeType, const FTypedElementHandle& InElementHandle)
+{
+	bHasPendingNotify = true;
+
+	if (LegacySync)
+	{
+		FTypedElementListLegacySync::ESyncType SyncType = FTypedElementListLegacySync::ESyncType::Modified;
+		switch (InChangeType)
+		{
+		case EChangeType::Added:
+			SyncType = FTypedElementListLegacySync::ESyncType::Added;
+			break;
+
+		case EChangeType::Removed:
+			SyncType = FTypedElementListLegacySync::ESyncType::Removed;
+			break;
+
+		case EChangeType::Cleared:
+			SyncType = FTypedElementListLegacySync::ESyncType::Cleared;
+			break;
+
+		default:
+			break;
+		}
+
+		LegacySync->Private_EmitSyncEvent(SyncType, InElementHandle);
+	}
+}
diff --git a/Engine/Source/Runtime/TypedElementFramework/Private/Elements/Framework/TypedElementRegistry.cpp b/Engine/Source/Runtime/TypedElementFramework/Private/Elements/Framework/TypedElementRegistry.cpp
new file mode 100644
--- /dev/null
+++ b/Engine/Source/Runtime/TypedElementFramework/Private/Elements/Framework/TypedElementRegistry.cpp
@@ -0,0 +1,177 @@
+// Copyright Epic Games, Inc. All Rights Reserved.
+
+#include "Elements/Framework/TypedElementRegistry.h"
+#include "Misc/CoreDelegates.h"
+#include "Misc/ScopeLock.h"
+#include "UObject/StrongObjectPtr.h"
+
+const FTypedElementId FTypedElementId::Unset;
+
+TStrongObjectPtr<UTypedElementRegistry>& GetTypedElementRegistryInstance()
+{
+	static TStrongObjectPtr<UTypedElementRegistry> TypedElementRegistryInstance;
+	return TypedElementRegistryInstance;
+}
+
+UTypedElementRegistry::UTypedElementRegistry()
+{
+	FCoreDelegates::OnEndFrame.AddUObject(this, &UTypedElementRegistry::NotifyElementListPendingChanges);
+}
+
+void UTypedElementRegistry::Private_InitializeInstance()
+{
+	TStrongObjectPtr<UTypedElementRegistry>& Instance = GetTypedElementRegistryInstance();
+	checkf(!Instance, TEXT("Instance was already initialized!"));
+	Instance.Reset(NewObject<UTypedElementRegistry>());
+}
+
+void UTypedElementRegistry::Private_ShutdownInstance()
+{
+	TStrongObjectPtr<UTypedElementRegistry>& Instance = GetTypedElementRegistryInstance();
+	Instance.Reset();
+}
+
+UTypedElementRegistry* UTypedElementRegistry::GetInstance()
+{
+	TStrongObjectPtr<UTypedElementRegistry>& Instance = GetTypedElementRegistryInstance();
+	return Instance.Get();
+}
+
+void UTypedElementRegistry::FinishDestroy()
+{
+	TStrongObjectPtr<UTypedElementRegistry>& Instance = GetTypedElementRegistryInstance();
+	if (Instance.Get() == this)
+	{
+		Instance.Reset();
+	}
+
+	Super::FinishDestroy();
+}
+
+void UTypedElementRegistry::AddReferencedObjects(UObject* InThis, FReferenceCollector& Collector)
+{
+	Super::AddReferencedObjects(InThis, Collector);
+
+	UTypedElementRegistry* This = CastChecked<UTypedElementRegistry>(InThis);
+	for (TUniquePtr<FRegisteredElementType>& RegisteredElementType : This->RegisteredElementTypes)
+	{
+		if (RegisteredElementType)
+		{
+			for (auto& InterfacesPair : RegisteredElementType->Interfaces)
+			{
+				Collector.AddReferencedObject(InterfacesPair.Value);
+			}
+		}
+	}
+}
+
+void UTypedElementRegistry::RegisterElementTypeImpl(const FName InElementTypeName, TUniquePtr<FRegisteredElementType>&& InRegisteredElementType)
+{
+	// Query whether this type has previously been registered in any type registry, and if so re-use that ID
+	// If not (or if the element is typeless) then assign the next available ID
+	FTypedHandleTypeId TypeId = InRegisteredElementType->GetDataTypeId();
+	if (TypeId == 0)
+	{
+		static FCriticalSection NextTypeIdCS;
+		static FTypedHandleTypeId NextTypeId = 1;
+
+		FScopeLock NextTypeIdLock(&NextTypeIdCS);
+
+		checkf(NextTypeId <= TypedHandleMaxTypeId, TEXT("Ran out of typed element type IDs!"));
+
+		TypeId = NextTypeId++;
+		InRegisteredElementType->SetDataTypeId(TypeId);
+	}
+
+	InRegisteredElementType->TypeId = TypeId;
+	InRegisteredElementType->TypeName = InElementTypeName;
+	AddRegisteredElementType(MoveTemp(InRegisteredElementType));
+}
+
+void UTypedElementRegistry::RegisterElementInterfaceImpl(const FName InElementTypeName, UTypedElementInterface* InElementInterface, const TSubclassOf<UTypedElementInterface>& InBaseInterfaceType, const bool InAllowOverride)
+{
+	checkf(InElementInterface->IsA(InBaseInterfaceType), TEXT("Interface '%s' of type '%s' does not derive from '%s'!"), *InElementInterface->GetPathName(), *InElementInterface->GetClass()->GetName(), *InBaseInterfaceType->GetName());
+
+	FRegisteredElementType* RegisteredElementType = GetRegisteredElementTypeFromName(InElementTypeName);
+	checkf(RegisteredElementType, TEXT("Element type '%s' has not been registered!"), *InElementTypeName.ToString());
+
+	checkf(InAllowOverride || !RegisteredElementType->Interfaces.Contains(InBaseInterfaceType->GetFName()), TEXT("Element type '%s' has already registered an interface for '%s'!"), *InElementTypeName.ToString(), *InBaseInterfaceType->GetName());
+	RegisteredElementType->Interfaces.Add(InBaseInterfaceType->GetFName(), InElementInterface);
+}
+
+UTypedElementInterface* UTypedElementRegistry::GetElementInterfaceImpl(const FTypedHandleTypeId InElementTypeId, const TSubclassOf<UTypedElementInterface>& InBaseInterfaceType) const
+{
+	if (!InElementTypeId)
+	{
+		return nullptr;
+	}
+
+	FRegisteredElementType* RegisteredElementType = GetRegisteredElementTypeFromId(InElementTypeId);
+	checkf(RegisteredElementType, TEXT("Element type ID '%d' has not been registered!"), InElementTypeId);
+
+	return RegisteredElementType->Interfaces.FindRef(InBaseInterfaceType->GetFName());
+}
+
+void UTypedElementRegistry::ReleaseElementId(FTypedElementId& InOutElementId)
+{
+	if (!InOutElementId)
+	{
+		return;
+	}
+
+	FRegisteredElementType* RegisteredElementType = GetRegisteredElementTypeFromId(InOutElementId.GetTypeId());
+	checkf(RegisteredElementType, TEXT("Element type ID '%d' has not been registered!"), InOutElementId.GetTypeId());
+
+	const FTypedElementInternalData& ElementData = RegisteredElementType->GetDataForElement(InOutElementId.GetElementId());
+	ElementData.ReleaseRef(INDEX_NONE); // Cannot track element ID references as we have no space to store the reference ID
+
+	InOutElementId.Private_DestroyNoRef();
+}
+
+FTypedElementHandle UTypedElementRegistry::GetElementHandle(const FTypedElementId& InElementId) const
+{
+	if (!InElementId)
+	{
+		return FTypedElementHandle();
+	}
+
+	FRegisteredElementType* RegisteredElementType = GetRegisteredElementTypeFromId(InElementId.GetTypeId());
+	checkf(RegisteredElementType, TEXT("Element type ID '%d' has not been registered!"), InElementId.GetTypeId());
+
+	FTypedElementHandle ElementHandle;
+	ElementHandle.Private_InitializeAddRef(InElementId.GetTypeId(), InElementId.GetElementId(), RegisteredElementType->GetDataForElement(InElementId.GetElementId()));
+
+	return ElementHandle;
+}
+
+UTypedElementList* UTypedElementRegistry::CreateElementList(TArrayView<const FTypedElementId> InElementIds)
+{
+	UTypedElementList* ElementList = CreateElementList();
+
+	for (const FTypedElementId& ElementId : InElementIds)
+	{
+		if (FTypedElementHandle ElementHandle = GetElementHandle(ElementId))
+		{
+			ElementList->Add(MoveTemp(ElementHandle));
+		}
+	}
+
+	return ElementList;
+}
+
+UTypedElementList* UTypedElementRegistry::CreateElementList(TArrayView<const FTypedElementHandle> InElementHandles)
+{
+	UTypedElementList* ElementList = CreateElementList();
+	ElementList->Append(InElementHandles);
+	return ElementList;
+}
+
+void UTypedElementRegistry::NotifyElementListPendingChanges()
+{
+	FReadScopeLock ActiveElementListsLock(ActiveElementListsRW);
+
+	for (UTypedElementList* ActiveElementList : ActiveElementLists)
+	{
+		ActiveElementList->NotifyPendingChanges();
+	}
+}
diff --git a/Engine/Source/Runtime/TypedElementFramework/Private/TypedElementFrameworkModule.cpp b/Engine/Source/Runtime/TypedElementFramework/Private/TypedElementFrameworkModule.cpp
--- a/Engine/Source/Runtime/TypedElementFramework/Private/TypedElementFrameworkModule.cpp
+++ b/Engine/Source/Runtime/TypedElementFramework/Private/TypedElementFrameworkModule.cpp
@@ -1,6 +1,6 @@
 // Copyright Epic Games, Inc. All Rights Reserved.
 
-#include "TypedElementRegistry.h"
+#include "Elements/Framework/TypedElementRegistry.h"
 #include "Modules/ModuleManager.h"
 
 class FTypedElementFrameworkModule : public FDefaultModuleImpl
diff --git a/Engine/Source/Runtime/TypedElementFramework/Private/TypedElementFrameworkTests.h b/Engine/Source/Runtime/TypedElementFramework/Private/TypedElementFrameworkTests.h
--- a/Engine/Source/Runtime/TypedElementFramework/Private/TypedElementFrameworkTests.h
+++ b/Engine/Source/Runtime/TypedElementFramework/Private/TypedElementFrameworkTests.h
@@ -2,7 +2,7 @@
 
 #pragma once
 
-#include "TypedElementRegistry.h"
+#include "Elements/Framework/TypedElementRegistry.h"
 #include "TypedElementFrameworkTests.generated.h"
 
 /**
diff --git a/Engine/Source/Runtime/TypedElementFramework/Private/TypedElementList.cpp b/Engine/Source/Runtime/TypedElementFramework/Private/TypedElementList.cpp
deleted file mode 100644
diff --git a/Engine/Source/Runtime/TypedElementFramework/Private/TypedElementRegistry.cpp b/Engine/Source/Runtime/TypedElementFramework/Private/TypedElementRegistry.cpp
deleted file mode 100644
diff --git a/Engine/Source/Runtime/TypedElementFramework/Public/Elements/Framework/TypedElementData.h b/Engine/Source/Runtime/TypedElementFramework/Public/Elements/Framework/TypedElementData.h
new file mode 100644
--- /dev/null
+++ b/Engine/Source/Runtime/TypedElementFramework/Public/Elements/Framework/TypedElementData.h
@@ -0,0 +1,356 @@
+// Copyright Epic Games, Inc. All Rights Reserved.
+
+#pragma once
+
+#include "CoreMinimal.h"
+#include "Misc/ScopeLock.h"
+#include "Misc/ScopeRWLock.h"
+#include "HAL/PlatformStackWalk.h"
+#include "Containers/SparseArray.h"
+#include "Containers/ChunkedArray.h"
+#include "Elements/Framework/TypedElementLimits.h"
+
+/**
+ * Macro to declare the required RTTI data for types representing element data.
+ * @note Place this in the public section of your type declaration.
+ */
+#define UE_DECLARE_TYPED_ELEMENT_DATA_RTTI(ELEMENT_DATA_TYPE)						\
+	static FTypedHandleTypeId Private_RegisteredTypeId;								\
+	static FTypedHandleTypeId StaticTypeId() { return Private_RegisteredTypeId; }	\
+	static FName StaticTypeName() { static const FName TypeName = #ELEMENT_DATA_TYPE; return TypeName; }
+
+/**
+ * Macro to define the required RTTI data for types representing element data.
+ * @note Place this in the cpp file for your type definition.
+ */
+#define UE_DEFINE_TYPED_ELEMENT_DATA_RTTI(ELEMENT_DATA_TYPE)						\
+	FTypedHandleTypeId ELEMENT_DATA_TYPE::Private_RegisteredTypeId = 0;
+
+#if UE_TYPED_ELEMENT_HAS_REFTRACKING
+/**
+ * Debugging information used to locate reference leaks.
+ */
+class FTypedElementReference
+{
+public:
+	FTypedElementReference()
+	{
+		CallstackDepth = FPlatformStackWalk::CaptureStackBackTrace(Callstack, UE_ARRAY_COUNT(Callstack));
+	}
+
+	FTypedElementReference(const FTypedElementReference&) = default;
+	FTypedElementReference& operator=(const FTypedElementReference&) = default;
+
+	FTypedElementReference(FTypedElementReference&&) = default;
+	FTypedElementReference& operator=(FTypedElementReference&&) = default;
+
+	void LogReference() const
+	{
+		ANSICHAR CallstackText[4096];
+		for (uint32 CallstackIndex = TypedHandleRefTrackingSkipCount; CallstackIndex < CallstackDepth; ++CallstackIndex)
+		{
+			CallstackText[0] = 0;
+			FPlatformStackWalk::ProgramCounterToHumanReadableString(CallstackIndex - TypedHandleRefTrackingSkipCount, Callstack[CallstackIndex], CallstackText, UE_ARRAY_COUNT(CallstackText));
+			UE_LOG(LogCore, Error, TEXT("%s"), ANSI_TO_TCHAR(CallstackText));
+		}
+	}
+
+private:
+	uint64 Callstack[TypedHandleRefTrackingDepth + TypedHandleRefTrackingSkipCount];
+	uint32 CallstackDepth = 0;
+};
+#endif	// UE_TYPED_ELEMENT_HAS_REFTRACKING
+
+/**
+ * Base class for the internal payload data associated with elements.
+ */
+class FTypedElementInternalData
+{
+public:
+	FTypedElementInternalData() = default;
+	
+	FTypedElementInternalData(const FTypedElementInternalData&) = delete;
+	FTypedElementInternalData& operator=(const FTypedElementInternalData&) = delete;
+	
+	FTypedElementInternalData(FTypedElementInternalData&& InOther) = delete;
+	FTypedElementInternalData& operator=(FTypedElementInternalData&&) = delete;
+
+	virtual ~FTypedElementInternalData() = default;
+
+	virtual void Reset()
+	{
+#if UE_TYPED_ELEMENT_HAS_REFCOUNTING
+		FPlatformAtomics::InterlockedExchange(&RefCount, 0);
+#endif	// UE_TYPED_ELEMENT_HAS_REFCOUNTING
+#if UE_TYPED_ELEMENT_HAS_REFTRACKING
+		{
+			FScopeLock ReferencesLock(&ReferencesCS);
+			References.Reset();
+		}
+#endif	// UE_TYPED_ELEMENT_HAS_REFTRACKING
+	}
+
+	FORCEINLINE FTypedElementReferenceId AddRef(const bool bCanTrackReference) const
+	{
+		FTypedElementReferenceId ReferenceId = INDEX_NONE;
+#if UE_TYPED_ELEMENT_HAS_REFCOUNTING
+		checkSlow(RefCount < TNumericLimits<FTypedElementRefCount>::Max());
+		FPlatformAtomics::InterlockedIncrement(&RefCount);
+#endif	// UE_TYPED_ELEMENT_HAS_REFCOUNTING
+#if UE_TYPED_ELEMENT_HAS_REFTRACKING
+		if (bCanTrackReference)
+		{
+			FScopeLock ReferencesLock(&ReferencesCS);
+			ReferenceId = References.Add(FTypedElementReference());
+		}
+#endif	// UE_TYPED_ELEMENT_HAS_REFTRACKING
+		return ReferenceId;
+	}
+
+	FORCEINLINE void ReleaseRef(const FTypedElementReferenceId InReferenceId) const
+	{
+#if UE_TYPED_ELEMENT_HAS_REFCOUNTING
+		checkSlow(RefCount > 0);
+		FPlatformAtomics::InterlockedDecrement(&RefCount);
+#endif	// UE_TYPED_ELEMENT_HAS_REFCOUNTING
+#if UE_TYPED_ELEMENT_HAS_REFTRACKING
+		if (InReferenceId != INDEX_NONE)
+		{
+			FScopeLock ReferencesLock(&ReferencesCS);
+			References.RemoveAt(InReferenceId);
+		}
+#endif	// UE_TYPED_ELEMENT_HAS_REFTRACKING
+	}
+
+	FORCEINLINE FTypedElementRefCount GetRefCount() const
+	{
+#if UE_TYPED_ELEMENT_HAS_REFCOUNTING
+		return FPlatformAtomics::AtomicRead(&RefCount);
+#else	// UE_TYPED_ELEMENT_HAS_REFCOUNTING
+		return 0;
+#endif	// UE_TYPED_ELEMENT_HAS_REFCOUNTING
+	}
+
+	void LogReferences() const
+	{
+#if UE_TYPED_ELEMENT_HAS_REFTRACKING
+		{
+			FScopeLock ReferencesLock(&ReferencesCS);
+			UE_LOG(LogCore, Error, TEXT("External Element References:"));
+			for (const FTypedElementReference& Reference : References)
+			{
+				UE_LOG(LogCore, Error, TEXT("-----------------------------------------------"));
+				Reference.LogReference();
+			}
+			UE_LOG(LogCore, Error, TEXT("==============================================="));
+		}
+#else	// UE_TYPED_ELEMENT_HAS_REFTRACKING
+		UE_LOG(LogCore, Error, TEXT("UE_TYPED_ELEMENT_HAS_REFTRACKING is disabled. Enable it and recompile to see reference tracking."));
+#endif	// UE_TYPED_ELEMENT_HAS_REFTRACKING
+	}
+
+	virtual const void* GetUntypedData() const
+	{
+		return nullptr;
+	}
+
+private:
+#if UE_TYPED_ELEMENT_HAS_REFCOUNTING
+	mutable FTypedElementRefCount RefCount = 0;
+#endif	// UE_TYPED_ELEMENT_HAS_REFCOUNTING
+#if UE_TYPED_ELEMENT_HAS_REFTRACKING
+	mutable FCriticalSection ReferencesCS;
+	mutable TSparseArray<FTypedElementReference> References;
+#endif	// UE_TYPED_ELEMENT_HAS_REFTRACKING
+};
+
+/**
+ * Internal payload data associated with typed elements.
+ */
+template <typename ElementDataType>
+class TTypedElementInternalData : public FTypedElementInternalData
+{
+public:
+	virtual void Reset() override
+	{
+		FTypedElementInternalData::Reset();
+		Data = ElementDataType();
+	}
+
+	FORCEINLINE const ElementDataType& GetData() const
+	{
+		return Data;
+	}
+
+	FORCEINLINE ElementDataType& GetMutableData()
+	{
+		return Data;
+	}
+
+	virtual const void* GetUntypedData() const override
+	{
+		return &Data;
+	}
+
+private:
+	ElementDataType Data;
+};
+
+/**
+ * Internal payload data associated with typeless elements.
+ */
+template <>
+class TTypedElementInternalData<void> : public FTypedElementInternalData
+{
+};
+
+/**
+ * Data store implementation used by the element registry to manage internal data. 
+ * @note This is the generic implementation that uses an array and manages the IDs itself.
+ */
+template <typename ElementDataType>
+class TTypedElementInternalDataStore
+{
+public:
+	static_assert(TNumericLimits<int32>::Max() >= TypedHandleMaxElementId, "TTypedElementInternalDataStore internally uses signed 32-bit indices so cannot store TypedHandleMaxElementId! Consider making this container 64-bit aware, or explicitly remove this compile time check.");
+
+	TTypedElementInternalData<ElementDataType>& AddDataForElement(FTypedHandleElementId& InOutElementId)
+	{
+		FWriteScopeLock InternalDataLock(InternalDataRW);
+
+		checkSlow(InOutElementId < 0);
+
+		InOutElementId = InternalDataFreeIndices.Num() > 0
+			? InternalDataFreeIndices.Pop(/*bAllowShrinking*/false)
+			: InternalDataArray.Add();
+
+		return InternalDataArray[InOutElementId];
+	}
+
+	void RemoveDataForElement(const FTypedHandleElementId InElementId, const FTypedElementInternalData* InExpectedDataPtr)
+	{
+		FWriteScopeLock InternalDataLock(InternalDataRW);
+
+		checkSlow(InternalDataArray.IsValidIndex(InElementId));
+		
+		TTypedElementInternalData<ElementDataType>& InternalData = InternalDataArray[InElementId];
+		checkf(InExpectedDataPtr == &InternalData, TEXT("Internal data pointer did not match the expected value! Does this handle belong to a different element registry?"));
+		InternalData.Reset();
+		InternalDataFreeIndices.Add(InElementId);
+	}
+
+	const TTypedElementInternalData<ElementDataType>& GetDataForElement(const FTypedHandleElementId InElementId) const
+	{
+		FReadScopeLock InternalDataLock(InternalDataRW);
+
+		checkSlow(InternalDataArray.IsValidIndex(InElementId));
+		return InternalDataArray[InElementId];
+	}
+
+	static FORCEINLINE void SetStaticDataTypeId(const FTypedHandleTypeId InTypeId)
+	{
+		checkSlow(ElementDataType::Private_RegisteredTypeId == 0);
+		ElementDataType::Private_RegisteredTypeId = InTypeId;
+	}
+
+	static FORCEINLINE FTypedHandleTypeId StaticDataTypeId()
+	{
+		return ElementDataType::StaticTypeId();
+	}
+
+	static FORCEINLINE FName StaticDataTypeName()
+	{
+		return ElementDataType::StaticTypeName();
+	}
+
+private:
+	mutable FRWLock InternalDataRW;
+	TChunkedArray<TTypedElementInternalData<ElementDataType>> InternalDataArray;
+	TArray<int32> InternalDataFreeIndices;
+};
+
+/**
+ * Data store implementation used by the element registry to manage internal data. 
+ * @note This is the typeless implementation that uses external IDs, and exists only to track ref counts.
+ */
+template <>
+class TTypedElementInternalDataStore<void>
+{
+public:
+	static_assert(TNumericLimits<int32>::Max() >= TypedHandleMaxElementId, "TTypedElementInternalDataStore internally uses signed 32-bit indices so cannot store TypedHandleMaxElementId! Consider making this container 64-bit aware, or explicitly remove this compile time check.");
+
+	TTypedElementInternalData<void>& AddDataForElement(FTypedHandleElementId& InOutElementId)
+	{
+#if UE_TYPED_ELEMENT_HAS_REFERENCING
+		FWriteScopeLock InternalDataLock(InternalDataRW);
+
+		checkSlow(InOutElementId >= 0);
+		checkSlow(!ElementIdToArrayIndex.Contains(InOutElementId));
+
+		const int32 InternalDataArrayIndex = InternalDataFreeIndices.Num() > 0
+			? InternalDataFreeIndices.Pop(/*bAllowShrinking*/false)
+			: InternalDataArray.Add();
+
+		ElementIdToArrayIndex.Add(InOutElementId, InternalDataArrayIndex);
+		return InternalDataArray[InternalDataArrayIndex];
+#else	// UE_TYPED_ELEMENT_HAS_REFERENCING
+		return SharedInternalData;
+#endif	// UE_TYPED_ELEMENT_HAS_REFERENCING
+	}
+
+	void RemoveDataForElement(const FTypedHandleElementId InElementId, const FTypedElementInternalData* InExpectedDataPtr)
+	{
+#if UE_TYPED_ELEMENT_HAS_REFERENCING
+		FWriteScopeLock InternalDataLock(InternalDataRW);
+
+		int32 InternalDataArrayIndex = INDEX_NONE;
+		ElementIdToArrayIndex.RemoveAndCopyValue(InElementId, InternalDataArrayIndex);
+
+		checkSlow(InternalDataArray.IsValidIndex(InternalDataArrayIndex));
+
+		TTypedElementInternalData<void>& InternalData = InternalDataArray[InElementId];
+		checkf(InExpectedDataPtr == &InternalData, TEXT("Internal data pointer did not match the expected value! Does this handle belong to a different element registry?"));
+		InternalData.Reset();
+		InternalDataFreeIndices.Add(InternalDataArrayIndex);
+#else	// UE_TYPED_ELEMENT_HAS_REFERENCING
+		checkf(InExpectedDataPtr == &SharedInternalData, TEXT("Internal data pointer did not match the expected value! Does this handle belong to a different element registry?"));
+#endif	// UE_TYPED_ELEMENT_HAS_REFERENCING
+	}
+
+	const TTypedElementInternalData<void>& GetDataForElement(const FTypedHandleElementId InElementId) const
+	{
+#if UE_TYPED_ELEMENT_HAS_REFERENCING
+		FReadScopeLock InternalDataLock(InternalDataRW);
+
+		const int32* InternalDataArrayIndexPtr = ElementIdToArrayIndex.Find(InElementId);
+		checkSlow(InternalDataArrayIndexPtr && InternalDataArray.IsValidIndex(*InternalDataArrayIndexPtr));
+		return InternalDataArray[*InternalDataArrayIndexPtr];
+#else	// UE_TYPED_ELEMENT_HAS_REFERENCING
+		return SharedInternalData;
+#endif	// UE_TYPED_ELEMENT_HAS_REFERENCING
+	}
+
+	static FORCEINLINE void SetStaticDataTypeId(const FTypedHandleTypeId InTypeId)
+	{
+	}
+
+	static FORCEINLINE FTypedHandleTypeId StaticDataTypeId()
+	{
+		return 0;
+	}
+
+	static FORCEINLINE FName StaticDataTypeName()
+	{
+		return FName();
+	}
+
+private:
+#if UE_TYPED_ELEMENT_HAS_REFERENCING
+	mutable FRWLock InternalDataRW;
+	TChunkedArray<TTypedElementInternalData<void>> InternalDataArray;
+	TArray<int32> InternalDataFreeIndices;
+	TMap<FTypedHandleElementId, int32> ElementIdToArrayIndex;
+#else	// UE_TYPED_ELEMENT_HAS_REFERENCING
+	TTypedElementInternalData<void> SharedInternalData;
+#endif	// UE_TYPED_ELEMENT_HAS_REFERENCING
+};
diff --git a/Engine/Source/Runtime/TypedElementFramework/Public/Elements/Framework/TypedElementHandle.h b/Engine/Source/Runtime/TypedElementFramework/Public/Elements/Framework/TypedElementHandle.h
new file mode 100644
--- /dev/null
+++ b/Engine/Source/Runtime/TypedElementFramework/Public/Elements/Framework/TypedElementHandle.h
@@ -0,0 +1,904 @@
+// Copyright Epic Games, Inc. All Rights Reserved.
+
+#pragma once
+
+#include <type_traits>
+#include "CoreMinimal.h"
+#include "Elements/Framework/TypedElementData.h"
+#include "Elements/Interfaces/TypedElementInterface.h"
+#include "TypedElementHandle.generated.h"
+
+/**
+ * The most minimal representation of an element - its ID!
+ * This type is not immediately useful on its own, but can be used to find an element from the element registry or an element list.
+ * @note This is ref-counted like handles themselves are, so as long as an ID is available, the handle will be too.
+ * @note IDs lack the information needed to auto-release on destruction, so must be manually released, either via the corresponding handle or their owner element registry.
+ */
+struct TYPEDELEMENTFRAMEWORK_API FTypedElementId
+{
+public:
+	FTypedElementId()
+		: CombinedId(0)
+	{
+	}
+
+	FTypedElementId(const FTypedElementId&) = delete;
+	FTypedElementId& operator=(const FTypedElementId&) = delete;
+	
+	FTypedElementId(FTypedElementId&& InOther)
+		: CombinedId(InOther.CombinedId)
+	{
+		InOther.Private_DestroyNoRef();
+	}
+
+	FTypedElementId& operator=(FTypedElementId&& InOther)
+	{
+		if (this != &InOther)
+		{
+			CombinedId = InOther.CombinedId;
+
+			InOther.Private_DestroyNoRef();
+		}
+		return *this;
+	}
+
+	~FTypedElementId()
+	{
+		checkf(!IsSet(), TEXT("Element ID was still set during destruction! This will leak an element reference, and you should release this ID prior to destruction!"));
+	}
+
+	FORCEINLINE explicit operator bool() const
+	{
+		return IsSet();
+	}
+
+	/**
+	 * Has this ID been initialized to a valid element?
+	 */
+	FORCEINLINE bool IsSet() const
+	{
+		return TypeId != 0;
+	}
+
+	/**
+	 * Access the type ID portion of this element ID.
+	 */
+	FORCEINLINE FTypedHandleTypeId GetTypeId() const
+	{
+		return TypeId;
+	}
+
+	/**
+	 * Access the element ID portion of this element ID.
+	 */
+	FORCEINLINE FTypedHandleElementId GetElementId() const
+	{
+		return ElementId;
+	}
+
+	/**
+	 * Access the combined value of this element ID.
+	 * @note You typically don't want to store this directly as the element ID could be re-used.
+	 *       It is primarily useful as a secondary cache where something is keeping a reference to an element ID or element handle (eg, how UTypedElementList uses it internally).
+	 */
+	FORCEINLINE FTypedHandleCombinedId GetCombinedId() const
+	{
+		return CombinedId;
+	}
+
+	FORCEINLINE friend bool operator==(const FTypedElementId& InLHS, const FTypedElementId& InRHS)
+	{
+		return InLHS.CombinedId == InRHS.CombinedId;
+	}
+
+	FORCEINLINE friend bool operator!=(const FTypedElementId& InLHS, const FTypedElementId& InRHS)
+	{
+		return !(InLHS == InRHS);
+	}
+
+	friend inline uint32 GetTypeHash(const FTypedElementId& InElementId)
+	{
+		return GetTypeHash(InElementId.CombinedId);
+		
+	}
+
+	FORCEINLINE void Private_InitializeNoRef(const FTypedHandleTypeId InTypeId, const FTypedHandleElementId InElementId)
+	{
+		TypeId = InTypeId;
+		ElementId = InElementId;
+	}
+
+	FORCEINLINE void Private_DestroyNoRef()
+	{
+		CombinedId = 0;
+	}
+
+	/** An unset element ID */
+	static const FTypedElementId Unset;
+
+private:
+	union
+	{
+		struct
+		{
+			// Note: These are arranged in this order to give CombinedId better hash distribution for GetTypeHash!
+			FTypedHandleCombinedId ElementId : TypedHandleElementIdBits;
+			FTypedHandleCombinedId TypeId : TypedHandleTypeIdBits;
+		};
+		FTypedHandleCombinedId CombinedId;
+	};
+};
+
+/**
+ * A representation of an element that includes its handle data.
+ * This type is the most standard way that an element is passed through to interfaces, and also the type that is stored in element lists.
+ * C++ code may choose to use TTypedElement instead, which is a combination of an element handle and its associated element interface.
+ * @note Handles auto-release on destruction.
+ */
+USTRUCT(BlueprintType)
+struct TYPEDELEMENTFRAMEWORK_API FTypedElementHandle
+{
+	GENERATED_BODY()
+
+public:
+	FTypedElementHandle() = default;
+
+	FTypedElementHandle(const FTypedElementHandle& InOther)
+	{
+		if (InOther)
+		{
+			Private_InitializeAddRef(InOther.Id.GetTypeId(), InOther.Id.GetElementId(), *InOther.DataPtr);
+		}
+	}
+
+	FTypedElementHandle& operator=(const FTypedElementHandle& InOther)
+	{
+		if (this != &InOther)
+		{
+			Private_DestroyReleaseRef();
+
+			if (InOther)
+			{
+				Private_InitializeAddRef(InOther.Id.GetTypeId(), InOther.Id.GetElementId(), *InOther.DataPtr);
+			}
+		}
+		return *this;
+	}
+
+	FTypedElementHandle(FTypedElementHandle&& InOther)
+		: Id(MoveTemp(InOther.Id))
+		, DataPtr(InOther.DataPtr)
+#if UE_TYPED_ELEMENT_HAS_REFTRACKING
+		, ReferenceId(InOther.ReferenceId)
+#endif	// UE_TYPED_ELEMENT_HAS_REFTRACKING
+	{
+		InOther.DataPtr = nullptr;
+#if UE_TYPED_ELEMENT_HAS_REFTRACKING
+		InOther.ReferenceId = INDEX_NONE;
+#endif	// UE_TYPED_ELEMENT_HAS_REFTRACKING
+		checkSlow(!InOther.IsSet());
+	}
+
+	FTypedElementHandle& operator=(FTypedElementHandle&& InOther)
+	{
+		if (this != &InOther)
+		{
+			Private_DestroyReleaseRef();
+
+			Id = MoveTemp(InOther.Id);
+			DataPtr = InOther.DataPtr;
+#if UE_TYPED_ELEMENT_HAS_REFTRACKING
+			ReferenceId = InOther.ReferenceId;
+#endif	// UE_TYPED_ELEMENT_HAS_REFTRACKING
+
+			InOther.DataPtr = nullptr;
+#if UE_TYPED_ELEMENT_HAS_REFTRACKING
+			InOther.ReferenceId = INDEX_NONE;
+#endif	// UE_TYPED_ELEMENT_HAS_REFTRACKING
+			checkSlow(!InOther.IsSet());
+		}
+		return *this;
+	}
+
+	~FTypedElementHandle()
+	{
+		Private_DestroyReleaseRef();
+	}
+
+	FORCEINLINE explicit operator bool() const
+	{
+		return IsSet();
+	}
+
+	/**
+	 * Has this handle been initialized to a valid element?
+	 */
+	FORCEINLINE bool IsSet() const
+	{
+		return Id.IsSet();
+	}
+
+	/**
+	 * Release this handle and set it back to an empty state.
+	 */
+	FORCEINLINE void Release()
+	{
+		Private_DestroyReleaseRef();
+	}
+
+	/**
+	 * Get the ID that this element represents.
+	 */
+	FORCEINLINE const FTypedElementId& GetId() const
+	{
+		return Id;
+	}
+
+	/**
+	 * Test to see whether the data stored within this handle is of the given type.
+	 * @note This is not typically something you'd want to query outside of data access within an interface implementation.
+	 */
+	template <typename ElementDataType>
+	FORCEINLINE bool IsDataOfType() const
+	{
+		return Id.GetTypeId() == ElementDataType::StaticTypeId();
+	}
+
+	/**
+	 * Attempt to access the data stored within this handle as the given type, returning null if it isn't possible and logging an access error for scripting.
+	 * @note This is not typically something you'd want to use outside of data access within an interface implementation.
+	 */
+	template <typename ElementDataType>
+	const ElementDataType* GetData(const bool bSilent = false) const
+	{
+		if (!DataPtr)
+		{
+			if (!bSilent)
+			{
+				FFrame::KismetExecutionMessage(TEXT("Element handle data is null!"), ELogVerbosity::Error);
+			}
+			return nullptr;
+		}
+
+		if (!IsDataOfType<ElementDataType>())
+		{
+			if (!bSilent)
+			{
+				FFrame::KismetExecutionMessage(*FString::Printf(TEXT("Element handle data type is '%d', but '%d' (%s) was requested!"), Id.GetTypeId(), ElementDataType::StaticTypeId(), *ElementDataType::StaticTypeName().ToString()), ELogVerbosity::Error);
+			}
+			return nullptr;
+		}
+
+		return static_cast<const ElementDataType*>(DataPtr->GetUntypedData());
+	}
+
+	/**
+	 * Attempt to access the data stored within this handle as the given type, asserting if it isn't possible.
+	 * @note This is not typically something you'd want to use outside of data access within an interface implementation.
+	 */
+	template <typename ElementDataType>
+	FORCEINLINE const ElementDataType& GetDataChecked() const
+	{
+		checkf(DataPtr, TEXT("Element handle data is null!"));
+		checkf(IsDataOfType<ElementDataType>(), TEXT("Element handle data type is '%d', but '%d' (%s) was requested!"), Id.GetTypeId(), ElementDataType::StaticTypeId(), *ElementDataType::StaticTypeName().ToString());
+		return *static_cast<const ElementDataType*>(DataPtr->GetUntypedData());
+	}
+	
+	/**
+	 * Acquire a copy of the ID that this element represents.
+	 * @note This must be paired with a call to ReleaseId.
+	 */
+	FTypedElementId AcquireId() const
+	{
+		FTypedElementId ElementId;
+		if (IsSet())
+		{
+			DataPtr->AddRef(/*bCanTrackReference*/false); // Cannot track element ID references as we have no space to store the reference ID
+			ElementId.Private_InitializeNoRef(Id.GetTypeId(), Id.GetElementId());
+		}
+		return ElementId;
+	}
+
+	/**
+	 * Release a copy of the ID that this element represents.
+	 * @note This should have come from a call to AcquireId.
+	 */
+	void ReleaseId(FTypedElementId& InOutElementId) const
+	{
+		checkf(InOutElementId == Id, TEXT("Element ID does not match this handle!"));
+		if (InOutElementId)
+		{
+			DataPtr->ReleaseRef(INDEX_NONE); // Cannot track element ID references as we have no space to store the reference ID
+			InOutElementId.Private_DestroyNoRef();
+		}
+	}
+
+	FORCEINLINE friend bool operator==(const FTypedElementHandle& InLHS, const FTypedElementHandle& InRHS)
+	{
+		return InLHS.Id == InRHS.Id;
+	}
+
+	FORCEINLINE friend bool operator!=(const FTypedElementHandle& InLHS, const FTypedElementHandle& InRHS)
+	{
+		return !(InLHS == InRHS);
+	}
+
+	friend inline uint32 GetTypeHash(const FTypedElementHandle& InElementHandle)
+	{
+		return GetTypeHash(InElementHandle.Id);
+	}
+
+	FORCEINLINE void Private_InitializeNoRef(const FTypedHandleTypeId InTypeId, const FTypedHandleElementId InElementId, const FTypedElementInternalData& InData)
+	{
+		Id.Private_InitializeNoRef(InTypeId, InElementId);
+		DataPtr = &InData;
+	}
+
+	FORCEINLINE void Private_InitializeAddRef(const FTypedHandleTypeId InTypeId, const FTypedHandleElementId InElementId, const FTypedElementInternalData& InData)
+	{
+		Private_InitializeNoRef(InTypeId, InElementId, InData);
+		RegisterRef();
+	}
+
+	FORCEINLINE void Private_DestroyNoRef()
+	{
+		Id.Private_DestroyNoRef();
+		DataPtr = nullptr;
+#if UE_TYPED_ELEMENT_HAS_REFTRACKING
+		ReferenceId = INDEX_NONE;
+#endif	// UE_TYPED_ELEMENT_HAS_REFTRACKING
+	}
+
+	FORCEINLINE void Private_DestroyReleaseRef()
+	{
+		UnregisterRef();
+		Private_DestroyNoRef();
+	}
+
+	FORCEINLINE const FTypedElementInternalData* Private_GetInternalData() const
+	{
+		return DataPtr;
+	}
+
+private:
+	FORCEINLINE void RegisterRef()
+	{
+		if (DataPtr)
+		{
+#if !UE_TYPED_ELEMENT_HAS_REFTRACKING
+			FTypedElementReferenceId ReferenceId = INDEX_NONE;
+#endif	// !UE_TYPED_ELEMENT_HAS_REFTRACKING
+			ReferenceId = DataPtr->AddRef(/*bCanTrackReference*/true);
+		}
+	}
+
+	FORCEINLINE void UnregisterRef()
+	{
+		if (DataPtr)
+		{
+#if !UE_TYPED_ELEMENT_HAS_REFTRACKING
+			FTypedElementReferenceId ReferenceId = INDEX_NONE;
+#endif	// !UE_TYPED_ELEMENT_HAS_REFTRACKING
+			DataPtr->ReleaseRef(ReferenceId);
+		}
+	}
+
+	FTypedElementId Id;
+	const FTypedElementInternalData* DataPtr = nullptr;
+#if UE_TYPED_ELEMENT_HAS_REFTRACKING
+	FTypedElementReferenceId ReferenceId = INDEX_NONE;
+#endif	// UE_TYPED_ELEMENT_HAS_REFTRACKING
+};
+
+/**
+ * Common implementation of TTypedElement that is inherited by all specializations.
+ */
+template <typename BaseInterfaceType>
+struct TTypedElementBase : public FTypedElementHandle
+{
+public:
+	TTypedElementBase() = default;
+
+	TTypedElementBase(const TTypedElementBase&) = default;
+	TTypedElementBase& operator=(const TTypedElementBase&) = default;
+
+	TTypedElementBase(TTypedElementBase&& InOther)
+		: FTypedElementHandle(MoveTemp(InOther))
+		, InterfacePtr(InOther.InterfacePtr)
+	{
+		InOther.InterfacePtr = nullptr;
+		checkSlow(!InOther.IsSet());
+	}
+
+	TTypedElementBase& operator=(TTypedElementBase&& InOther)
+	{
+		if (this != &InOther)
+		{
+			FTypedElementHandle::operator=(MoveTemp(InOther));
+			InterfacePtr = InOther.InterfacePtr;
+
+			InOther.InterfacePtr = nullptr;
+			checkSlow(!InOther.IsSet());
+		}
+		return *this;
+	}
+
+	~TTypedElementBase()
+	{
+		Private_DestroyReleaseRef();
+	}
+
+	FORCEINLINE explicit operator bool() const
+	{
+		return IsSet();
+	}
+
+	/**
+	 * Has this element been initialized to a valid handle and interface?
+	 */
+	FORCEINLINE bool IsSet() const
+	{
+		return FTypedElementHandle::IsSet()
+			&& InterfacePtr;
+	}
+
+	/**
+	 * Release this element and set it back to an empty state.
+	 */
+	FORCEINLINE void Release()
+	{
+		Private_DestroyReleaseRef();
+	}
+
+	/**
+	 * Test to see whether the interface stored within this element is of the given type.
+	 */
+	template <typename InterfaceType>
+	FORCEINLINE bool HasInterface() const
+	{
+		return HasInterface(InterfaceType::StaticClass());
+	}
+
+	/**
+	 * Test to see whether the interface stored within this element is of the given type.
+	 */
+	FORCEINLINE bool HasInterface(const UClass* InInterfaceType) const
+	{
+		return InterfacePtr
+			&& InterfacePtr->IsA(InInterfaceType);
+	}
+
+	/**
+	 * Attempt to access the interface stored within this element, returning null if it isn't set.
+	 */
+	FORCEINLINE BaseInterfaceType* GetInterface() const
+	{
+		return InterfacePtr;
+	}
+
+	/**
+	 * Attempt to access the interface stored within this element, asserting if it isn't set.
+	 */
+	FORCEINLINE BaseInterfaceType& GetInterfaceChecked() const
+	{
+		checkf(InterfacePtr, TEXT("Interface is null!"));
+		return *InterfacePtr;
+	}
+
+	FORCEINLINE friend bool operator==(const TTypedElementBase& InLHS, const TTypedElementBase& InRHS)
+	{
+		return static_cast<const FTypedElementHandle&>(InLHS) == static_cast<const FTypedElementHandle&>(InRHS)
+			&& InLHS.InterfacePtr == InRHS.InterfacePtr;
+	}
+
+	FORCEINLINE friend bool operator!=(const TTypedElementBase& InLHS, const TTypedElementBase& InRHS)
+	{
+		return !(InLHS == InRHS);
+	}
+
+	friend inline uint32 GetTypeHash(const TTypedElementBase& InElement)
+	{
+		return HashCombine(GetTypeHash(static_cast<const FTypedElementHandle&>(InElement)), GetTypeHash(InElement.InterfacePtr));
+	}
+
+	FORCEINLINE void Private_InitializeNoRef(const FTypedHandleTypeId InTypeId, const FTypedHandleElementId InElementId, const FTypedElementInternalData& InData, BaseInterfaceType* InInterfacePtr)
+	{
+		FTypedElementHandle::Private_InitializeNoRef(InTypeId, InElementId, InData);
+		InterfacePtr = InInterfacePtr;
+	}
+
+	FORCEINLINE void Private_InitializeAddRef(const FTypedHandleTypeId InTypeId, const FTypedHandleElementId InElementId, const FTypedElementInternalData& InData, BaseInterfaceType* InInterfacePtr)
+	{
+		FTypedElementHandle::Private_InitializeAddRef(InTypeId, InElementId, InData);
+		InterfacePtr = InInterfacePtr;
+	}
+
+	FORCEINLINE void Private_DestroyNoRef()
+	{
+		FTypedElementHandle::Private_DestroyNoRef();
+		InterfacePtr = nullptr;
+	}
+
+	FORCEINLINE void Private_DestroyReleaseRef()
+	{
+		FTypedElementHandle::Private_DestroyReleaseRef();
+		InterfacePtr = nullptr;
+	}
+
+protected:
+	BaseInterfaceType* InterfacePtr = nullptr;
+};
+
+/**
+ * A combination of an element handle and its associated element interface.
+ * @note This should be specialized for top-level element interfaces to include their interface API.
+ * @note Elements auto-release on destruction.
+ */
+template <typename BaseInterfaceType>
+struct TTypedElement : public TTypedElementBase<BaseInterfaceType>
+{
+};
+using FTypedElement = TTypedElement<UTypedElementInterface>;
+
+template <typename OtherInterfaceType, typename ThisInterfaceType>
+FORCEINLINE TTypedElement<OtherInterfaceType> CastTypedElement(const TTypedElement<ThisInterfaceType>& InElement)
+{
+	static_assert(sizeof(TTypedElement<ThisInterfaceType>) == sizeof(TTypedElement<OtherInterfaceType>), "All TTypedElement instances must be the same size for this cast implementation to work!");
+	if (InElement && InElement.template HasInterface<OtherInterfaceType>())
+	{
+		return reinterpret_cast<const TTypedElement<OtherInterfaceType>&>(InElement);
+	}
+	return TTypedElement<OtherInterfaceType>();
+}
+
+template <typename OtherInterfaceType, typename ThisInterfaceType>
+FORCEINLINE void CastTypedElement(const TTypedElement<ThisInterfaceType>& InElement, TTypedElement<OtherInterfaceType>& OutCastElement)
+{
+	static_assert(sizeof(TTypedElement<ThisInterfaceType>) == sizeof(TTypedElement<OtherInterfaceType>), "All TTypedElement instances must be the same size for this cast implementation to work!");
+	if (InElement && InElement.template HasInterface<OtherInterfaceType>())
+	{
+		OutCastElement = reinterpret_cast<const TTypedElement<OtherInterfaceType>&>(InElement);
+	}
+	else
+	{
+		OutCastElement.Private_DestroyReleaseRef();
+	}
+}
+
+template <typename OtherInterfaceType, typename ThisInterfaceType>
+FORCEINLINE TTypedElement<OtherInterfaceType> CastTypedElement(TTypedElement<ThisInterfaceType>&& InElement)
+{
+	static_assert(sizeof(TTypedElement<ThisInterfaceType>) == sizeof(TTypedElement<OtherInterfaceType>), "All TTypedElement instances must be the same size for this cast implementation to work!");
+	if (InElement && InElement.template HasInterface<OtherInterfaceType>())
+	{
+		return reinterpret_cast<TTypedElement<OtherInterfaceType>&&>(InElement);
+	}
+	return TTypedElement<OtherInterfaceType>();
+}
+
+template <typename OtherInterfaceType, typename ThisInterfaceType>
+FORCEINLINE void CastTypedElement(TTypedElement<ThisInterfaceType>&& InElement, TTypedElement<OtherInterfaceType>& OutCastElement)
+{
+	static_assert(sizeof(TTypedElement<ThisInterfaceType>) == sizeof(TTypedElement<OtherInterfaceType>), "All TTypedElement instances must be the same size for this cast implementation to work!");
+	if (InElement && InElement.template HasInterface<OtherInterfaceType>())
+	{
+		OutCastElement = reinterpret_cast<TTypedElement<OtherInterfaceType>&&>(InElement);
+	}
+	else
+	{
+		OutCastElement.Private_DestroyReleaseRef();
+	}
+}
+
+template <typename OtherInterfaceType, typename ThisInterfaceType>
+FORCEINLINE TTypedElement<OtherInterfaceType> CastTypedElementChecked(const TTypedElement<ThisInterfaceType>& InElement)
+{
+	static_assert(sizeof(TTypedElement<ThisInterfaceType>) == sizeof(TTypedElement<OtherInterfaceType>), "All TTypedElement instances must be the same size for this cast implementation to work!");
+	if (InElement)
+	{
+		checkf(InElement.template HasInterface<OtherInterfaceType>(), TEXT("Element does not implement the required interface for this cast!"));
+		return reinterpret_cast<const TTypedElement<OtherInterfaceType>&>(InElement);
+	}
+	return TTypedElement<OtherInterfaceType>();
+}
+
+template <typename OtherInterfaceType, typename ThisInterfaceType>
+FORCEINLINE void CastTypedElementChecked(const TTypedElement<ThisInterfaceType>& InElement, TTypedElement<OtherInterfaceType>& OutCastElement)
+{
+	static_assert(sizeof(TTypedElement<ThisInterfaceType>) == sizeof(TTypedElement<OtherInterfaceType>), "All TTypedElement instances must be the same size for this cast implementation to work!");
+	if (InElement)
+	{
+		checkf(InElement.template HasInterface<OtherInterfaceType>(), TEXT("Element does not implement the required interface for this cast!"));
+		OutCastElement = reinterpret_cast<const TTypedElement<OtherInterfaceType>&>(InElement);
+	}
+	else
+	{
+		OutCastElement.Private_DestroyReleaseRef();
+	}
+}
+
+template <typename OtherInterfaceType, typename ThisInterfaceType>
+FORCEINLINE TTypedElement<OtherInterfaceType> CastTypedElementChecked(TTypedElement<ThisInterfaceType>&& InElement)
+{
+	static_assert(sizeof(TTypedElement<ThisInterfaceType>) == sizeof(TTypedElement<OtherInterfaceType>), "All TTypedElement instances must be the same size for this cast implementation to work!");
+	if (InElement)
+	{
+		checkf(InElement.template HasInterface<OtherInterfaceType>(), TEXT("Element does not implement the required interface for this cast!"));
+		return reinterpret_cast<TTypedElement<OtherInterfaceType>&&>(InElement);
+	}
+	return TTypedElement<OtherInterfaceType>();
+}
+
+template <typename OtherInterfaceType, typename ThisInterfaceType>
+FORCEINLINE void CastTypedElementChecked(TTypedElement<ThisInterfaceType>&& InElement, TTypedElement<OtherInterfaceType>& OutCastElement)
+{
+	static_assert(sizeof(TTypedElement<ThisInterfaceType>) == sizeof(TTypedElement<OtherInterfaceType>), "All TTypedElement instances must be the same size for this cast implementation to work!");
+	if (InElement)
+	{
+		checkf(InElement.template HasInterface<OtherInterfaceType>(), TEXT("Element does not implement the required interface for this cast!"));
+		OutCastElement = reinterpret_cast<TTypedElement<OtherInterfaceType>&&>(InElement);
+	}
+	else
+	{
+		OutCastElement.Private_DestroyReleaseRef();
+	}
+}
+
+/**
+ * A representation of the owner of an element that includes its mutable handle data.
+ * This type is returned when creating an element, and should be used to populate its internal payload data (if any).
+ * @note Owners do not auto-release on destruction, and must be manually destroyed via their owner element registry.
+ */
+template <typename ElementDataType>
+struct TTypedElementOwner
+{
+public:
+	TTypedElementOwner() = default;
+
+	TTypedElementOwner(const TTypedElementOwner&) = delete;
+	TTypedElementOwner& operator=(const TTypedElementOwner&) = delete;
+
+	TTypedElementOwner(TTypedElementOwner&& InOther)
+		: Id(MoveTemp(InOther.Id))
+		, DataPtr(InOther.DataPtr)
+	{
+		InOther.DataPtr = nullptr;
+		checkSlow(!InOther.IsSet());
+	}
+
+	TTypedElementOwner& operator=(TTypedElementOwner&& InOther)
+	{
+		if (this != &InOther)
+		{
+			Id = MoveTemp(InOther.Id);
+			DataPtr = InOther.DataPtr;
+
+			InOther.DataPtr = nullptr;
+			checkSlow(!InOther.IsSet());
+		}
+		return *this;
+	}
+
+	~TTypedElementOwner()
+	{
+		checkf(!IsSet(), TEXT("Element owner was still set during destruction! This will leak an element, and you should destroy this element prior to destruction!"));
+	}
+
+	FORCEINLINE explicit operator bool() const
+	{
+		return IsSet();
+	}
+
+	/**
+	 * Has this owner been initialized to a valid element?
+	 */
+	FORCEINLINE bool IsSet() const
+	{
+		return Id.IsSet();
+	}
+
+	/**
+	 * Get the ID that this element represents.
+	 */
+	FORCEINLINE const FTypedElementId& GetId() const
+	{
+		return Id;
+	}
+
+	/**
+	 * Attempt to access the mutable data stored within this owner, returning null if it isn't possible.
+	 */
+	template <typename U = ElementDataType, std::enable_if_t<!std::is_void<U>::value>* = nullptr>
+	FORCEINLINE U* GetData() const
+	{
+		static_assert(std::is_same<U, ElementDataType>::value, "Don't explicitly specify U!");
+		return DataPtr
+			? &DataPtr->GetMutableData()
+			: nullptr;
+	}
+
+	/**
+	 * Attempt to access the mutable data stored within this owner, asserting if it isn't possible.
+	 */
+	template <typename U = ElementDataType, std::enable_if_t<!std::is_void<U>::value>* = nullptr>
+	FORCEINLINE U& GetDataChecked() const
+	{
+		static_assert(std::is_same<U, ElementDataType>::value, "Don't explicitly specify U!");
+		checkf(DataPtr, TEXT("Handle data is null!"));
+		return DataPtr->GetMutableData();
+	}
+
+	/**
+	 * Acquire a copy of the ID that this element represents.
+	 * @note This must be paired with a call to ReleaseId.
+	 */
+	FTypedElementId AcquireId() const
+	{
+		FTypedElementId ElementId;
+		if (IsSet())
+		{
+			DataPtr->AddRef(/*bCanTrackReference*/false); // Cannot track element ID references as we have no space to store the reference ID
+			ElementId.Private_InitializeNoRef(Id.GetTypeId(), Id.GetElementId());
+		}
+		return ElementId;
+	}
+
+	/**
+	 * Release a copy of the ID that this element represents.
+	 * @note This should have come from a call to AcquireId.
+	 */
+	void ReleaseId(FTypedElementId& InOutElementId) const
+	{
+		checkf(InOutElementId == Id, TEXT("Element ID does not match this owner!"));
+		if (InOutElementId)
+		{
+			DataPtr->ReleaseRef(INDEX_NONE); // Cannot track element ID references as we have no space to store the reference ID
+			InOutElementId.Private_DestroyNoRef();
+		}
+	}
+
+	/**
+	 * Acquire a copy of the handle that this element represents.
+	 * @note This must be paired with a call to ReleaseHandle (or a call to Release on the handle instance).
+	 */
+	FTypedElementHandle AcquireHandle() const
+	{
+		FTypedElementHandle ElementHandle;
+		if (IsSet())
+		{
+			ElementHandle.Private_InitializeAddRef(Id.GetTypeId(), Id.GetElementId(), *DataPtr);
+		}
+		return ElementHandle;
+	}
+
+	/**
+	 * Release a copy of the handle that this element represents.
+	 * @note This should have come from a call to AcquireHandle.
+	 */
+	void ReleaseHandle(FTypedElementHandle& InOutElementHandle) const
+	{
+		checkf(InOutElementHandle.GetId() == Id, TEXT("Element handle ID does not match this owner!"));
+		InOutElementHandle.Release();
+	}
+
+	FORCEINLINE friend bool operator==(const TTypedElementOwner& InLHS, const TTypedElementOwner& InRHS)
+	{
+		return InLHS.Id == InRHS.Id;
+	}
+
+	FORCEINLINE friend bool operator!=(const TTypedElementOwner& InLHS, const TTypedElementOwner& InRHS)
+	{
+		return !(InLHS == InRHS);
+	}
+
+	friend inline uint32 GetTypeHash(const TTypedElementOwner& InElementOwner)
+	{
+		return GetTypeHash(InElementOwner.Id);
+	}
+
+	FORCEINLINE void Private_InitializeNoRef(const FTypedHandleTypeId InTypeId, const FTypedHandleElementId InElementId, TTypedElementInternalData<ElementDataType>& InData)
+	{
+		Id.Private_InitializeNoRef(InTypeId, InElementId);
+		DataPtr = &InData;
+	}
+
+	FORCEINLINE void Private_InitializeAddRef(const FTypedHandleTypeId InTypeId, const FTypedHandleElementId InElementId, TTypedElementInternalData<ElementDataType>& InData)
+	{
+		Private_InitializeNoRef(InTypeId, InElementId, InData);
+		RegisterRef();
+	}
+
+	FORCEINLINE void Private_DestroyNoRef()
+	{
+		Id.Private_DestroyNoRef();
+		DataPtr = nullptr;
+#if UE_TYPED_ELEMENT_HAS_REFTRACKING
+		ReferenceId = INDEX_NONE;
+#endif	// UE_TYPED_ELEMENT_HAS_REFTRACKING
+	}
+
+	FORCEINLINE void Private_DestroyReleaseRef()
+	{
+		UnregisterRef();
+		Private_DestroyNoRef();
+	}
+
+	FORCEINLINE const TTypedElementInternalData<ElementDataType>* Private_GetInternalData() const
+	{
+		return DataPtr;
+	}
+
+private:
+	FORCEINLINE void RegisterRef()
+	{
+		if (DataPtr)
+		{
+#if !UE_TYPED_ELEMENT_HAS_REFTRACKING
+			FTypedElementReferenceId ReferenceId = INDEX_NONE;
+#endif	// !UE_TYPED_ELEMENT_HAS_REFTRACKING
+			ReferenceId = DataPtr->AddRef(/*bCanTrackReference*/true);
+		}
+	}
+
+	FORCEINLINE void UnregisterRef()
+	{
+		if (DataPtr)
+		{
+#if !UE_TYPED_ELEMENT_HAS_REFTRACKING
+			FTypedElementReferenceId ReferenceId = INDEX_NONE;
+#endif	// !UE_TYPED_ELEMENT_HAS_REFTRACKING
+			DataPtr->ReleaseRef(ReferenceId);
+		}
+	}
+
+	FTypedElementId Id;
+	TTypedElementInternalData<ElementDataType>* DataPtr = nullptr;
+#if UE_TYPED_ELEMENT_HAS_REFTRACKING
+	FTypedElementReferenceId ReferenceId = INDEX_NONE;
+#endif	// UE_TYPED_ELEMENT_HAS_REFTRACKING
+};
+using FTypedElementOwner = TTypedElementOwner<void>;
+
+/** Script exposure for FTypedElementHandle. */
+UCLASS()
+class UTypedElementHandleLibrary : public UObject
+{
+	GENERATED_BODY()
+
+public:
+	/**
+	 * Has this handle been initialized to a valid element?
+	 */
+	UFUNCTION(BlueprintPure, Category="TypedElementFramework|Handle", meta=(ScriptMethod, ScriptOperator="bool"))
+	static bool IsSet(const FTypedElementHandle& ElementHandle)
+	{
+		return ElementHandle.IsSet();
+	}
+
+	/**
+	 * Release this handle and set it back to an empty state.
+	 */
+	UFUNCTION(BlueprintCallable, Category="TypedElementFramework|Handle", meta=(ScriptMethod))
+	static void Release(UPARAM(ref) FTypedElementHandle& ElementHandle)
+	{
+		ElementHandle.Release();
+	}
+
+	/**
+	 * Are these two handles equal?
+	 */
+	UFUNCTION(BlueprintPure, Category="TypedElementFramework|Handle", meta=(DisplayName="Equal (TypedElementHandle)", CompactNodeTitle="==", Keywords="== equal", ScriptMethod, ScriptOperator="=="))
+	static bool Equal(const FTypedElementHandle& LHS, const FTypedElementHandle& RHS)
+	{
+		return LHS == RHS;
+	}
+
+	/**
+	 * Are these two handles not equal?
+	 */
+	UFUNCTION(BlueprintPure, Category="TypedElementFramework|Handle", meta=(DisplayName="NotEqual (TypedElementHandle)", CompactNodeTitle="!=", Keywords="!= not equal", ScriptMethod, ScriptOperator="!="))
+	static bool NotEqual(const FTypedElementHandle& LHS, const FTypedElementHandle& RHS)
+	{
+		return LHS != RHS;
+	}
+};
diff --git a/Engine/Source/Runtime/TypedElementFramework/Public/Elements/Framework/TypedElementLimits.h b/Engine/Source/Runtime/TypedElementFramework/Public/Elements/Framework/TypedElementLimits.h
new file mode 100644
--- /dev/null
+++ b/Engine/Source/Runtime/TypedElementFramework/Public/Elements/Framework/TypedElementLimits.h
@@ -0,0 +1,39 @@
+// Copyright Epic Games, Inc. All Rights Reserved.
+
+#pragma once
+
+#include "CoreTypes.h"
+#include "Math/NumericLimits.h"
+
+#define UE_TYPED_ELEMENT_HAS_REFCOUNTING (1)
+#define UE_TYPED_ELEMENT_HAS_REFTRACKING (0)
+#define UE_TYPED_ELEMENT_HAS_REFERENCING (UE_TYPED_ELEMENT_HAS_REFCOUNTING || UE_TYPED_ELEMENT_HAS_REFTRACKING)
+
+/**
+ * Handle ID limits, as used by FTypedElementId.
+ * @note Limited to a combined 32-bits so that they can be used directly within render targets, though could be made 64-bits if the 
+ *       editor used 64-bit render targets (this would also require 64-bit container support in TTypedElementInternalDataStore).
+ */
+constexpr SIZE_T TypedHandleTypeIdBits = 8;
+constexpr SIZE_T TypedHandleElementIdBits = 24;
+
+constexpr SIZE_T TypedHandleTypeIdBytes = TypedHandleTypeIdBits >> 3;
+constexpr SIZE_T TypedHandleElementIdBytes = TypedHandleElementIdBits >> 3;
+
+constexpr SIZE_T TypedHandleMaxTypeId = ((SIZE_T)1 << TypedHandleTypeIdBits) - 1;
+constexpr SIZE_T TypedHandleMaxElementId = ((SIZE_T)1 << TypedHandleElementIdBits) - 1;
+
+constexpr SIZE_T TypedHandleRefTrackingSkipCount = 1;
+constexpr SIZE_T TypedHandleRefTrackingDepth = 31;
+
+using FTypedHandleTypeId = uint8;
+using FTypedHandleElementId = int32;
+using FTypedHandleCombinedId = uint32;
+
+using FTypedElementRefCount = int32;
+using FTypedElementReferenceId = int32;
+
+static_assert(sizeof(FTypedHandleCombinedId) >= (TypedHandleTypeIdBytes + TypedHandleElementIdBytes), "FTypedHandleCombinedId is not large enough to hold the combination of TypedHandleTypeIdBytes and TypedHandleElementIdBytes!");
+
+static_assert(TNumericLimits<FTypedHandleTypeId>::Max() >= TypedHandleMaxTypeId, "FTypedHandleTypeId is not large enough to hold TypedHandleMaxTypeId!");
+static_assert(TNumericLimits<FTypedHandleElementId>::Max() >= TypedHandleMaxElementId, "FTypedHandleElementId is not large enough to hold TypedHandleMaxElementId!");
diff --git a/Engine/Source/Runtime/TypedElementFramework/Public/Elements/Framework/TypedElementList.h b/Engine/Source/Runtime/TypedElementFramework/Public/Elements/Framework/TypedElementList.h
new file mode 100644
--- /dev/null
+++ b/Engine/Source/Runtime/TypedElementFramework/Public/Elements/Framework/TypedElementList.h
@@ -0,0 +1,661 @@
+// Copyright Epic Games, Inc. All Rights Reserved.
+
+#pragma once
+
+#include "CoreMinimal.h"
+#include "Delegates/Delegate.h"
+#include "Delegates/DelegateCombinations.h"
+#include "Templates/SubclassOf.h"
+#include "Containers/ArrayView.h"
+#include "Templates/UniquePtr.h"
+#include "Elements/Framework/TypedElementHandle.h"
+#include "TypedElementList.generated.h"
+
+class UTypedElementList;
+class UTypedElementRegistry;
+
+namespace TypedElementList_Private
+{
+
+TYPEDELEMENTFRAMEWORK_API void GetElementImpl(const UTypedElementRegistry* InRegistry, const FTypedElementHandle& InElementHandle, const UClass* InBaseInterfaceType, FTypedElement& OutElement);
+
+template <typename BaseInterfaceType>
+FORCEINLINE void GetElement(const UTypedElementRegistry* InRegistry, const FTypedElementHandle& InElementHandle, TTypedElement<BaseInterfaceType>& OutElement)
+{
+	static_assert(sizeof(TTypedElement<BaseInterfaceType>) == sizeof(FTypedElement), "All TTypedElement instances must be the same size for this cast implementation to work!");
+	GetElementImpl(InRegistry, InElementHandle, BaseInterfaceType::StaticClass(), reinterpret_cast<FTypedElement&>(OutElement));
+}
+
+template <typename BaseInterfaceType>
+FORCEINLINE TTypedElement<BaseInterfaceType> GetElement(const UTypedElementRegistry* InRegistry, const FTypedElementHandle& InElementHandle)
+{
+	TTypedElement<BaseInterfaceType> Element;
+	GetElement(InRegistry, InElementHandle, Element);
+	return Element;
+}
+
+/**
+ * Iterator type for ranged-for loops which only returns handles that implement the given interface.
+ * @note This iterator type only supports the minimal functionality needed to support C++ ranged-for syntax.
+ */
+template <typename BaseInterfaceType>
+struct TTypedElementListInterfaceIterator
+{
+public:
+	TTypedElementListInterfaceIterator(const UTypedElementRegistry* InRegistry, const TArray<FTypedElementHandle>& InElementHandles, const int32 InIndex)
+		: Registry(InRegistry)
+		, ElementHandles(InElementHandles)
+		, InitialNum(InElementHandles.Num())
+		, Index(InIndex)
+	{
+		// Walk to the first valid element
+		while (Index < ElementHandles.Num())
+		{
+			TypedElementList_Private::GetElement(Registry, ElementHandles[Index], Element);
+			if (Element)
+			{
+				break;
+			}
+			++Index;
+		}
+	}
+
+	FORCEINLINE const TTypedElement<BaseInterfaceType>& operator*() const
+	{
+		return Element;
+	}
+
+	TTypedElementListInterfaceIterator& operator++()
+	{
+		// Walk to the next valid element
+		while (++Index < ElementHandles.Num())
+		{
+			TypedElementList_Private::GetElement(Registry, ElementHandles[Index], Element);
+			if (Element)
+			{
+				break;
+			}
+		}
+		return *this;
+	}
+
+private:
+	const UTypedElementRegistry* Registry;
+	const TArray<FTypedElementHandle>& ElementHandles;
+	const int32 InitialNum;
+	int32 Index;
+	TTypedElement<BaseInterfaceType> Element;
+
+	FORCEINLINE friend bool operator!=(const TTypedElementListInterfaceIterator& Lhs, const TTypedElementListInterfaceIterator& Rhs)
+	{
+		// We only need to do the check in this operator, because no other operator will be
+		// called until after this one returns.
+		//
+		// Also, we should only need to check one side of this comparison - if the other iterator isn't
+		// even from the same array then the compiler has generated bad code.
+		ensureMsgf(Lhs.ElementHandles.Num() == Lhs.InitialNum, TEXT("Array has changed during ranged-for iteration!"));
+		return Lhs.Index != Rhs.Index;
+	}
+};
+
+/**
+ * Proxy type to enable ranged-for loops which only returns handles that implement the given interface.
+ */
+template <typename BaseInterfaceType>
+struct TTypedElementListInterfaceIteratorProxy
+{
+public:
+	TTypedElementListInterfaceIteratorProxy(const UTypedElementRegistry* InRegistry, const TArray<FTypedElementHandle>& InElementHandles)
+		: Registry(InRegistry)
+		, ElementHandles(InElementHandles)
+	{
+	}
+
+	TTypedElementListInterfaceIteratorProxy(const TTypedElementListInterfaceIteratorProxy&) = delete;
+	TTypedElementListInterfaceIteratorProxy& operator=(const TTypedElementListInterfaceIteratorProxy&) = delete;
+
+	TTypedElementListInterfaceIteratorProxy(TTypedElementListInterfaceIteratorProxy&&) = default;
+	TTypedElementListInterfaceIteratorProxy& operator=(TTypedElementListInterfaceIteratorProxy&&) = default;
+
+	/**
+	 * DO NOT USE DIRECTLY
+	 * STL-like iterators to enable range-based for loop support.
+	 */
+	using RangedForConstIteratorType = TTypedElementListInterfaceIterator<BaseInterfaceType>;
+	FORCEINLINE RangedForConstIteratorType begin() const { return RangedForConstIteratorType(Registry, ElementHandles, 0); }
+	FORCEINLINE RangedForConstIteratorType end() const { return RangedForConstIteratorType(Registry, ElementHandles, ElementHandles.Num()); }
+
+private:
+	const UTypedElementRegistry* Registry;
+	const TArray<FTypedElementHandle>& ElementHandles;
+};
+
+} // namespace TypedElementList_Private
+
+/**
+ * Interface to allow external systems (such as USelection) to receive immediate sync notifications as an element list is changed.
+ * This exists purely as a bridging mechanism and shouldn't be relied on for new code. It is lazily created as needed.
+ */
+class TYPEDELEMENTFRAMEWORK_API FTypedElementListLegacySync
+{
+public:
+	enum class ESyncType : uint8
+	{
+		/**
+		 * An element was added to the element list.
+		 * The ElementHandle argument will be set to the element that was added.
+		 */
+		Added,
+
+		/**
+		 * An element was removed from the element list.
+		 * The ElementHandle argument will be set to the element that was removed.
+		 */
+		Removed,
+
+		/**
+		 * The element list was modified in an unknown way.
+		 * The ElementHandle argument will be unset.
+		 */
+		Modified,
+
+		/**
+		 * The element list was cleared.
+		 * The ElementHandle argument will be unset.
+		 */
+		Cleared,
+
+		/**
+		 * The element list was modified as part of a batch or bulk operation.
+		 * The ElementHandle argument will be unset.
+		 * @note A batch operation will emit internal (bIsWithinBatchOperation=true) Added, Removed, Modified and Cleared updates during the batch, 
+		 *       so if you respond to those internal updates you may choose to ignore this one. Otherwise you should treat it the same as Modified.
+		 */
+		BatchComplete,
+	};
+	
+	FTypedElementListLegacySync(const UTypedElementList* InElementList);
+
+	void Private_EmitSyncEvent(const ESyncType InSyncType, const FTypedElementHandle& InElementHandle = FTypedElementHandle());
+
+	DECLARE_EVENT_FourParams(FTypedElementListLegacySync, FOnSyncEvent, const UTypedElementList* /*InElementList*/, ESyncType /*InSyncType*/, const FTypedElementHandle& /*InElementHandle*/, bool /*bIsWithinBatchOperation*/);
+	FOnSyncEvent& OnSyncEvent();
+
+	bool IsRunningBatchOperation() const;
+	void BeginBatchOperation();
+	void EndBatchOperation(const bool InNotify = true);
+	bool IsBatchOperationDirty() const;
+	void ForceBatchOperationDirty();
+
+private:
+	const UTypedElementList* ElementList;
+
+	FOnSyncEvent OnSyncEventDelegate;
+
+	int32 NumOpenBatchOperations = 0;
+	bool bBatchOperationIsDirty = false;
+};
+
+/**
+ * Helper to batch immediate sync notifications for legacy code.
+ * Does nothing if no legacy sync has been created for the given instance.
+ */
+class TYPEDELEMENTFRAMEWORK_API FTypedElementListLegacySyncScopedBatch
+{
+public:
+	explicit FTypedElementListLegacySyncScopedBatch(UTypedElementList* InElementList, const bool InNotify = true);
+	~FTypedElementListLegacySyncScopedBatch();
+
+	FTypedElementListLegacySyncScopedBatch(const FTypedElementListLegacySyncScopedBatch&) = delete;
+	FTypedElementListLegacySyncScopedBatch& operator=(const FTypedElementListLegacySyncScopedBatch&) = delete;
+
+	FTypedElementListLegacySyncScopedBatch(FTypedElementListLegacySyncScopedBatch&&) = delete;
+	FTypedElementListLegacySyncScopedBatch& operator=(FTypedElementListLegacySyncScopedBatch&&) = delete;
+
+private:
+	FTypedElementListLegacySync* ElementListLegacySync = nullptr;
+	bool bNotify = true;
+};
+
+/**
+ * A list of element handles.
+ * Provides high-level access to groups of elements, including accessing elements that implement specific interfaces.
+ */
+UCLASS(Transient)
+class TYPEDELEMENTFRAMEWORK_API UTypedElementList : public UObject
+{
+	GENERATED_BODY()
+
+public:
+	//~ UObject interface
+	virtual void BeginDestroy() override;
+
+	/**
+	 * Internal function used by the element registry to create an element list instance.
+	 */
+	static UTypedElementList* Private_CreateElementList(UTypedElementRegistry* InRegistry);
+
+	/**
+	 * Clone this list instance.
+	 * @note Only copies elements; does not copy any bindings!
+	 */
+	UFUNCTION(BlueprintCallable, BlueprintPure=false, Category="TypedElementFramework|List")
+	UTypedElementList* Clone() const;
+
+	/**
+	 * Get the element handle at the given index.
+	 * @note Use IsValidIndex to test for validity.
+	 */
+	FORCEINLINE FTypedElementHandle operator[](const int32 InIndex) const
+	{
+		return GetElementHandleAt(InIndex);
+	}
+
+	/**
+	 * Get the element handle at the given index.
+	 * @note Use IsValidIndex to test for validity.
+	 */
+	UFUNCTION(BlueprintPure, Category="TypedElementFramework|List")
+	FORCEINLINE FTypedElementHandle GetElementHandleAt(const int32 InIndex) const
+	{
+		return ElementHandles[InIndex];
+	}
+
+	/**
+	 * Get the element at the given index.
+	 * @note Use IsValidIndex to test for validity.
+	 */
+	template <typename BaseInterfaceType>
+	FORCEINLINE TTypedElement<BaseInterfaceType> GetElementAt(const int32 InIndex) const
+	{
+		return GetElement<BaseInterfaceType>(GetElementHandleAt(InIndex));
+	}
+
+	/**
+	 * Get the element at the given index.
+	 * @note Use IsValidIndex to test for validity.
+	 */
+	template <typename BaseInterfaceType>
+	FORCEINLINE void GetElementAt(const int32 InIndex, TTypedElement<BaseInterfaceType>& OutElement) const
+	{
+		GetElement(GetElementHandleAt(InIndex), OutElement);
+	}
+
+	/**
+	 * Get the element from the given handle.
+	 */
+	template <typename BaseInterfaceType>
+	FORCEINLINE TTypedElement<BaseInterfaceType> GetElement(const FTypedElementHandle& InElementHandle) const
+	{
+		return TypedElementList_Private::GetElement<BaseInterfaceType>(Registry.Get(), InElementHandle);
+	}
+
+	/**
+	 * Get the element from the given handle.
+	 */
+	template <typename BaseInterfaceType>
+	FORCEINLINE void GetElement(const FTypedElementHandle& InElementHandle, TTypedElement<BaseInterfaceType>& OutElement) const
+	{
+		TypedElementList_Private::GetElement(Registry.Get(), InElementHandle, OutElement);
+	}
+
+	/**
+	 * Get the element interface from the given handle.
+	 */
+	template <typename BaseInterfaceType>
+	BaseInterfaceType* GetElementInterface(const FTypedElementHandle& InElementHandle) const
+	{
+		return static_cast<BaseInterfaceType*>(GetElementInterface(InElementHandle, BaseInterfaceType::StaticClass()));
+	}
+
+	/**
+	 * Get the element interface from the given handle.
+	 */
+	UFUNCTION(BlueprintPure, Category="TypedElementFramework|List")
+	UTypedElementInterface* GetElementInterface(const FTypedElementHandle& InElementHandle, const TSubclassOf<UTypedElementInterface>& InBaseInterfaceType) const;
+
+	/**
+	 * Is the given index a valid entry within this element list?
+	 */
+	UFUNCTION(BlueprintPure, Category="TypedElementFramework|List")
+	FORCEINLINE bool IsValidIndex(const int32 InIndex) const
+	{
+		return ElementHandles.IsValidIndex(InIndex);
+	}
+
+	/**
+	 * Get the number of entries within this element list.
+	 */
+	UFUNCTION(BlueprintPure, Category="TypedElementFramework|List")
+	FORCEINLINE int32 Num() const
+	{
+		return ElementHandles.Num();
+	}
+
+	/**
+	 * Shrink this element list storage to avoid slack.
+	 */
+	UFUNCTION(BlueprintCallable, Category="TypedElementFramework|List")
+	FORCEINLINE void Shrink()
+	{
+		ElementCombinedIds.Shrink();
+		ElementHandles.Shrink();
+	}
+
+	/**
+	 * Pre-allocate enough memory in this element list to store the given number of entries.
+	 */
+	UFUNCTION(BlueprintCallable, Category="TypedElementFramework|List")
+	FORCEINLINE void Reserve(const int32 InSize)
+	{
+		ElementCombinedIds.Reserve(InSize);
+		ElementHandles.Reserve(InSize);
+	}
+
+	/**
+	 * Remove all entries from this element list, potentially leaving space allocated for the given number of entries.
+	 */
+	UFUNCTION(BlueprintCallable, Category="TypedElementFramework|List")
+	FORCEINLINE void Empty(const int32 InSlack = 0)
+	{
+		NoteListMayChange();
+		ElementCombinedIds.Empty(InSlack);
+		ElementHandles.Empty(InSlack);
+		NoteListChanged(EChangeType::Cleared);
+	}
+
+	/**
+	 * Remove all entries from this element list, preserving existing allocations.
+	 */
+	UFUNCTION(BlueprintCallable, Category="TypedElementFramework|List")
+	FORCEINLINE void Reset()
+	{
+		NoteListMayChange();
+		ElementCombinedIds.Reset();
+		ElementHandles.Reset();
+		NoteListChanged(EChangeType::Cleared);
+	}
+
+	/**
+	 * Does this element list contain an entry for the given element ID?
+	 */
+	FORCEINLINE bool Contains(const FTypedElementId& InElementId) const
+	{
+		return ContainsElementImpl(InElementId);
+	}
+
+	/**
+	 * Does this element list contain an entry for the given element handle?
+	 */
+	UFUNCTION(BlueprintPure, Category="TypedElementFramework|List")
+	FORCEINLINE bool Contains(const FTypedElementHandle& InElementHandle) const
+	{
+		return ContainsElementImpl(InElementHandle.GetId());
+	}
+
+	/**
+	 * Does this element list contain an entry for the given element owner?
+	 */
+	template <typename ElementDataType>
+	FORCEINLINE bool Contains(const TTypedElementOwner<ElementDataType>& InElementOwner)
+	{
+		return ContainsElementImpl(InElementOwner.GetId());
+	}
+
+	/**
+	 * Add the given element handle to this element list, if it isn't already in the list.
+	 * @return True if the element handle was added, false if it is already in the list.
+	 */
+	UFUNCTION(BlueprintCallable, Category="TypedElementFramework|List")
+	FORCEINLINE bool Add(const FTypedElementHandle& InElementHandle)
+	{
+		return AddElementImpl(CopyTemp(InElementHandle));
+	}
+
+	/**
+	 * Add the given element handle to this element list, if it isn't already in the list.
+	 * @return True if the element handle was added, false if it is already in the list.
+	 */
+	FORCEINLINE bool Add(FTypedElementHandle&& InElementHandle)
+	{
+		return AddElementImpl(MoveTemp(InElementHandle));
+	}
+
+	/**
+	 * Add the given element owner to this element list, if it isn't already in the list.
+	 * @return True if the element owner was added, false if it is already in the list.
+	 */
+	template <typename ElementDataType>
+	FORCEINLINE bool Add(const TTypedElementOwner<ElementDataType>& InElementOwner)
+	{
+		return AddElementImpl(InElementOwner.AcquireHandle());
+	}
+
+	/**
+	 * Append the given element handles to this element list, for any that already in the list.
+	 */
+	UFUNCTION(BlueprintCallable, Category="TypedElementFramework|List")
+	FORCEINLINE void Append(const TArray<FTypedElementHandle>& InElementHandles)
+	{
+		Append(MakeArrayView(InElementHandles));
+	}
+	
+	/**
+	 * Append the given element handles to this element list, for any that already in the list.
+	 */
+	void Append(TArrayView<const FTypedElementHandle> InElementHandles)
+	{
+		FTypedElementListLegacySyncScopedBatch LegacySyncBatch(this);
+
+		Reserve(Num() + InElementHandles.Num());
+		for (const FTypedElementHandle& ElementHandle : InElementHandles)
+		{
+			AddElementImpl(CopyTemp(ElementHandle));
+		}
+	}
+
+	/**
+	 * Append the given element owners to this element list, for any that already in the list.
+	 */
+	template <typename ElementDataType>
+	FORCEINLINE void Append(const TArray<TTypedElementOwner<ElementDataType>>& InElementOwners)
+	{
+		Append(MakeArrayView(InElementOwners));
+	}
+
+	/**
+	 * Append the given element owners to this element list, for any that already in the list.
+	 */
+	template <typename ElementDataType>
+	void Append(TArrayView<const TTypedElementOwner<ElementDataType>> InElementOwners)
+	{
+		FTypedElementListLegacySyncScopedBatch LegacySyncBatch(this);
+
+		Reserve(Num() + InElementOwners.Num());
+		for (const TTypedElementOwner<ElementDataType>& ElementOwner : InElementOwners)
+		{
+			AddElementImpl(ElementOwner.AcquireHandle());
+		}
+	}
+
+	/**
+	 * Remove the given element ID from this element list, if it is in the list.
+	 * @return True if the element ID was removed, false if it isn't in the list.
+	 */
+	FORCEINLINE bool Remove(const FTypedElementId& InElementId)
+	{
+		return RemoveElementImpl(InElementId);
+	}
+
+	/**
+	 * Remove the given element handle from this element list, if it is in the list.
+	 * @return True if the element handle was removed, false if it isn't in the list.
+	 */
+	UFUNCTION(BlueprintCallable, Category="TypedElementFramework|List")
+	FORCEINLINE bool Remove(const FTypedElementHandle& InElementHandle)
+	{
+		return RemoveElementImpl(InElementHandle.GetId());
+	}
+
+	/**
+	 * Remove the given element owner from this element list, if it is in the list.
+	 * @return True if the element owner was removed, false if it isn't in the list.
+	 */
+	template <typename ElementDataType>
+	FORCEINLINE bool Remove(const TTypedElementOwner<ElementDataType>& InElementOwner)
+	{
+		return RemoveElementImpl(InElementOwner.GetId());
+	}
+
+	/**
+	 * Remove any element handles that match the given predicate from this element list.
+	 * @return The number of element handles removed.
+	 */
+	FORCEINLINE int32 RemoveAll(TFunctionRef<bool(const FTypedElementHandle&)> InPredicate)
+	{
+		return RemoveAllElementsImpl(InPredicate);
+	}
+
+	/**
+	 * Remove any elements that match the given predicate from this element list.
+	 * @return The number of elements removed.
+	 */
+	template <typename BaseInterfaceType>
+	int32 RemoveAll(TFunctionRef<bool(const TTypedElement<BaseInterfaceType>&)> InPredicate)
+	{
+		TTypedElement<BaseInterfaceType> TempElement;
+		return RemoveAllElementsImpl([this, &TempElement, &InPredicate](const FTypedElementHandle& InElementHandle)
+		{
+			GetElement(InElementHandle, TempElement);
+			return TempElement && InPredicate(TempElement);
+		});
+	}
+
+	/**
+	 * Access the delegate that is invoked whenever this element list is potentially about to change.
+	 * @note This may be called even if no actual change happens, though once a change does happen it won't be called again until after the next call to NotifyPendingChanges.
+	 */
+	DECLARE_EVENT_OneParam(UTypedElementList, FOnPreChange, const UTypedElementList* /*InElementList*/);
+	FOnPreChange& OnPreChange()
+	{
+		return OnPreChangeDelegate;
+	}
+
+	/**
+	 * Access the delegate that is invoked whenever this element list has been changed.
+	 * @note This is called automatically at the end of each frame, but can also be manually invoked by NotifyPendingChanges.
+	 */
+	DECLARE_EVENT_OneParam(UTypedElementList, FOnChanged, const UTypedElementList* /*InElementList*/);
+	FOnChanged& OnChanged()
+	{
+		return OnChangedDelegate;
+	}
+
+	/**
+	 * Invoke the delegate called whenever this element list has been changed.
+	 */
+	void NotifyPendingChanges();
+
+	/**
+	 * Clear whether there are pending changes for OnChangedDelegate to notify for, without emitting a notification.
+	 */
+	void ClearPendingChanges();
+
+	/**
+	 * Access the interface to allow external systems (such as USelection) to receive immediate sync notifications as an element list is changed.
+	 * This exists purely as a bridging mechanism and shouldn't be relied on for new code. It is lazily created as needed.
+	 */
+	FTypedElementListLegacySync& Legacy_GetSync();
+
+	/**
+	 * Access the interface to allow external systems (such as USelection) to receive immediate sync notifications as an element list is changed.
+	 * This exists purely as a bridging mechanism and shouldn't be relied on for new code. This will return null if no legacy sync has been created for this instance.
+	 */
+	FTypedElementListLegacySync* Legacy_GetSyncPtr();
+
+	/**
+	 * Get an iterator for elements that implement the given element interface.
+	 */
+	template <typename BaseInterfaceType>
+	FORCEINLINE TypedElementList_Private::TTypedElementListInterfaceIteratorProxy<BaseInterfaceType> IterateInterface() const
+	{
+		return TypedElementList_Private::TTypedElementListInterfaceIteratorProxy<BaseInterfaceType>(Registry.Get(), ElementHandles);
+	}
+
+	/**
+	 * DO NOT USE DIRECTLY
+	 * STL-like iterators to enable range-based for loop support.
+	 */
+	FORCEINLINE TArray<FTypedElementHandle>::RangedForConstIteratorType begin() const { return ElementHandles.begin(); }
+	FORCEINLINE TArray<FTypedElementHandle>::RangedForConstIteratorType end() const { return ElementHandles.end(); }
+
+private:
+	enum class EChangeType : uint8
+	{
+		/**
+		* An element was added to the element list.
+		* The ElementHandle argument will be set to the element that was added.
+		*/
+		Added,
+
+		/**
+		* An element was removed from the element list.
+		* The ElementHandle argument will be set to the element that was removed.
+		*/
+		Removed,
+
+		/**
+		* The element list was cleared.
+		* The ElementHandle argument will be unset.
+		*/
+		Cleared,
+	};
+
+	void Initialize(UTypedElementRegistry* InRegistry);
+
+	bool AddElementImpl(FTypedElementHandle&& InElementHandle);
+	bool RemoveElementImpl(const FTypedElementId& InElementId);
+	int32 RemoveAllElementsImpl(TFunctionRef<bool(const FTypedElementHandle&)> InPredicate);
+	bool ContainsElementImpl(const FTypedElementId& InElementId) const;
+
+	void NoteListMayChange();
+	void NoteListChanged(const EChangeType InChangeType, const FTypedElementHandle& InElementHandle = FTypedElementHandle());
+
+	/**
+	 * Element registry this element list is associated with.
+	 */
+	TWeakObjectPtr<UTypedElementRegistry> Registry;
+
+	/**
+	 * Set of combined ID values that are currently present in this element list.
+	 * Used to perform optimized querying of which elements are in this list, and to avoid adding duplicate entries.
+	 */
+	TSet<FTypedHandleCombinedId> ElementCombinedIds;
+
+	/**
+	 * Array of element handles present in this element list.
+	 * These are stored in the same order that they are added, and the set above can be used to optimize certain queries.
+	 */
+	TArray<FTypedElementHandle> ElementHandles;
+
+	/**
+	 * Delegate that is invoked whenever this element list is potentially about to change.
+	 */
+	FOnPreChange OnPreChangeDelegate;
+
+	/**
+	 * Delegate that is invoked whenever this element list has been changed.
+	 */
+	FOnChanged OnChangedDelegate;
+
+	/**
+	 * Whether there are pending changes for OnChangedDelegate to notify for.
+	 */
+	bool bHasPendingNotify = false;
+
+	/**
+	 * Interface to allow external systems (such as USelection) to receive immediate sync notifications as an element list is changed.
+	 * This exists purely as a bridging mechanism and shouldn't be relied on for new code. It is lazily created as needed.
+	 */
+	TUniquePtr<FTypedElementListLegacySync> LegacySync;
+};
diff --git a/Engine/Source/Runtime/TypedElementFramework/Public/Elements/Framework/TypedElementOwnerStore.h b/Engine/Source/Runtime/TypedElementFramework/Public/Elements/Framework/TypedElementOwnerStore.h
new file mode 100644
--- /dev/null
+++ b/Engine/Source/Runtime/TypedElementFramework/Public/Elements/Framework/TypedElementOwnerStore.h
@@ -0,0 +1,217 @@
+// Copyright Epic Games, Inc. All Rights Reserved.
+
+#pragma once
+
+#include "CoreMinimal.h"
+#include "Misc/ScopeLock.h"
+#include "Elements/Framework/TypedElementHandle.h"
+
+/**
+ * Type to keep a lock on the TTypedElementOwnerStore for the duration 
+ * of time that a TTypedElementOwner is being externally referenced.
+ * This avoids invalid memory access if the map is reallocated.
+ */
+template <typename ElementOwnerType>
+struct TTypedElementOwnerScopedAccessImpl
+{
+public:
+	explicit TTypedElementOwnerScopedAccessImpl(FCriticalSection* InSynchObject)
+		: SynchObject(InSynchObject)
+	{
+		check(SynchObject);
+		SynchObject->Lock();
+	}
+
+	~TTypedElementOwnerScopedAccessImpl()
+	{
+		if (SynchObject)
+		{
+			SynchObject->Unlock();
+		}
+	}
+
+	TTypedElementOwnerScopedAccessImpl(const TTypedElementOwnerScopedAccessImpl&) = delete;
+	TTypedElementOwnerScopedAccessImpl& operator=(const TTypedElementOwnerScopedAccessImpl&) = delete;
+
+	TTypedElementOwnerScopedAccessImpl(TTypedElementOwnerScopedAccessImpl&& InOther)
+		: SynchObject(InOther.SynchObject)
+		, ElementOwner(InOther.ElementOwner)
+	{
+		InOther.SynchObject = nullptr;
+		InOther.ElementOwner = nullptr;
+	}
+
+	TTypedElementOwnerScopedAccessImpl& operator=(TTypedElementOwnerScopedAccessImpl&&) = delete;
+
+	FORCEINLINE explicit operator bool() const
+	{
+		return IsSet();
+	}
+
+	FORCEINLINE bool IsSet() const
+	{
+		return ElementOwner != nullptr;
+	}
+
+	FORCEINLINE ElementOwnerType& operator*() const
+	{
+		check(ElementOwner);
+		return *ElementOwner;
+	}
+
+	FORCEINLINE ElementOwnerType* operator->() const
+	{
+		check(ElementOwner);
+		return ElementOwner;
+	}
+
+	void Private_SetElementOwner(ElementOwnerType* InElementOwner)
+	{
+		checkf(!ElementOwner, TEXT("Element owner was already set!"));
+		ElementOwner = InElementOwner;
+
+		// If we didn't find an element then we can unlock immediately
+		// Otherwise we hold the lock for as long as the element is being read to avoid a map write from invalidating the reference from reallocation
+		check(SynchObject);
+		if (!ElementOwner)
+		{
+			SynchObject->Unlock();
+			SynchObject = nullptr;
+		}
+	}
+
+private:
+	FCriticalSection* SynchObject = nullptr;
+	ElementOwnerType* ElementOwner = nullptr;
+};
+
+template <typename ElementDataType>
+using TTypedElementOwnerScopedAccess = TTypedElementOwnerScopedAccessImpl<const TTypedElementOwner<ElementDataType>>;
+
+template <typename ElementDataType>
+using TTypedElementOwnerScopedMutableAccess = TTypedElementOwnerScopedAccessImpl<TTypedElementOwner<ElementDataType>>;
+
+template <typename KeyDataType>
+inline FString GetTypedElementOwnerStoreKeyDebugString(const KeyDataType& InKey)
+{
+	return TEXT("<GetTypedElementOwnerStoreKeyDebugString not implemented for key type>");
+}
+
+inline FString GetTypedElementOwnerStoreKeyDebugString(const UObject* InKey)
+{
+	return InKey->GetPathName();
+}
+
+/**
+ * A store of element owners tied to their corresponding owner instance.
+ * This can be used to track element owners for an instance without having to add new data to the type itself.
+ */
+template <typename ElementDataType, typename KeyDataType = const void*>
+struct TTypedElementOwnerStore
+{
+public:
+	static_assert(TNumericLimits<int32>::Max() >= TypedHandleMaxElementId, "TTypedElementOwnerStore internally uses TMap so cannot store TypedHandleMaxElementId! Consider making this container 64-bit aware, or explicitly remove this compile time check.");
+
+	TTypedElementOwnerStore() = default;
+
+	~TTypedElementOwnerStore()
+	{
+		checkf(ElementOwnerMap.Num() == 0, TEXT("Element owners were still registered during destruction! This will leak elements, and you should unregister and destroy all elements prior to destruction!"));
+	}
+
+	/**
+	 * Register the given element owner, transferring ownership of the owner instance to this store.
+	 * @note This will lock the store for writes until the scoped access struct is no longer being used.
+	 * @note This must be paired with a call to UnregisterElementOwner.
+	 */
+	TTypedElementOwnerScopedAccess<ElementDataType> RegisterElementOwner(const KeyDataType& InKey, TTypedElementOwner<ElementDataType>&& InElementOwner)
+	{
+		TTypedElementOwnerScopedAccess<ElementDataType> ScopedAccess(&ElementOwnerMapCS);
+		checkf(!ElementOwnerMap.Contains(InKey), TEXT("Element owner has already been registered for this key (%s)! This will leak elements, and you should unregister and destroy the old element owner for this instance before adding a new one!"), *GetTypedElementOwnerStoreKeyDebugString(InKey));
+		ScopedAccess.Private_SetElementOwner(&ElementOwnerMap.Add(InKey, MoveTemp(InElementOwner)));
+		return ScopedAccess;
+	}
+
+	/**
+	 * Unregister the given element owner, transferring ownership of the owner instance back to the caller.
+	 * @note The returned owner may be unset if the key was not registered.
+	 */
+	TTypedElementOwner<ElementDataType> UnregisterElementOwner(const KeyDataType& InKey)
+	{
+		TTypedElementOwner<ElementDataType> ElementOwner;
+		{
+			FScopeLock Lock(&ElementOwnerMapCS);
+			ElementOwnerMap.RemoveAndCopyValue(InKey, ElementOwner);
+		}
+		return ElementOwner;
+	}
+
+	/**
+	 * Unregister any given element owners that match the given predicate, transferring ownership of the owner instance to the given callback.
+	 */
+	void UnregisterElementOwners(TFunctionRef<bool(const TTypedElementOwner<ElementDataType>&)> InShouldUnregisterElement, TFunctionRef<void(TTypedElementOwner<ElementDataType>&&)> InOnUnregisterElement)
+	{
+		FScopeLock Lock(&ElementOwnerMapCS);
+		for (auto It = ElementOwnerMap.CreateIterator(); It; ++It)
+		{
+			if (InShouldUnregisterElement(It->Value))
+			{
+				InOnUnregisterElement(MoveTemp(It->Value));
+				It.RemoveCurrent();
+			}
+		}
+	}
+
+	/**
+	 * Find the given element owner, or call the register callback to create one if none is found.
+	 * @note This must be paired with a call to UnregisterElementOwner.
+	 */
+	TTypedElementOwnerScopedAccess<ElementDataType> FindOrRegisterElementOwner(const KeyDataType& InKey, TFunctionRef<TTypedElementOwner<ElementDataType>()> InCreateElement)
+	{
+		TTypedElementOwnerScopedAccess<ElementDataType> ScopedAccess(&ElementOwnerMapCS);
+		if (const TTypedElementOwner<ElementDataType>* ExistingElement = ElementOwnerMap.Find(InKey))
+		{
+			ScopedAccess.Private_SetElementOwner(ExistingElement);
+		}
+		else
+		{
+			ScopedAccess.Private_SetElementOwner(&ElementOwnerMap.Add(InKey, InCreateElement()));
+		}
+		return ScopedAccess;
+	}
+
+	/**
+	 * Provide const access to the given element owner.
+	 * @note This will lock the store for writes until the scoped access struct is no longer being used. The returned struct may be unset if the key was not registered.
+	 */
+	TTypedElementOwnerScopedAccess<ElementDataType> FindElementOwner(const KeyDataType& InKey) const
+	{
+		TTypedElementOwnerScopedAccess<ElementDataType> ScopedAccess(&ElementOwnerMapCS);
+		ScopedAccess.Private_SetElementOwner(ElementOwnerMap.Find(InKey));
+		return ScopedAccess;
+	}
+
+	/**
+	 * Provide mutable access to the given element owner.
+	 * @note This will lock the store for writes until the scoped access struct is no longer being used. The returned struct may be unset if the key was not registered.
+	 */
+	TTypedElementOwnerScopedMutableAccess<ElementDataType> FindMutableElementOwner(const KeyDataType& InKey)
+	{
+		TTypedElementOwnerScopedMutableAccess<ElementDataType> ScopedAccess(&ElementOwnerMapCS);
+		ScopedAccess.Private_SetElementOwner(ElementOwnerMap.Find(InKey));
+		return ScopedAccess;
+	}
+
+	/**
+	 * Test to see whether the store has an entry for the given element owner.
+	 */
+	bool ContainsElementOwner(const KeyDataType& InKey) const
+	{
+		FScopeLock Lock(&ElementOwnerMapCS);
+		return ElementOwnerMap.Contains(InKey);
+	}
+
+private:
+	TMap<KeyDataType, TTypedElementOwner<ElementDataType>> ElementOwnerMap;
+	mutable FCriticalSection ElementOwnerMapCS;
+};
diff --git a/Engine/Source/Runtime/TypedElementFramework/Public/Elements/Framework/TypedElementRegistry.h b/Engine/Source/Runtime/TypedElementFramework/Public/Elements/Framework/TypedElementRegistry.h
new file mode 100644
--- /dev/null
+++ b/Engine/Source/Runtime/TypedElementFramework/Public/Elements/Framework/TypedElementRegistry.h
@@ -0,0 +1,425 @@
+// Copyright Epic Games, Inc. All Rights Reserved.
+
+#pragma once
+
+#include "CoreMinimal.h"
+#include "Misc/ScopeRWLock.h"
+#include "Containers/ArrayView.h"
+#include "Containers/SortedMap.h"
+#include "Elements/Framework/TypedElementHandle.h"
+#include "Elements/Framework/TypedElementList.h"
+#include "TypedElementRegistry.generated.h"
+
+/**
+ * Registry of element types and their associated interfaces, along with the elements that represent their instances.
+ */
+UCLASS(Transient)
+class TYPEDELEMENTFRAMEWORK_API UTypedElementRegistry : public UObject
+{
+	GENERATED_BODY()
+
+public:
+	UTypedElementRegistry();
+
+	//~ UObject interface
+	virtual void FinishDestroy() override;
+	static void AddReferencedObjects(UObject* InThis, FReferenceCollector& Collector);
+
+	/**
+	 * Initialize the singleton instance of the registry used in most cases.
+	 */
+	static void Private_InitializeInstance();
+
+	/**
+	 * Shutdown the singleton instance of the registry used in most cases.
+	 */
+	static void Private_ShutdownInstance();
+
+	/**
+	 * Get the singleton instance of the registry used in most cases.
+	 */
+	UFUNCTION(BlueprintPure, Category="TypedElementFramework|Registry")
+	static UTypedElementRegistry* GetInstance();
+
+	/**
+	 * Get the element type ID for the associated element type name, if any.
+	 * @return The element type ID, or 0 if the given name wasn't registered.
+	 */
+	FORCEINLINE FTypedHandleTypeId GetRegisteredElementTypeId(const FName InElementTypeName) const
+	{
+		FReadScopeLock RegisteredElementTypesLock(RegisteredElementTypesRW);
+
+		const FTypedHandleTypeId* TypeId = RegisteredElementTypesNameToId.Find(InElementTypeName);
+		return TypeId ? *TypeId : 0;
+	}
+
+	/**
+	 * Register an element type that doesn't require any additional payload data.
+	 */
+	FORCEINLINE void RegisterElementType(const FName InElementTypeName)
+	{
+		RegisterElementTypeImpl(InElementTypeName, MakeUnique<TRegisteredElementType<void>>());
+	}
+
+	/**
+	 * Register an element type that has additional payload data.
+	 */
+	template <typename ElementDataType>
+	FORCEINLINE void RegisterElementType(const FName InElementTypeName)
+	{
+		RegisterElementTypeImpl(InElementTypeName, MakeUnique<TRegisteredElementType<ElementDataType>>());
+	}
+
+	/**
+	 * Register that an element interface is supported for the given type, which must have previously been registered via RegisterElementType.
+	 */
+	template <typename BaseInterfaceType>
+	FORCEINLINE void RegisterElementInterface(const FName InElementTypeName, UTypedElementInterface* InElementInterface, const bool InAllowOverride = false)
+	{
+		RegisterElementInterfaceImpl(InElementTypeName, InElementInterface, BaseInterfaceType::StaticClass(), InAllowOverride);
+	}
+
+	/**
+	 * Get the element interface supported by the given type, or null if there is no support for this interface.
+	 */
+	template <typename BaseInterfaceType>
+	FORCEINLINE BaseInterfaceType* GetElementInterface(const FTypedHandleTypeId InElementTypeId) const
+	{
+		return static_cast<BaseInterfaceType*>(GetElementInterfaceImpl(InElementTypeId, BaseInterfaceType::StaticClass()));
+	}
+
+	/**
+	 * Get the element interface supported by the given handle, or null if there is no support for this interface.
+	 */
+	template <typename BaseInterfaceType>
+	FORCEINLINE BaseInterfaceType* GetElementInterface(const FTypedElementHandle& InElementHandle) const
+	{
+		return static_cast<BaseInterfaceType*>(GetElementInterfaceImpl(InElementHandle.GetId().GetTypeId(), BaseInterfaceType::StaticClass()));
+	}
+
+	/**
+	 * Get the element interface supported by the given handle, or null if there is no support for this interface.
+	 */
+	UFUNCTION(BlueprintPure, Category="TypedElementFramework|Registry")
+	FORCEINLINE UTypedElementInterface* GetElementInterface(const FTypedElementHandle& InElementHandle, const TSubclassOf<UTypedElementInterface>& InBaseInterfaceType) const
+	{
+		return GetElementInterfaceImpl(InElementHandle.GetId().GetTypeId(), InBaseInterfaceType);
+	}
+
+	/**
+	 * Create an element that doesn't require any additional payload data.
+	 * @note The associated handle ID should be something that can externally be used to uniquely identify this element, until DestroyElementHandle is called on this handle.
+	 */
+	FORCEINLINE FTypedElementOwner CreateElement(const FName InElementTypeName, const FTypedHandleElementId InElementId)
+	{
+		return CreateElementImpl<void>(InElementTypeName, InElementId);
+	}
+
+	/**
+	 * Create an element that has additional payload data.
+	 * @note Allocation of the payload data and the associated handle ID are managed internally, and the data will remain valid until DestroyElementHandle is called on this handle.
+	 */
+	template <typename ElementDataType>
+	FORCEINLINE TTypedElementOwner<ElementDataType> CreateElement(const FName InElementTypeName)
+	{
+		return CreateElementImpl<ElementDataType>(InElementTypeName, INDEX_NONE);
+	}
+
+	/**
+	 * Destroy an element.
+	 */
+	FORCEINLINE void DestroyElement(FTypedElementOwner& InOutElementOwner)
+	{
+		return DestroyElementImpl<void>(InOutElementOwner);
+	}
+
+	/**
+	 * Destroy an element.
+	 */
+	template <typename ElementDataType>
+	FORCEINLINE void DestroyElement(TTypedElementOwner<ElementDataType>& InOutElementOwner)
+	{
+		return DestroyElementImpl<ElementDataType>(InOutElementOwner);
+	}
+
+	/**
+	 * Release an element ID that was previously acquired from an existing handle.
+	 */
+	void ReleaseElementId(FTypedElementId& InOutElementId);
+
+	/**
+	 * Get an element handle from its minimal ID.
+	 */
+	FTypedElementHandle GetElementHandle(const FTypedElementId& InElementId) const;
+
+	/**
+	 * Get an element that implements the given interface from its minimal ID.
+	 */
+	FORCEINLINE FTypedElement GetElement(const FTypedElementId& InElementId, const TSubclassOf<UTypedElementInterface>& InBaseInterfaceType) const
+	{
+		FTypedElement Element;
+		GetElementImpl(InElementId, InBaseInterfaceType, Element);
+		return Element;
+	}
+
+	/**
+	 * Get an element that implements the given interface from its minimal ID.
+	 */
+	template <typename BaseInterfaceType>
+	FORCEINLINE TTypedElement<BaseInterfaceType> GetElement(const FTypedElementId& InElementId) const
+	{
+		TTypedElement<BaseInterfaceType> Element;
+		GetElementImpl(InElementId, BaseInterfaceType::StaticClass(), Element);
+		return Element;
+	}
+
+	/**
+	 * Get an element that implements the given interface from its handle.
+	 */
+	FORCEINLINE FTypedElement GetElement(const FTypedElementHandle& InElementHandle, const TSubclassOf<UTypedElementInterface>& InBaseInterfaceType) const
+	{
+		FTypedElement Element;
+		GetElementImpl(InElementHandle, InBaseInterfaceType, Element);
+		return Element;
+	}
+
+	/**
+	 * Get an element that implements the given interface from its handle.
+	 */
+	template <typename BaseInterfaceType>
+	FORCEINLINE TTypedElement<BaseInterfaceType> GetElement(const FTypedElementHandle& InElementHandle) const
+	{
+		TTypedElement<BaseInterfaceType> Element;
+		GetElementImpl(InElementHandle, BaseInterfaceType::StaticClass(), Element);
+		return Element;
+	}
+
+	/**
+	 * Create an empty list of elements associated with this registry.
+	 */
+	UFUNCTION(BlueprintCallable, Category="TypedElementFramework|Registry")
+	FORCEINLINE UTypedElementList* CreateElementList()
+	{
+		return UTypedElementList::Private_CreateElementList(this);
+	}
+
+	/**
+	 * Create an empty list of elements associated with this registry, populated from the given minimal IDs that are valid.
+	 */
+	UTypedElementList* CreateElementList(TArrayView<const FTypedElementId> InElementIds);
+
+	/**
+	 * Create an empty list of elements associated with this registry, populated from the given handles that are valid.
+	 */
+	UTypedElementList* CreateElementList(TArrayView<const FTypedElementHandle> InElementHandles);
+
+	/**
+	 * Create an empty list of elements associated with this registry, populated from the given owners that are valid.
+	 */
+	template <typename ElementDataType>
+	FORCEINLINE UTypedElementList* CreateElementList(const TArray<TTypedElementOwner<ElementDataType>>& InElementOwners)
+	{
+		return CreateElementList(MakeArrayView(InElementOwners));
+	}
+
+	/**
+	 * Create an empty list of elements associated with this registry, populated from the given owners that are valid.
+	 */
+	template <typename ElementDataType>
+	UTypedElementList* CreateElementList(TArrayView<const TTypedElementOwner<ElementDataType>> InElementOwners)
+	{
+		UTypedElementList* ElementList = CreateElementList();
+		ElementList->Append(InElementOwners);
+		return ElementList;
+	}
+
+	void Private_OnElementListCreated(UTypedElementList* InElementList)
+	{
+		FWriteScopeLock ActiveElementListsLock(ActiveElementListsRW);
+		ActiveElementLists.Add(InElementList);
+	}
+
+	void Private_OnElementListDestroyed(UTypedElementList* InElementList)
+	{
+		FWriteScopeLock ActiveElementListsLock(ActiveElementListsRW);
+		ActiveElementLists.Remove(InElementList);
+	}
+
+	// Note: Access for UTypedElementList
+	FORCEINLINE void Private_GetElementImpl(const FTypedElementHandle& InElementHandle, const UClass* InBaseInterfaceType, FTypedElement& OutElement) const
+	{
+		GetElementImpl(InElementHandle, InBaseInterfaceType, OutElement);
+	}
+
+private:
+	struct FRegisteredElementType
+	{
+		virtual ~FRegisteredElementType() = default;
+
+		virtual FTypedElementInternalData& AddDataForElement(FTypedHandleElementId& InOutElementId) = 0;
+		virtual void RemoveDataForElement(const FTypedHandleElementId InElementId, const FTypedElementInternalData* InExpectedDataPtr) = 0;
+		virtual const FTypedElementInternalData& GetDataForElement(const FTypedHandleElementId InElementId) const = 0;
+		virtual void SetDataTypeId(const FTypedHandleTypeId InTypeId) = 0;
+		virtual FTypedHandleTypeId GetDataTypeId() const = 0;
+		virtual FName GetDataTypeName() const = 0;
+
+		FTypedHandleTypeId TypeId = 0;
+		FName TypeName;
+		TSortedMap<FName, UTypedElementInterface*, FDefaultAllocator, FNameFastLess> Interfaces;
+	};
+
+	template <typename ElementDataType>
+	struct TRegisteredElementType : public FRegisteredElementType
+	{
+		virtual ~TRegisteredElementType() = default;
+		
+		virtual FTypedElementInternalData& AddDataForElement(FTypedHandleElementId& InOutElementId) override
+		{
+			return HandleDataStore.AddDataForElement(InOutElementId);
+		}
+
+		virtual void RemoveDataForElement(const FTypedHandleElementId InElementId, const FTypedElementInternalData* InExpectedDataPtr) override
+		{
+			HandleDataStore.RemoveDataForElement(InElementId, InExpectedDataPtr);
+		}
+
+		virtual const FTypedElementInternalData& GetDataForElement(const FTypedHandleElementId InElementId) const override
+		{
+			return HandleDataStore.GetDataForElement(InElementId);
+		}
+
+		virtual void SetDataTypeId(const FTypedHandleTypeId InTypeId) override
+		{
+			TTypedElementInternalDataStore<ElementDataType>::SetStaticDataTypeId(InTypeId);
+		}
+
+		virtual FTypedHandleTypeId GetDataTypeId() const override
+		{
+			return TTypedElementInternalDataStore<ElementDataType>::StaticDataTypeId();
+		}
+
+		virtual FName GetDataTypeName() const override
+		{
+			return TTypedElementInternalDataStore<ElementDataType>::StaticDataTypeName();
+		}
+
+		TTypedElementInternalDataStore<ElementDataType> HandleDataStore;
+	};
+
+	void RegisterElementTypeImpl(const FName InElementTypeName, TUniquePtr<FRegisteredElementType>&& InRegisteredElementType);
+	void RegisterElementInterfaceImpl(const FName InElementTypeName, UTypedElementInterface* InElementInterface, const TSubclassOf<UTypedElementInterface>& InBaseInterfaceType, const bool InAllowOverride);
+	UTypedElementInterface* GetElementInterfaceImpl(const FTypedHandleTypeId InElementTypeId, const TSubclassOf<UTypedElementInterface>& InBaseInterfaceType) const;
+
+	template <typename ElementDataType>
+	TTypedElementOwner<ElementDataType> CreateElementImpl(const FName InElementTypeName, const FTypedHandleElementId InElementId)
+	{
+		FRegisteredElementType* RegisteredElementType = GetRegisteredElementTypeFromName(InElementTypeName);
+		checkf(RegisteredElementType, TEXT("Element type '%s' has not been registered!"), *InElementTypeName.ToString());
+
+		checkf(RegisteredElementType->GetDataTypeId() == TTypedElementInternalDataStore<ElementDataType>::StaticDataTypeId(), TEXT("Element type '%s' uses '%s' as its handle data type, but '%s' was requested!"), *InElementTypeName.ToString(), *RegisteredElementType->GetDataTypeName().ToString(), *TTypedElementInternalDataStore<ElementDataType>::StaticDataTypeName().ToString());
+
+		FTypedHandleElementId NewElementId = InElementId;
+		FTypedElementInternalData& NewElementData = RegisteredElementType->AddDataForElement(NewElementId);
+
+		TTypedElementOwner<ElementDataType> ElementOwner;
+		ElementOwner.Private_InitializeAddRef(RegisteredElementType->TypeId, NewElementId, static_cast<TTypedElementInternalData<ElementDataType>&>(NewElementData));
+
+		return ElementOwner;
+	}
+
+	template <typename ElementDataType>
+	void DestroyElementImpl(TTypedElementOwner<ElementDataType>& InOutElementOwner)
+	{
+		FRegisteredElementType* RegisteredElementType = GetRegisteredElementTypeFromId(InOutElementOwner.GetId().GetTypeId());
+		checkf(RegisteredElementType, TEXT("Element type ID '%d' has not been registered!"), InOutElementOwner.GetId().GetTypeId());
+
+#if DO_CHECK
+		{
+			const FTypedElementInternalData& ElementData = RegisteredElementType->GetDataForElement(InOutElementOwner.GetId().GetElementId());
+			const FTypedElementRefCount RefCount = ElementData.GetRefCount();
+
+			const bool bIsExternallyReferenced = RefCount > 1;
+			if (bIsExternallyReferenced)
+			{
+				ElementData.LogReferences();
+				UE_LOG(LogCore, Fatal, TEXT("Element is still externally referenced when being destroyed! Ref-count: %d; see above for reference information (if available)."), RefCount);
+			}
+		}
+#endif	// DO_CHECK
+
+		RegisteredElementType->RemoveDataForElement(InOutElementOwner.GetId().GetElementId(), InOutElementOwner.Private_GetInternalData());
+		InOutElementOwner.Private_DestroyNoRef();
+	}
+
+	template <typename BaseInterfaceType>
+	void GetElementImpl(const FTypedElementId& InElementId, const UClass* InBaseInterfaceType, TTypedElement<BaseInterfaceType>& OutElement) const
+	{
+		OutElement.Private_DestroyReleaseRef();
+
+		if (InElementId)
+		{
+			FRegisteredElementType* RegisteredElementType = GetRegisteredElementTypeFromId(InElementId.GetTypeId());
+			checkf(RegisteredElementType, TEXT("Element type ID '%d' has not been registered!"), InElementId.GetTypeId());
+
+			OutElement.Private_InitializeAddRef(InElementId.GetTypeId(), InElementId.GetElementId(), RegisteredElementType->GetDataForElement(InElementId.GetElementId()), static_cast<BaseInterfaceType*>(RegisteredElementType->Interfaces.FindRef(InBaseInterfaceType->GetFName())));
+		}
+	}
+
+	template <typename BaseInterfaceType>
+	void GetElementImpl(const FTypedElementHandle& InElementHandle, const UClass* InBaseInterfaceType, TTypedElement<BaseInterfaceType>& OutElement) const
+	{
+		OutElement.Private_DestroyReleaseRef();
+
+		if (InElementHandle)
+		{
+			FRegisteredElementType* RegisteredElementType = GetRegisteredElementTypeFromId(InElementHandle.GetId().GetTypeId());
+			checkf(RegisteredElementType, TEXT("Element type ID '%d' has not been registered!"), InElementHandle.GetId().GetTypeId());
+
+			OutElement.Private_InitializeAddRef(InElementHandle.GetId().GetTypeId(), InElementHandle.GetId().GetElementId(), *InElementHandle.Private_GetInternalData(), static_cast<BaseInterfaceType*>(RegisteredElementType->Interfaces.FindRef(InBaseInterfaceType->GetFName())));
+		}
+	}
+
+	void AddRegisteredElementType(TUniquePtr<FRegisteredElementType>&& InRegisteredElementType)
+	{
+		checkf(InRegisteredElementType->TypeId > 0, TEXT("Element type ID was unassigned!"));
+		checkf(!GetRegisteredElementTypeFromId(InRegisteredElementType->TypeId), TEXT("Element type '%d' has already been registered!"), InRegisteredElementType->TypeId);
+		checkf(!GetRegisteredElementTypeFromName(InRegisteredElementType->TypeName), TEXT("Element type '%s' has already been registered!"), *InRegisteredElementType->TypeName.ToString());
+
+		{
+			FWriteScopeLock RegisteredElementTypesLock(RegisteredElementTypesRW);
+
+			RegisteredElementTypesNameToId.Add(InRegisteredElementType->TypeName, InRegisteredElementType->TypeId);
+			RegisteredElementTypes[InRegisteredElementType->TypeId - 1] = MoveTemp(InRegisteredElementType);
+		}
+	}
+
+	FRegisteredElementType* GetRegisteredElementTypeFromId(const FTypedHandleTypeId InTypeId) const
+	{
+		FReadScopeLock RegisteredElementTypesLock(RegisteredElementTypesRW);
+
+		return InTypeId > 0
+			? RegisteredElementTypes[InTypeId - 1].Get()
+			: nullptr;
+	}
+
+	FRegisteredElementType* GetRegisteredElementTypeFromName(const FName& InTypeName) const
+	{
+		FReadScopeLock RegisteredElementTypesLock(RegisteredElementTypesRW);
+
+		if (const FTypedHandleTypeId* TypeId = RegisteredElementTypesNameToId.Find(InTypeName))
+		{
+			return RegisteredElementTypes[(*TypeId) - 1].Get();
+		}
+
+		return nullptr;
+	}
+
+	void NotifyElementListPendingChanges();
+	
+	mutable FRWLock RegisteredElementTypesRW;
+	TUniquePtr<FRegisteredElementType> RegisteredElementTypes[TypedHandleMaxTypeId - 1];
+	TSortedMap<FName, FTypedHandleTypeId, FDefaultAllocator, FNameFastLess> RegisteredElementTypesNameToId;
+
+	mutable FRWLock ActiveElementListsRW;
+	TSet<UTypedElementList*> ActiveElementLists;
+};
diff --git a/Engine/Source/Runtime/TypedElementFramework/Public/Elements/Interfaces/TypedElementInterface.h b/Engine/Source/Runtime/TypedElementFramework/Public/Elements/Interfaces/TypedElementInterface.h
new file mode 100644
--- /dev/null
+++ b/Engine/Source/Runtime/TypedElementFramework/Public/Elements/Interfaces/TypedElementInterface.h
@@ -0,0 +1,16 @@
+// Copyright Epic Games, Inc. All Rights Reserved.
+
+#pragma once
+
+#include "CoreMinimal.h"
+#include "TypedElementInterface.generated.h"
+
+/**
+ * Base typed used to represent element interfaces.
+ * @note Top-level element interfaces that inherit from this should also specialize TTypedElement for their API.
+ */
+UCLASS(Abstract)
+class TYPEDELEMENTFRAMEWORK_API UTypedElementInterface : public UObject
+{
+	GENERATED_BODY()
+};
diff --git a/Engine/Source/Runtime/TypedElementFramework/Public/TypedElementData.h b/Engine/Source/Runtime/TypedElementFramework/Public/TypedElementData.h
deleted file mode 100644
diff --git a/Engine/Source/Runtime/TypedElementFramework/Public/TypedElementHandle.h b/Engine/Source/Runtime/TypedElementFramework/Public/TypedElementHandle.h
deleted file mode 100644
diff --git a/Engine/Source/Runtime/TypedElementFramework/Public/TypedElementLimits.h b/Engine/Source/Runtime/TypedElementFramework/Public/TypedElementLimits.h
deleted file mode 100644
diff --git a/Engine/Source/Runtime/TypedElementFramework/Public/TypedElementList.h b/Engine/Source/Runtime/TypedElementFramework/Public/TypedElementList.h
deleted file mode 100644
diff --git a/Engine/Source/Runtime/TypedElementFramework/Public/TypedElementOwnerStore.h b/Engine/Source/Runtime/TypedElementFramework/Public/TypedElementOwnerStore.h
deleted file mode 100644
diff --git a/Engine/Source/Runtime/TypedElementFramework/Public/TypedElementRegistry.h b/Engine/Source/Runtime/TypedElementFramework/Public/TypedElementRegistry.h
deleted file mode 100644
diff --git a/Engine/Source/Runtime/TypedElementInterfaces/Private/Elements/Interfaces/TypedElementObjectInterface.cpp b/Engine/Source/Runtime/TypedElementInterfaces/Private/Elements/Interfaces/TypedElementObjectInterface.cpp
new file mode 100644
--- /dev/null
+++ b/Engine/Source/Runtime/TypedElementInterfaces/Private/Elements/Interfaces/TypedElementObjectInterface.cpp
@@ -0,0 +1,3 @@
+// Copyright Epic Games, Inc. All Rights Reserved.
+
+#include "Elements/Interfaces/TypedElementObjectInterface.h"
diff --git a/Engine/Source/Runtime/TypedElementInterfaces/Private/Elements/Interfaces/TypedElementSelectionInterface.cpp b/Engine/Source/Runtime/TypedElementInterfaces/Private/Elements/Interfaces/TypedElementSelectionInterface.cpp
new file mode 100644
--- /dev/null
+++ b/Engine/Source/Runtime/TypedElementInterfaces/Private/Elements/Interfaces/TypedElementSelectionInterface.cpp
@@ -0,0 +1,19 @@
+// Copyright Epic Games, Inc. All Rights Reserved.
+
+#include "Elements/Interfaces/TypedElementSelectionInterface.h"
+#include "Elements/Framework/TypedElementList.h"
+
+bool UTypedElementSelectionInterface::IsElementSelected(const FTypedElementHandle& InElementHandle, const UTypedElementList* InSelectionSet, const FTypedElementIsSelectedOptions& InSelectionOptions)
+{
+	return InSelectionSet->Contains(InElementHandle);
+}
+
+bool UTypedElementSelectionInterface::SelectElement(const FTypedElementHandle& InElementHandle, UTypedElementList* InSelectionSet, const FTypedElementSelectionOptions& InSelectionOptions)
+{
+	return InSelectionSet->Add(InElementHandle);
+}
+
+bool UTypedElementSelectionInterface::DeselectElement(const FTypedElementHandle& InElementHandle, UTypedElementList* InSelectionSet, const FTypedElementSelectionOptions& InSelectionOptions)
+{
+	return InSelectionSet->Remove(InElementHandle);
+}
diff --git a/Engine/Source/Runtime/TypedElementInterfaces/Private/TypedElementObjectInterface.cpp b/Engine/Source/Runtime/TypedElementInterfaces/Private/TypedElementObjectInterface.cpp
deleted file mode 100644
diff --git a/Engine/Source/Runtime/TypedElementInterfaces/Private/TypedElementSelectionInterface.cpp b/Engine/Source/Runtime/TypedElementInterfaces/Private/TypedElementSelectionInterface.cpp
deleted file mode 100644
diff --git a/Engine/Source/Runtime/TypedElementInterfaces/Public/Elements/Interfaces/TypedElementObjectInterface.h b/Engine/Source/Runtime/TypedElementInterfaces/Public/Elements/Interfaces/TypedElementObjectInterface.h
new file mode 100644
--- /dev/null
+++ b/Engine/Source/Runtime/TypedElementInterfaces/Public/Elements/Interfaces/TypedElementObjectInterface.h
@@ -0,0 +1,25 @@
+// Copyright Epic Games, Inc. All Rights Reserved.
+
+#pragma once
+
+#include "Elements/Framework/TypedElementHandle.h"
+#include "TypedElementObjectInterface.generated.h"
+
+UCLASS(Abstract)
+class TYPEDELEMENTINTERFACES_API UTypedElementObjectInterface : public UTypedElementInterface
+{
+	GENERATED_BODY()
+
+public:
+	/**
+	 * Retreive the object instance that this handle represents, if any.
+	 */
+	UFUNCTION(BlueprintPure, Category="TypedElementInterfaces|Object")
+	virtual UObject* GetObject(const FTypedElementHandle& InElementHandle) { return nullptr; }
+};
+
+template <>
+struct TTypedElement<UTypedElementObjectInterface> : public TTypedElementBase<UTypedElementObjectInterface>
+{
+	UObject* GetObject() const { return InterfacePtr->GetObject(*this); }
+};
diff --git a/Engine/Source/Runtime/TypedElementInterfaces/Public/Elements/Interfaces/TypedElementSelectionInterface.h b/Engine/Source/Runtime/TypedElementInterfaces/Public/Elements/Interfaces/TypedElementSelectionInterface.h
new file mode 100644
--- /dev/null
+++ b/Engine/Source/Runtime/TypedElementInterfaces/Public/Elements/Interfaces/TypedElementSelectionInterface.h
@@ -0,0 +1,146 @@
+// Copyright Epic Games, Inc. All Rights Reserved.
+
+#pragma once
+
+#include "Elements/Framework/TypedElementHandle.h"
+#include "TypedElementSelectionInterface.generated.h"
+
+class UTypedElementList;
+
+UENUM()
+enum class ETypedElementSelectionMethod : uint8
+{
+	/** Select the "primary" element (eg, a component favor selecting its owner actor) */
+	Primary,
+	/** Select the "secondary" element (eg, a component would favor selecting itself) */
+	Secondary,
+};
+
+USTRUCT(BlueprintType)
+struct FTypedElementIsSelectedOptions
+{
+	GENERATED_BODY()
+
+public:
+	FTypedElementIsSelectedOptions& SetAllowIndirect(const bool InAllowIndirect) { bAllowIndirect = InAllowIndirect; return *this; }
+	bool AllowIndirect() const { return bAllowIndirect; }
+	
+private:
+	UPROPERTY(BlueprintReadWrite, Category="TypedElementInterfaces|Selection|IsSelectedOptions", meta=(AllowPrivateAccess=true))
+	bool bAllowIndirect = false;
+};
+
+USTRUCT(BlueprintType)
+struct FTypedElementSelectionOptions
+{
+	GENERATED_BODY()
+
+public:
+	FTypedElementSelectionOptions& SetAllowHidden(const bool InAllowHidden) { bAllowHidden = InAllowHidden; return *this; }
+	bool AllowHidden() const { return bAllowHidden; }
+
+	FTypedElementSelectionOptions& SetAllowGroups(const bool InAllowGroups) { bAllowGroups = InAllowGroups; return *this; }
+	bool AllowGroups() const { return bAllowGroups; }
+
+	FTypedElementSelectionOptions& SetAllowLegacyNotifications(const bool InAllowLegacyNotifications) { bAllowLegacyNotifications = InAllowLegacyNotifications; return *this; }
+	bool AllowLegacyNotifications() const { return bAllowLegacyNotifications; }
+
+	FTypedElementSelectionOptions& SetWarnIfLocked(const bool InWarnIfLocked) { bWarnIfLocked = InWarnIfLocked; return *this; }
+	bool WarnIfLocked() const { return bWarnIfLocked; }
+	
+private:
+	UPROPERTY(BlueprintReadWrite, Category="TypedElementInterfaces|Selection|SelectionOptions", meta=(AllowPrivateAccess=true))
+	bool bAllowHidden = false;
+
+	UPROPERTY(BlueprintReadWrite, Category="TypedElementInterfaces|Selection|SelectionOptions", meta=(AllowPrivateAccess=true))
+	bool bAllowGroups = true;
+
+	UPROPERTY(BlueprintReadWrite, Category="TypedElementInterfaces|Selection|SelectionOptions", meta=(AllowPrivateAccess=true))
+	bool bAllowLegacyNotifications = true;
+
+	UPROPERTY(BlueprintReadWrite, Category="TypedElementInterfaces|Selection|SelectionOptions", meta=(AllowPrivateAccess=true))
+	bool bWarnIfLocked = false;
+};
+
+UCLASS(Abstract)
+class TYPEDELEMENTINTERFACES_API UTypedElementSelectionInterface : public UTypedElementInterface
+{
+	GENERATED_BODY()
+
+public:
+	/**
+	 * Test to see whether the given element is currently considered selected.
+	 */
+	UFUNCTION(BlueprintPure, Category="TypedElementInterfaces|Selection")
+	virtual bool IsElementSelected(const FTypedElementHandle& InElementHandle, const UTypedElementList* InSelectionSet, const FTypedElementIsSelectedOptions& InSelectionOptions);
+
+	/**
+	 * Test to see whether the given element can be selected.
+	 */
+	UFUNCTION(BlueprintPure, Category="TypedElementInterfaces|Selection")
+	virtual bool CanSelectElement(const FTypedElementHandle& InElementHandle, const FTypedElementSelectionOptions& InSelectionOptions) { return true; }
+
+	/**
+	 * Test to see whether the given element can be deselected.
+	 */
+	UFUNCTION(BlueprintPure, Category="TypedElementInterfaces|Selection")
+	virtual bool CanDeselectElement(const FTypedElementHandle& InElementHandle, const FTypedElementSelectionOptions& InSelectionOptions) { return true; }
+
+	/**
+	 * Attempt to select the given element.
+	 * @return True if the selection was changed, false otherwise.
+	 */
+	UFUNCTION(BlueprintCallable, Category="TypedElementInterfaces|Selection")
+	virtual bool SelectElement(const FTypedElementHandle& InElementHandle, UTypedElementList* InSelectionSet, const FTypedElementSelectionOptions& InSelectionOptions);
+
+	/**
+	 * Attempt to deselect the given element.
+	 * @return True if the selection was changed, false otherwise.
+	 */
+	UFUNCTION(BlueprintCallable, Category="TypedElementInterfaces|Selection")
+	virtual bool DeselectElement(const FTypedElementHandle& InElementHandle, UTypedElementList* InSelectionSet, const FTypedElementSelectionOptions& InSelectionOptions);
+
+	/**
+	 * Test to see whether selection modifiers (Ctrl or Shift) are allowed while selecting this element.
+	 */
+	UFUNCTION(BlueprintPure, Category="TypedElementInterfaces|Selection")
+	virtual bool AllowSelectionModifiers(const FTypedElementHandle& InElementHandle, const UTypedElementList* InSelectionSet) { return true; }
+
+	/**
+	 * Given an element, return the element that should actually perform a selection operation.
+	 */
+	UFUNCTION(BlueprintPure, Category="TypedElementInterfaces|Selection")
+	virtual FTypedElementHandle GetSelectionElement(const FTypedElementHandle& InElementHandle, const UTypedElementList* InCurrentSelection, const ETypedElementSelectionMethod InSelectionMethod) { return InElementHandle; }
+
+	/**
+	 * Test to see whether the given element prevents the selection set state from being transacted for undo/redo (eg, if the element belongs to a PIE instance).
+	 */
+	virtual bool ShouldPreventTransactions(const FTypedElementHandle& InElementHandle) { return false; }
+
+	/**
+	 * Write the information needed to find the given element again when this transaction is replayed for undo/redo.
+	 * @note The data written must match the data read by ReadTransactedElement.
+	 */
+	virtual void WriteTransactedElement(const FTypedElementHandle& InElementHandle, FArchive& InArchive) {}
+
+	/**
+	 * Read the information needed to be find a previously tracked element, and attempt to get it now that this transaction is being replayed for undo/redo.
+	 * @note The data read must match the data written by WriteTransactedElement.
+	 */
+	virtual FTypedElementHandle ReadTransactedElement(FArchive& InArchive) { return FTypedElementHandle(); }
+};
+
+template <>
+struct TTypedElement<UTypedElementSelectionInterface> : public TTypedElementBase<UTypedElementSelectionInterface>
+{
+	bool IsElementSelected(const UTypedElementList* InSelectionSet, const FTypedElementIsSelectedOptions& InSelectionOptions) const { return InterfacePtr->IsElementSelected(*this, InSelectionSet, InSelectionOptions); }
+	bool CanSelectElement(const FTypedElementSelectionOptions& InSelectionOptions) const { return InterfacePtr->CanSelectElement(*this, InSelectionOptions); }
+	bool CanDeselectElement(const FTypedElementSelectionOptions& InSelectionOptions) const { return InterfacePtr->CanDeselectElement(*this, InSelectionOptions); }
+	bool SelectElement(UTypedElementList* InSelectionSet, const FTypedElementSelectionOptions& InSelectionOptions) const { return InterfacePtr->SelectElement(*this, InSelectionSet, InSelectionOptions); }
+	bool DeselectElement(UTypedElementList* InSelectionSet, const FTypedElementSelectionOptions& InSelectionOptions) const { return InterfacePtr->DeselectElement(*this, InSelectionSet, InSelectionOptions); }
+	bool AllowSelectionModifiers(const UTypedElementList* InSelectionSet) const { return InterfacePtr->AllowSelectionModifiers(*this, InSelectionSet); }
+	FTypedElementHandle GetSelectionElement(const UTypedElementList* InCurrentSelection, const ETypedElementSelectionMethod InSelectionMethod) const { return InterfacePtr->GetSelectionElement(*this, InCurrentSelection, InSelectionMethod); }
+	bool ShouldPreventTransactions() const { return InterfacePtr->ShouldPreventTransactions(*this); }
+	void WriteTransactedElement(FArchive& InArchive) const { return InterfacePtr->WriteTransactedElement(*this, InArchive); }
+	FTypedElementHandle ReadTransactedElement(FArchive& InArchive) const { return InterfacePtr->ReadTransactedElement(InArchive); }
+};
diff --git a/Engine/Source/Runtime/TypedElementInterfaces/Public/TypedElementObjectInterface.h b/Engine/Source/Runtime/TypedElementInterfaces/Public/TypedElementObjectInterface.h
deleted file mode 100644
diff --git a/Engine/Source/Runtime/TypedElementInterfaces/Public/TypedElementSelectionInterface.h b/Engine/Source/Runtime/TypedElementInterfaces/Public/TypedElementSelectionInterface.h
deleted file mode 100644
