Added functions to enumerate selected objects from the typed element selection set

Moved the typed element list to also use enumeration functions for consistency

diff --git a/Engine/Source/Editor/EditorFramework/Private/Elements/Framework/TypedElementSelectionSet.cpp b/Engine/Source/Editor/EditorFramework/Private/Elements/Framework/TypedElementSelectionSet.cpp
--- a/Engine/Source/Editor/EditorFramework/Private/Elements/Framework/TypedElementSelectionSet.cpp
+++ b/Engine/Source/Editor/EditorFramework/Private/Elements/Framework/TypedElementSelectionSet.cpp
@@ -2,6 +2,7 @@
 
 #include "Elements/Framework/TypedElementSelectionSet.h"
 #include "Elements/Framework/TypedElementRegistry.h"
+#include "Elements/Interfaces/TypedElementObjectInterface.h"
 
 UTypedElementSelectionSet::UTypedElementSelectionSet()
 {
@@ -17,12 +18,16 @@ bool UTypedElementSelectionSet::Modify(bool bAlwaysMarkDirty)
 {
 	if (GUndo && CanModify())
 	{
-		for (const TTypedElement<UTypedElementSelectionInterface>& SelectionElement : ElementList->IterateInterface<UTypedElementSelectionInterface>())
+		bool bCanModify = true;
+		ElementList->ForEachElement<UTypedElementSelectionInterface>([&bCanModify](const TTypedElement<UTypedElementSelectionInterface>& InSelectionElement)
 		{
-			if (SelectionElement.ShouldPreventTransactions())
-			{
-				return false;
-			}
+			bCanModify = InSelectionElement.ShouldPreventTransactions();
+			return bCanModify;
+		});
+
+		if (!bCanModify)
+		{
+			return false;
 		}
 
 		return Super::Modify(bAlwaysMarkDirty);
@@ -39,13 +44,14 @@ void UTypedElementSelectionSet::Serialize(FArchive& Ar)
 	{
 		FTypedHandleTypeId ElementTypeId = 0;
 
-		for (const TTypedElement<UTypedElementSelectionInterface>& SelectionElement : ElementList->IterateInterface<UTypedElementSelectionInterface>())
+		ElementList->ForEachElement<UTypedElementSelectionInterface>([&Ar, &ElementTypeId](const TTypedElement<UTypedElementSelectionInterface>& InSelectionElement)
 		{
-			ElementTypeId = SelectionElement.GetId().GetTypeId();
+			ElementTypeId = InSelectionElement.GetId().GetTypeId();
 			Ar << ElementTypeId;
 
-			SelectionElement.WriteTransactedElement(Ar);
-		}
+			InSelectionElement.WriteTransactedElement(Ar);
+			return true;
+		});
 
 		// End of the list
 		ElementTypeId = 0;
@@ -173,11 +179,7 @@ bool UTypedElementSelectionSet::ClearSelection(const FTypedElementSelectionOptio
 	{
 		// Take a copy of the currently selected elements to avoid mutating the selection set while iterating
 		TArray<FTypedElementHandle, TInlineAllocator<8>> ElementsCopy;
-		ElementsCopy.Reserve(ElementList->Num());
-		for (const FTypedElementHandle& ElementHandle : *ElementList)
-		{
-			ElementsCopy.Emplace(ElementHandle);
-		}
+		ElementList->GetElementHandles(ElementsCopy);
 
 		for (const FTypedElementHandle& ElementHandle : ElementsCopy)
 		{
@@ -209,6 +211,33 @@ FTypedElementHandle UTypedElementSelectionSet::GetSelectionElement(const FTypedE
 	return SelectionSetElement ? SelectionSetElement.GetSelectionElement(InSelectionMethod) : FTypedElementHandle();
 }
 
+TArray<UObject*> UTypedElementSelectionSet::GetSelectedObjects(const UClass* InRequiredClass) const
+{
+	TArray<UObject*> SelectedObjects;
+	SelectedObjects.Reserve(ElementList->Num());
+
+	ForEachSelectedObject([&SelectedObjects](UObject* InObject)
+	{
+		SelectedObjects.Add(InObject);
+		return true;
+	}, InRequiredClass);
+
+	return SelectedObjects;
+}
+
+void UTypedElementSelectionSet::ForEachSelectedObject(TFunctionRef<bool(UObject*)> InCallback, const UClass* InRequiredClass) const
+{
+	ElementList->ForEachElement<UTypedElementObjectInterface>([&InCallback, InRequiredClass](const TTypedElement<UTypedElementObjectInterface>& InObjectElement)
+	{
+		UObject* ElementObject = InObjectElement.GetObject();
+		if (ElementObject && (!InRequiredClass || ElementObject->IsA(InRequiredClass)))
+		{
+			return InCallback(ElementObject);
+		}
+		return true;
+	});
+}
+
 FTypedElementSelectionSetElement UTypedElementSelectionSet::ResolveSelectionSetElement(const FTypedElementHandle& InElementHandle) const
 {
 	return InElementHandle
diff --git a/Engine/Source/Editor/EditorFramework/Public/Elements/Framework/TypedElementSelectionSet.h b/Engine/Source/Editor/EditorFramework/Public/Elements/Framework/TypedElementSelectionSet.h
--- a/Engine/Source/Editor/EditorFramework/Public/Elements/Framework/TypedElementSelectionSet.h
+++ b/Engine/Source/Editor/EditorFramework/Public/Elements/Framework/TypedElementSelectionSet.h
@@ -157,6 +157,86 @@ public:
 	UFUNCTION(BlueprintPure, Category="TypedElementFramework|Selection")
 	FTypedElementHandle GetSelectionElement(const FTypedElementHandle& InElementHandle, const ETypedElementSelectionMethod InSelectionMethod) const;
 
+	/**
+	 * Get the handle of every selected element, optionally filtering to elements that implement the given interface.
+	 */
+	UFUNCTION(BlueprintCallable, BlueprintPure=false, Category="TypedElementFramework|Selection")
+	TArray<FTypedElementHandle> GetSelectedElementHandles(const TSubclassOf<UTypedElementInterface>& InBaseInterfaceType = nullptr) const
+	{
+		return ElementList->GetElementHandles(InBaseInterfaceType);
+	}
+
+	/**
+	 * Get the handle of every selected element, optionally filtering to elements that implement the given interface.
+	 */
+	template <typename ArrayAllocator>
+	void GetSelectedElementHandles(TArray<FTypedElementHandle, ArrayAllocator>& OutArray, const TSubclassOf<UTypedElementInterface>& InBaseInterfaceType = nullptr) const
+	{
+		ElementList->GetElementHandles(OutArray, InBaseInterfaceType);
+	}
+
+	/**
+	 * Enumerate the handle of every selected element, optionally filtering to elements that implement the given interface.
+	 * @note Return true from the callback to continue enumeration.
+	 */
+	void ForEachSelectedElementHandle(TFunctionRef<bool(const FTypedElementHandle&)> InCallback, const TSubclassOf<UTypedElementInterface>& InBaseInterfaceType = nullptr) const
+	{
+		ElementList->ForEachElementHandle(InCallback, InBaseInterfaceType);
+	}
+
+	/**
+	 * Enumerate the selected elements that implement the given interface.
+	 * @note Return true from the callback to continue enumeration.
+	 */
+	template <typename BaseInterfaceType>
+	void ForEachSelectedElement(TFunctionRef<bool(const TTypedElement<BaseInterfaceType>&)> InCallback) const
+	{
+		ElementList->ForEachElement<BaseInterfaceType>(InCallback);
+	}
+
+	/**
+	 * Get the array of selected objects from the currently selected elements.
+	 */
+	UFUNCTION(BlueprintCallable, BlueprintPure=false, Category="TypedElementFramework|Selection")
+	TArray<UObject*> GetSelectedObjects(const UClass* InRequiredClass = nullptr) const;
+
+	/**
+	 * Get the array of selected objects from the currently selected elements.
+	 */
+	template <typename RequiredClassType>
+	TArray<RequiredClassType*> GetSelectedObjects() const
+	{
+		TArray<RequiredClassType*> SelectedObjects;
+		SelectedObjects.Reserve(ElementList->Num());
+
+		ForEachSelectedObject<RequiredClassType>([&SelectedObjects](RequiredClassType* InObject)
+		{
+			SelectedObjects.Add(InObject);
+			return true;
+		});
+
+		return SelectedObjects;
+	}
+
+	/**
+	 * Enumerate the selected objects from the currently selected elements.
+	 * @note Return true from the callback to continue enumeration.
+	 */
+	void ForEachSelectedObject(TFunctionRef<bool(UObject*)> InCallback, const UClass* InRequiredClass = nullptr) const;
+
+	/**
+	 * Enumerate the selected objects from the currently selected elements.
+	 * @note Return true from the callback to continue enumeration.
+	 */
+	template <typename RequiredClassType>
+	void ForEachSelectedObject(TFunctionRef<bool(RequiredClassType*)> InCallback) const
+	{
+		ForEachSelectedObject([&InCallback](UObject* InObject)
+		{
+			return InCallback(CastChecked<RequiredClassType>(InObject));
+		}, RequiredClassType::StaticClass());
+	}
+
 	/**
 	 * Access the delegate that is invoked whenever this element list is potentially about to change.
 	 * @note This may be called even if no actual change happens, though once a change does happen it won't be called again until after the next call to NotifyPendingChanges.
diff --git a/Engine/Source/Runtime/Engine/Private/Elements/Actor/ActorElementSelectionInterface.cpp b/Engine/Source/Runtime/Engine/Private/Elements/Actor/ActorElementSelectionInterface.cpp
--- a/Engine/Source/Runtime/Engine/Private/Elements/Actor/ActorElementSelectionInterface.cpp
+++ b/Engine/Source/Runtime/Engine/Private/Elements/Actor/ActorElementSelectionInterface.cpp
@@ -8,24 +8,24 @@
 int32 UActorElementSelectionInterface::GetNumSelectedActors(const UTypedElementList* InCurrentSelection)
 {
 	int32 NumSelected = 0;
-	for (const FTypedElementHandle& SelectedElement : *InCurrentSelection)
+	InCurrentSelection->ForEachElementHandle([&NumSelected](const FTypedElementHandle& InSelectedElement)
 	{
-		if (SelectedElement.GetData<FActorElementData>(/*bSilent*/true))
+		if (InSelectedElement.GetData<FActorElementData>(/*bSilent*/true))
 		{
 			++NumSelected;
 		}
-	}
+		return true;
+	});
 	return NumSelected;
 }
 
 bool UActorElementSelectionInterface::HasSelectedActors(const UTypedElementList* InCurrentSelection)
 {
-	for (const FTypedElementHandle& SelectedElement : *InCurrentSelection)
+	bool bHasSelectedActors = false;
+	InCurrentSelection->ForEachElementHandle([&bHasSelectedActors](const FTypedElementHandle& InSelectedElement)
 	{
-		if (SelectedElement.GetData<FActorElementData>(/*bSilent*/true))
-		{
-			return true;
-		}
-	}
-	return false;
+		bHasSelectedActors = InSelectedElement.GetData<FActorElementData>(/*bSilent*/true) != nullptr;
+		return !bHasSelectedActors;
+	});
+	return bHasSelectedActors;
 }
diff --git a/Engine/Source/Runtime/Engine/Private/Elements/Component/ComponentElementSelectionInterface.cpp b/Engine/Source/Runtime/Engine/Private/Elements/Component/ComponentElementSelectionInterface.cpp
--- a/Engine/Source/Runtime/Engine/Private/Elements/Component/ComponentElementSelectionInterface.cpp
+++ b/Engine/Source/Runtime/Engine/Private/Elements/Component/ComponentElementSelectionInterface.cpp
@@ -8,24 +8,24 @@
 int32 UComponentElementSelectionInterface::GetNumSelectedComponents(const UTypedElementList* InCurrentSelection)
 {
 	int32 NumSelected = 0;
-	for (const FTypedElementHandle& SelectedElement : *InCurrentSelection)
+	InCurrentSelection->ForEachElementHandle([&NumSelected](const FTypedElementHandle& InSelectedElement)
 	{
-		if (SelectedElement.GetData<FComponentElementData>(/*bSilent*/true))
+		if (InSelectedElement.GetData<FComponentElementData>(/*bSilent*/true))
 		{
 			++NumSelected;
 		}
-	}
+		return true;
+	});
 	return NumSelected;
 }
 
 bool UComponentElementSelectionInterface::HasSelectedComponents(const UTypedElementList* InCurrentSelection)
 {
-	for (const FTypedElementHandle& SelectedElement : *InCurrentSelection)
+	bool bHasSelectedComponents = false;
+	InCurrentSelection->ForEachElementHandle([&bHasSelectedComponents](const FTypedElementHandle& InSelectedElement)
 	{
-		if (SelectedElement.GetData<FComponentElementData>(/*bSilent*/true))
-		{
-			return true;
-		}
-	}
-	return false;
+		bHasSelectedComponents = InSelectedElement.GetData<FComponentElementData>(/*bSilent*/true) != nullptr;
+		return !bHasSelectedComponents;
+	});
+	return bHasSelectedComponents;
 }
diff --git a/Engine/Source/Runtime/TypedElementFramework/Private/Elements/Framework/TypedElementList.cpp b/Engine/Source/Runtime/TypedElementFramework/Private/Elements/Framework/TypedElementList.cpp
--- a/Engine/Source/Runtime/TypedElementFramework/Private/Elements/Framework/TypedElementList.cpp
+++ b/Engine/Source/Runtime/TypedElementFramework/Private/Elements/Framework/TypedElementList.cpp
@@ -131,6 +131,34 @@ UTypedElementInterface* UTypedElementList::GetElementInterface(const FTypedEleme
 	return Registry->GetElementInterface(InElementHandle, InBaseInterfaceType);
 }
 
+TArray<FTypedElementHandle> UTypedElementList::GetElementHandles(const TSubclassOf<UTypedElementInterface>& InBaseInterfaceType) const
+{
+	TArray<FTypedElementHandle> FilteredElementHandles;
+	FilteredElementHandles.Reserve(ElementHandles.Num());
+
+	ForEachElementHandle([&FilteredElementHandles](const FTypedElementHandle& InElementHandle)
+	{
+		FilteredElementHandles.Add(InElementHandle);
+		return true;
+	}, InBaseInterfaceType);
+
+	return FilteredElementHandles;
+}
+
+void UTypedElementList::ForEachElementHandle(TFunctionRef<bool(const FTypedElementHandle&)> InCallback, const TSubclassOf<UTypedElementInterface>& InBaseInterfaceType) const
+{
+	for (const FTypedElementHandle& ElementHandle : ElementHandles)
+	{
+		if (ElementHandle && (!InBaseInterfaceType || GetElementInterface(ElementHandle, InBaseInterfaceType)))
+		{
+			if (!InCallback(ElementHandle))
+			{
+				break;
+			}
+		}
+	}
+}
+
 bool UTypedElementList::AddElementImpl(FTypedElementHandle&& InElementHandle)
 {
 	if (!InElementHandle)
diff --git a/Engine/Source/Runtime/TypedElementFramework/Private/TypedElementFrameworkTests.cpp b/Engine/Source/Runtime/TypedElementFramework/Private/TypedElementFrameworkTests.cpp
--- a/Engine/Source/Runtime/TypedElementFramework/Private/TypedElementFrameworkTests.cpp
+++ b/Engine/Source/Runtime/TypedElementFramework/Private/TypedElementFrameworkTests.cpp
@@ -117,15 +117,17 @@ bool FTypedElementRegistrySmokeTest::RunTest(const FString& Parameters)
 	ElementList->Add(UntypedElement2);
 	ElementList->Add(UntypedElement3);
 	
-	for (const FTypedElementHandle& ElementHandle : *ElementList)
+	ElementList->ForEachElementHandle([&TestInterfaceAccess](const FTypedElementHandle& InElementHandle)
 	{
-		TestInterfaceAccess(ElementHandle);
-	}
+		TestInterfaceAccess(InElementHandle);
+		return true;
+	});
 	
-	for (const TTypedElement<UTestTypedElementInterfaceA>& ElementHandle : ElementList->IterateInterface<UTestTypedElementInterfaceA>())
+	ElementList->ForEachElement<UTestTypedElementInterfaceA>([&TestInterfaceAHandle](const TTypedElement<UTestTypedElementInterfaceA>& InElementHandle)
 	{
-		TestInterfaceAHandle(ElementHandle);
-	}
+		TestInterfaceAHandle(InElementHandle);
+		return true;
+	});
 
 	ElementList->Empty();
 	ElementList = nullptr;
@@ -195,27 +197,30 @@ bool FTypedElementRegistryPerfTest::RunTest(const FString& Parameters)
 	// Find an interface from each handle
 	{
 		FScopedDurationTimeLogger Timer(FString::Printf(TEXT("Finding %d interfaces from list"), ElementList->Num()));
-		for (const FTypedElementHandle& ElementHandle : *ElementList)
+		ElementList->ForEachElementHandle([Registry](const FTypedElementHandle& InElementHandle)
 		{
-			Registry->GetElementInterface<UTestTypedElementInterfaceA>(ElementHandle);
-		}
+			Registry->GetElementInterface<UTestTypedElementInterfaceA>(InElementHandle);
+			return true;
+		});
 	}
 
 	// Find an element from each handle
 	{
 		FScopedDurationTimeLogger Timer(FString::Printf(TEXT("Finding %d elements from list"), ElementList->Num()));
-		for (const FTypedElementHandle& ElementHandle : *ElementList)
+		ElementList->ForEachElementHandle([Registry](const FTypedElementHandle& InElementHandle)
 		{
-			Registry->GetElement<UTestTypedElementInterfaceA>(ElementHandle);
-		}
+			Registry->GetElement<UTestTypedElementInterfaceA>(InElementHandle);
+			return true;
+		});
 	}
 
-	// Iterate all elements that implement an interface
+	// Enumerate all elements that implement an interface
 	{
-		FScopedDurationTimeLogger Timer(FString::Printf(TEXT("Iterating %d elements in list"), ElementList->Num()));
-		for (const TTypedElement<UTestTypedElementInterfaceA>& ElementHandle : ElementList->IterateInterface<UTestTypedElementInterfaceA>())
+		FScopedDurationTimeLogger Timer(FString::Printf(TEXT("Enumerating %d elements in list"), ElementList->Num()));
+		ElementList->ForEachElement<UTestTypedElementInterfaceA>([Registry](const TTypedElement<UTestTypedElementInterfaceA>& InElementHandle)
 		{
-		}
+			return true;
+		});
 	}
 
 	// Clear the element list
diff --git a/Engine/Source/Runtime/TypedElementFramework/Public/Elements/Framework/TypedElementList.h b/Engine/Source/Runtime/TypedElementFramework/Public/Elements/Framework/TypedElementList.h
--- a/Engine/Source/Runtime/TypedElementFramework/Public/Elements/Framework/TypedElementList.h
+++ b/Engine/Source/Runtime/TypedElementFramework/Public/Elements/Framework/TypedElementList.h
@@ -34,102 +34,6 @@ FORCEINLINE TTypedElement<BaseInterfaceType> GetElement(const UTypedElementRegis
 	return Element;
 }
 
-/**
- * Iterator type for ranged-for loops which only returns handles that implement the given interface.
- * @note This iterator type only supports the minimal functionality needed to support C++ ranged-for syntax.
- */
-template <typename BaseInterfaceType>
-struct TTypedElementListInterfaceIterator
-{
-public:
-	TTypedElementListInterfaceIterator(const UTypedElementRegistry* InRegistry, const TArray<FTypedElementHandle>& InElementHandles, const int32 InIndex)
-		: Registry(InRegistry)
-		, ElementHandles(InElementHandles)
-		, InitialNum(InElementHandles.Num())
-		, Index(InIndex)
-	{
-		// Walk to the first valid element
-		while (Index < ElementHandles.Num())
-		{
-			TypedElementList_Private::GetElement(Registry, ElementHandles[Index], Element);
-			if (Element)
-			{
-				break;
-			}
-			++Index;
-		}
-	}
-
-	FORCEINLINE const TTypedElement<BaseInterfaceType>& operator*() const
-	{
-		return Element;
-	}
-
-	TTypedElementListInterfaceIterator& operator++()
-	{
-		// Walk to the next valid element
-		while (++Index < ElementHandles.Num())
-		{
-			TypedElementList_Private::GetElement(Registry, ElementHandles[Index], Element);
-			if (Element)
-			{
-				break;
-			}
-		}
-		return *this;
-	}
-
-private:
-	const UTypedElementRegistry* Registry;
-	const TArray<FTypedElementHandle>& ElementHandles;
-	const int32 InitialNum;
-	int32 Index;
-	TTypedElement<BaseInterfaceType> Element;
-
-	FORCEINLINE friend bool operator!=(const TTypedElementListInterfaceIterator& Lhs, const TTypedElementListInterfaceIterator& Rhs)
-	{
-		// We only need to do the check in this operator, because no other operator will be
-		// called until after this one returns.
-		//
-		// Also, we should only need to check one side of this comparison - if the other iterator isn't
-		// even from the same array then the compiler has generated bad code.
-		ensureMsgf(Lhs.ElementHandles.Num() == Lhs.InitialNum, TEXT("Array has changed during ranged-for iteration!"));
-		return Lhs.Index != Rhs.Index;
-	}
-};
-
-/**
- * Proxy type to enable ranged-for loops which only returns handles that implement the given interface.
- */
-template <typename BaseInterfaceType>
-struct TTypedElementListInterfaceIteratorProxy
-{
-public:
-	TTypedElementListInterfaceIteratorProxy(const UTypedElementRegistry* InRegistry, const TArray<FTypedElementHandle>& InElementHandles)
-		: Registry(InRegistry)
-		, ElementHandles(InElementHandles)
-	{
-	}
-
-	TTypedElementListInterfaceIteratorProxy(const TTypedElementListInterfaceIteratorProxy&) = delete;
-	TTypedElementListInterfaceIteratorProxy& operator=(const TTypedElementListInterfaceIteratorProxy&) = delete;
-
-	TTypedElementListInterfaceIteratorProxy(TTypedElementListInterfaceIteratorProxy&&) = default;
-	TTypedElementListInterfaceIteratorProxy& operator=(TTypedElementListInterfaceIteratorProxy&&) = default;
-
-	/**
-	 * DO NOT USE DIRECTLY
-	 * STL-like iterators to enable range-based for loop support.
-	 */
-	using RangedForConstIteratorType = TTypedElementListInterfaceIterator<BaseInterfaceType>;
-	FORCEINLINE RangedForConstIteratorType begin() const { return RangedForConstIteratorType(Registry, ElementHandles, 0); }
-	FORCEINLINE RangedForConstIteratorType end() const { return RangedForConstIteratorType(Registry, ElementHandles, ElementHandles.Num()); }
-
-private:
-	const UTypedElementRegistry* Registry;
-	const TArray<FTypedElementHandle>& ElementHandles;
-};
-
 } // namespace TypedElementList_Private
 
 /**
@@ -314,6 +218,51 @@ public:
 	UFUNCTION(BlueprintPure, Category="TypedElementFramework|List")
 	UTypedElementInterface* GetElementInterface(const FTypedElementHandle& InElementHandle, const TSubclassOf<UTypedElementInterface>& InBaseInterfaceType) const;
 
+	/**
+	 * Get the handle of every element in this list, optionally filtering to elements that implement the given interface.
+	 */
+	UFUNCTION(BlueprintCallable, BlueprintPure=false, Category="TypedElementFramework|List")
+	TArray<FTypedElementHandle> GetElementHandles(const TSubclassOf<UTypedElementInterface>& InBaseInterfaceType = nullptr) const;
+
+	/**
+	 * Get the handle of every element in this list, optionally filtering to elements that implement the given interface.
+	 */
+	template <typename ArrayAllocator>
+	void GetElementHandles(TArray<FTypedElementHandle, ArrayAllocator>& OutArray, const TSubclassOf<UTypedElementInterface>& InBaseInterfaceType = nullptr) const
+	{
+		OutArray.Reset();
+		OutArray.Reserve(ElementHandles.Num());
+		ForEachElementHandle([&OutArray](const FTypedElementHandle& InElementHandle)
+		{
+			OutArray.Add(InElementHandle);
+			return true;
+		}, InBaseInterfaceType);
+	}
+
+	/**
+	 * Enumerate the handle of every element in this list, optionally filtering to elements that implement the given interface.
+	 * @note Return true from the callback to continue enumeration.
+	 */
+	void ForEachElementHandle(TFunctionRef<bool(const FTypedElementHandle&)> InCallback, const TSubclassOf<UTypedElementInterface>& InBaseInterfaceType = nullptr) const;
+
+	/**
+	 * Enumerate the elements in this list that implement the given interface.
+	 * @note Return true from the callback to continue enumeration.
+	 */
+	template <typename BaseInterfaceType>
+	void ForEachElement(TFunctionRef<bool(const TTypedElement<BaseInterfaceType>&)> InCallback) const
+	{
+		TTypedElement<BaseInterfaceType> TempElement;
+		for (const FTypedElementHandle& ElementHandle : ElementHandles)
+		{
+			GetElement(ElementHandle, TempElement);
+			if (TempElement && !InCallback(TempElement))
+			{
+				break;
+			}
+		}
+	}
+
 	/**
 	 * Is the given index a valid entry within this element list?
 	 */
@@ -573,22 +522,6 @@ public:
 	 */
 	FTypedElementListLegacySync* Legacy_GetSyncPtr();
 
-	/**
-	 * Get an iterator for elements that implement the given element interface.
-	 */
-	template <typename BaseInterfaceType>
-	FORCEINLINE TypedElementList_Private::TTypedElementListInterfaceIteratorProxy<BaseInterfaceType> IterateInterface() const
-	{
-		return TypedElementList_Private::TTypedElementListInterfaceIteratorProxy<BaseInterfaceType>(Registry.Get(), ElementHandles);
-	}
-
-	/**
-	 * DO NOT USE DIRECTLY
-	 * STL-like iterators to enable range-based for loop support.
-	 */
-	FORCEINLINE TArray<FTypedElementHandle>::RangedForConstIteratorType begin() const { return ElementHandles.begin(); }
-	FORCEINLINE TArray<FTypedElementHandle>::RangedForConstIteratorType end() const { return ElementHandles.end(); }
-
 private:
 	enum class EChangeType : uint8
 	{
