Trim 4 bytes off FTypedElementHandle and FTypedElementOwner by storing the FTypedElementId in the shared payload instance

#rnx

diff --git a/Engine/Source/Runtime/TypedElementFramework/Private/Elements/Framework/TypedElementRegistry.cpp b/Engine/Source/Runtime/TypedElementFramework/Private/Elements/Framework/TypedElementRegistry.cpp
--- a/Engine/Source/Runtime/TypedElementFramework/Private/Elements/Framework/TypedElementRegistry.cpp
+++ b/Engine/Source/Runtime/TypedElementFramework/Private/Elements/Framework/TypedElementRegistry.cpp
@@ -139,7 +139,7 @@ FTypedElementHandle UTypedElementRegistry::GetElementHandle(const FTypedElementI
 	checkf(RegisteredElementType, TEXT("Element type ID '%d' has not been registered!"), InElementId.GetTypeId());
 
 	FTypedElementHandle ElementHandle;
-	ElementHandle.Private_InitializeAddRef(InElementId.GetTypeId(), InElementId.GetElementId(), RegisteredElementType->GetDataForElement(InElementId.GetElementId()));
+	ElementHandle.Private_InitializeAddRef(RegisteredElementType->GetDataForElement(InElementId.GetElementId()));
 
 	return ElementHandle;
 }
diff --git a/Engine/Source/Runtime/TypedElementFramework/Public/Elements/Framework/TypedElementData.h b/Engine/Source/Runtime/TypedElementFramework/Public/Elements/Framework/TypedElementData.h
--- a/Engine/Source/Runtime/TypedElementFramework/Public/Elements/Framework/TypedElementData.h
+++ b/Engine/Source/Runtime/TypedElementFramework/Public/Elements/Framework/TypedElementData.h
@@ -8,7 +8,7 @@
 #include "HAL/PlatformStackWalk.h"
 #include "Containers/SparseArray.h"
 #include "Containers/ChunkedArray.h"
-#include "Elements/Framework/TypedElementLimits.h"
+#include "Elements/Framework/TypedElementId.h"
 
 /**
  * Macro to declare the required RTTI data for types representing element data.
@@ -75,10 +75,20 @@ public:
 	FTypedElementInternalData(FTypedElementInternalData&& InOther) = delete;
 	FTypedElementInternalData& operator=(FTypedElementInternalData&&) = delete;
 
-	virtual ~FTypedElementInternalData() = default;
+	virtual ~FTypedElementInternalData()
+	{
+		Id.Private_DestroyNoRef();
+	}
+
+	void Initialize(const FTypedHandleTypeId InTypeId, const FTypedHandleElementId InElementId)
+	{
+		checkSlow(!Id.IsSet());
+		Id.Private_InitializeNoRef(InTypeId, InElementId);
+	}
 
 	virtual void Reset()
 	{
+		Id.Private_DestroyNoRef();
 #if UE_TYPED_ELEMENT_HAS_REFCOUNTING
 		FPlatformAtomics::InterlockedExchange(&RefCount, 0);
 #endif	// UE_TYPED_ELEMENT_HAS_REFCOUNTING
@@ -90,6 +100,11 @@ public:
 #endif	// UE_TYPED_ELEMENT_HAS_REFTRACKING
 	}
 
+	FORCEINLINE const FTypedElementId& GetId() const
+	{
+		return Id;
+	}
+
 	FORCEINLINE FTypedElementReferenceId AddRef(const bool bCanTrackReference) const
 	{
 		FTypedElementReferenceId ReferenceId = INDEX_NONE;
@@ -155,6 +170,7 @@ public:
 	}
 
 private:
+	FTypedElementId Id;
 #if UE_TYPED_ELEMENT_HAS_REFCOUNTING
 	mutable FTypedElementRefCount RefCount = 0;
 #endif	// UE_TYPED_ELEMENT_HAS_REFCOUNTING
@@ -214,7 +230,7 @@ class TTypedElementInternalDataStore
 public:
 	static_assert(TNumericLimits<int32>::Max() >= TypedHandleMaxElementId, "TTypedElementInternalDataStore internally uses signed 32-bit indices so cannot store TypedHandleMaxElementId! Consider making this container 64-bit aware, or explicitly remove this compile time check.");
 
-	TTypedElementInternalData<ElementDataType>& AddDataForElement(FTypedHandleElementId& InOutElementId)
+	TTypedElementInternalData<ElementDataType>& AddDataForElement(const FTypedHandleTypeId InTypeId, FTypedHandleElementId& InOutElementId)
 	{
 		FWriteScopeLock InternalDataLock(InternalDataRW);
 
@@ -224,7 +240,9 @@ public:
 			? InternalDataFreeIndices.Pop(/*bAllowShrinking*/false)
 			: InternalDataArray.Add();
 
-		return InternalDataArray[InOutElementId];
+		TTypedElementInternalData<ElementDataType>& InternalData = InternalDataArray[InOutElementId];
+		InternalData.Initialize(InTypeId, InOutElementId);
+		return InternalData;
 	}
 
 	void RemoveDataForElement(const FTypedHandleElementId InElementId, const FTypedElementInternalData* InExpectedDataPtr)
@@ -279,9 +297,8 @@ class TTypedElementInternalDataStore<void>
 public:
 	static_assert(TNumericLimits<int32>::Max() >= TypedHandleMaxElementId, "TTypedElementInternalDataStore internally uses signed 32-bit indices so cannot store TypedHandleMaxElementId! Consider making this container 64-bit aware, or explicitly remove this compile time check.");
 
-	TTypedElementInternalData<void>& AddDataForElement(FTypedHandleElementId& InOutElementId)
+	TTypedElementInternalData<void>& AddDataForElement(const FTypedHandleTypeId InTypeId, FTypedHandleElementId& InOutElementId)
 	{
-#if UE_TYPED_ELEMENT_HAS_REFERENCING
 		FWriteScopeLock InternalDataLock(InternalDataRW);
 
 		checkSlow(InOutElementId >= 0);
@@ -292,15 +309,14 @@ public:
 			: InternalDataArray.Add();
 
 		ElementIdToArrayIndex.Add(InOutElementId, InternalDataArrayIndex);
-		return InternalDataArray[InternalDataArrayIndex];
-#else	// UE_TYPED_ELEMENT_HAS_REFERENCING
-		return SharedInternalData;
-#endif	// UE_TYPED_ELEMENT_HAS_REFERENCING
+
+		TTypedElementInternalData<void>& InternalData = InternalDataArray[InOutElementId];
+		InternalData.Initialize(InTypeId, InOutElementId);
+		return InternalData;
 	}
 
 	void RemoveDataForElement(const FTypedHandleElementId InElementId, const FTypedElementInternalData* InExpectedDataPtr)
 	{
-#if UE_TYPED_ELEMENT_HAS_REFERENCING
 		FWriteScopeLock InternalDataLock(InternalDataRW);
 
 		int32 InternalDataArrayIndex = INDEX_NONE;
@@ -312,22 +328,15 @@ public:
 		checkf(InExpectedDataPtr == &InternalData, TEXT("Internal data pointer did not match the expected value! Does this handle belong to a different element registry?"));
 		InternalData.Reset();
 		InternalDataFreeIndices.Add(InternalDataArrayIndex);
-#else	// UE_TYPED_ELEMENT_HAS_REFERENCING
-		checkf(InExpectedDataPtr == &SharedInternalData, TEXT("Internal data pointer did not match the expected value! Does this handle belong to a different element registry?"));
-#endif	// UE_TYPED_ELEMENT_HAS_REFERENCING
 	}
 
 	const TTypedElementInternalData<void>& GetDataForElement(const FTypedHandleElementId InElementId) const
 	{
-#if UE_TYPED_ELEMENT_HAS_REFERENCING
 		FReadScopeLock InternalDataLock(InternalDataRW);
 
 		const int32* InternalDataArrayIndexPtr = ElementIdToArrayIndex.Find(InElementId);
 		checkSlow(InternalDataArrayIndexPtr && InternalDataArray.IsValidIndex(*InternalDataArrayIndexPtr));
 		return InternalDataArray[*InternalDataArrayIndexPtr];
-#else	// UE_TYPED_ELEMENT_HAS_REFERENCING
-		return SharedInternalData;
-#endif	// UE_TYPED_ELEMENT_HAS_REFERENCING
 	}
 
 	static FORCEINLINE void SetStaticDataTypeId(const FTypedHandleTypeId InTypeId)
@@ -345,12 +354,8 @@ public:
 	}
 
 private:
-#if UE_TYPED_ELEMENT_HAS_REFERENCING
 	mutable FRWLock InternalDataRW;
 	TChunkedArray<TTypedElementInternalData<void>> InternalDataArray;
 	TArray<int32> InternalDataFreeIndices;
 	TMap<FTypedHandleElementId, int32> ElementIdToArrayIndex;
-#else	// UE_TYPED_ELEMENT_HAS_REFERENCING
-	TTypedElementInternalData<void> SharedInternalData;
-#endif	// UE_TYPED_ELEMENT_HAS_REFERENCING
 };
diff --git a/Engine/Source/Runtime/TypedElementFramework/Public/Elements/Framework/TypedElementHandle.h b/Engine/Source/Runtime/TypedElementFramework/Public/Elements/Framework/TypedElementHandle.h
--- a/Engine/Source/Runtime/TypedElementFramework/Public/Elements/Framework/TypedElementHandle.h
+++ b/Engine/Source/Runtime/TypedElementFramework/Public/Elements/Framework/TypedElementHandle.h
@@ -8,127 +8,6 @@
 #include "Elements/Interfaces/TypedElementInterface.h"
 #include "TypedElementHandle.generated.h"
 
-/**
- * The most minimal representation of an element - its ID!
- * This type is not immediately useful on its own, but can be used to find an element from the element registry or an element list.
- * @note This is ref-counted like handles themselves are, so as long as an ID is available, the handle will be too.
- * @note IDs lack the information needed to auto-release on destruction, so must be manually released, either via the corresponding handle or their owner element registry.
- */
-struct TYPEDELEMENTFRAMEWORK_API FTypedElementId
-{
-public:
-	FTypedElementId()
-		: CombinedId(0)
-	{
-	}
-
-	FTypedElementId(const FTypedElementId&) = delete;
-	FTypedElementId& operator=(const FTypedElementId&) = delete;
-	
-	FTypedElementId(FTypedElementId&& InOther)
-		: CombinedId(InOther.CombinedId)
-	{
-		InOther.Private_DestroyNoRef();
-	}
-
-	FTypedElementId& operator=(FTypedElementId&& InOther)
-	{
-		if (this != &InOther)
-		{
-			CombinedId = InOther.CombinedId;
-
-			InOther.Private_DestroyNoRef();
-		}
-		return *this;
-	}
-
-	~FTypedElementId()
-	{
-		checkf(!IsSet(), TEXT("Element ID was still set during destruction! This will leak an element reference, and you should release this ID prior to destruction!"));
-	}
-
-	FORCEINLINE explicit operator bool() const
-	{
-		return IsSet();
-	}
-
-	/**
-	 * Has this ID been initialized to a valid element?
-	 */
-	FORCEINLINE bool IsSet() const
-	{
-		return TypeId != 0;
-	}
-
-	/**
-	 * Access the type ID portion of this element ID.
-	 */
-	FORCEINLINE FTypedHandleTypeId GetTypeId() const
-	{
-		return TypeId;
-	}
-
-	/**
-	 * Access the element ID portion of this element ID.
-	 */
-	FORCEINLINE FTypedHandleElementId GetElementId() const
-	{
-		return ElementId;
-	}
-
-	/**
-	 * Access the combined value of this element ID.
-	 * @note You typically don't want to store this directly as the element ID could be re-used.
-	 *       It is primarily useful as a secondary cache where something is keeping a reference to an element ID or element handle (eg, how UTypedElementList uses it internally).
-	 */
-	FORCEINLINE FTypedHandleCombinedId GetCombinedId() const
-	{
-		return CombinedId;
-	}
-
-	FORCEINLINE friend bool operator==(const FTypedElementId& InLHS, const FTypedElementId& InRHS)
-	{
-		return InLHS.CombinedId == InRHS.CombinedId;
-	}
-
-	FORCEINLINE friend bool operator!=(const FTypedElementId& InLHS, const FTypedElementId& InRHS)
-	{
-		return !(InLHS == InRHS);
-	}
-
-	friend inline uint32 GetTypeHash(const FTypedElementId& InElementId)
-	{
-		return GetTypeHash(InElementId.CombinedId);
-		
-	}
-
-	FORCEINLINE void Private_InitializeNoRef(const FTypedHandleTypeId InTypeId, const FTypedHandleElementId InElementId)
-	{
-		TypeId = InTypeId;
-		ElementId = InElementId;
-	}
-
-	FORCEINLINE void Private_DestroyNoRef()
-	{
-		CombinedId = 0;
-	}
-
-	/** An unset element ID */
-	static const FTypedElementId Unset;
-
-private:
-	union
-	{
-		struct
-		{
-			// Note: These are arranged in this order to give CombinedId better hash distribution for GetTypeHash!
-			FTypedHandleCombinedId ElementId : TypedHandleElementIdBits;
-			FTypedHandleCombinedId TypeId : TypedHandleTypeIdBits;
-		};
-		FTypedHandleCombinedId CombinedId;
-	};
-};
-
 /**
  * A representation of an element that includes its handle data.
  * This type is the most standard way that an element is passed through to interfaces, and also the type that is stored in element lists.
@@ -147,7 +26,7 @@ public:
 	{
 		if (InOther)
 		{
-			Private_InitializeAddRef(InOther.Id.GetTypeId(), InOther.Id.GetElementId(), *InOther.DataPtr);
+			Private_InitializeAddRef(*InOther.DataPtr);
 		}
 	}
 
@@ -159,15 +38,14 @@ public:
 
 			if (InOther)
 			{
-				Private_InitializeAddRef(InOther.Id.GetTypeId(), InOther.Id.GetElementId(), *InOther.DataPtr);
+				Private_InitializeAddRef(*InOther.DataPtr);
 			}
 		}
 		return *this;
 	}
 
 	FTypedElementHandle(FTypedElementHandle&& InOther)
-		: Id(MoveTemp(InOther.Id))
-		, DataPtr(InOther.DataPtr)
+		: DataPtr(InOther.DataPtr)
 #if UE_TYPED_ELEMENT_HAS_REFTRACKING
 		, ReferenceId(InOther.ReferenceId)
 #endif	// UE_TYPED_ELEMENT_HAS_REFTRACKING
@@ -185,7 +63,6 @@ public:
 		{
 			Private_DestroyReleaseRef();
 
-			Id = MoveTemp(InOther.Id);
 			DataPtr = InOther.DataPtr;
 #if UE_TYPED_ELEMENT_HAS_REFTRACKING
 			ReferenceId = InOther.ReferenceId;
@@ -215,7 +92,7 @@ public:
 	 */
 	FORCEINLINE bool IsSet() const
 	{
-		return Id.IsSet();
+		return DataPtr != nullptr;
 	}
 
 	/**
@@ -231,7 +108,9 @@ public:
 	 */
 	FORCEINLINE const FTypedElementId& GetId() const
 	{
-		return Id;
+		return DataPtr
+			? DataPtr->GetId()
+			: FTypedElementId::Unset;
 	}
 
 	/**
@@ -241,7 +120,7 @@ public:
 	template <typename ElementDataType>
 	FORCEINLINE bool IsDataOfType() const
 	{
-		return Id.GetTypeId() == ElementDataType::StaticTypeId();
+		return GetId().GetTypeId() == ElementDataType::StaticTypeId();
 	}
 
 	/**
@@ -264,7 +143,7 @@ public:
 		{
 			if (!bSilent)
 			{
-				FFrame::KismetExecutionMessage(*FString::Printf(TEXT("Element handle data type is '%d', but '%d' (%s) was requested!"), Id.GetTypeId(), ElementDataType::StaticTypeId(), *ElementDataType::StaticTypeName().ToString()), ELogVerbosity::Error);
+				FFrame::KismetExecutionMessage(*FString::Printf(TEXT("Element handle data type is '%d', but '%d' (%s) was requested!"), GetId().GetTypeId(), ElementDataType::StaticTypeId(), *ElementDataType::StaticTypeName().ToString()), ELogVerbosity::Error);
 			}
 			return nullptr;
 		}
@@ -280,7 +159,7 @@ public:
 	FORCEINLINE const ElementDataType& GetDataChecked() const
 	{
 		checkf(DataPtr, TEXT("Element handle data is null!"));
-		checkf(IsDataOfType<ElementDataType>(), TEXT("Element handle data type is '%d', but '%d' (%s) was requested!"), Id.GetTypeId(), ElementDataType::StaticTypeId(), *ElementDataType::StaticTypeName().ToString());
+		checkf(IsDataOfType<ElementDataType>(), TEXT("Element handle data type is '%d', but '%d' (%s) was requested!"), GetId().GetTypeId(), ElementDataType::StaticTypeId(), *ElementDataType::StaticTypeName().ToString());
 		return *static_cast<const ElementDataType*>(DataPtr->GetUntypedData());
 	}
 	
@@ -294,7 +173,7 @@ public:
 		if (IsSet())
 		{
 			DataPtr->AddRef(/*bCanTrackReference*/false); // Cannot track element ID references as we have no space to store the reference ID
-			ElementId.Private_InitializeNoRef(Id.GetTypeId(), Id.GetElementId());
+			ElementId.Private_InitializeNoRef(DataPtr->GetId().GetTypeId(), DataPtr->GetId().GetElementId());
 		}
 		return ElementId;
 	}
@@ -305,7 +184,7 @@ public:
 	 */
 	void ReleaseId(FTypedElementId& InOutElementId) const
 	{
-		checkf(InOutElementId == Id, TEXT("Element ID does not match this handle!"));
+		checkf(InOutElementId == GetId(), TEXT("Element ID does not match this handle!"));
 		if (InOutElementId)
 		{
 			DataPtr->ReleaseRef(INDEX_NONE); // Cannot track element ID references as we have no space to store the reference ID
@@ -315,7 +194,7 @@ public:
 
 	FORCEINLINE friend bool operator==(const FTypedElementHandle& InLHS, const FTypedElementHandle& InRHS)
 	{
-		return InLHS.Id == InRHS.Id;
+		return InLHS.DataPtr == InRHS.DataPtr;
 	}
 
 	FORCEINLINE friend bool operator!=(const FTypedElementHandle& InLHS, const FTypedElementHandle& InRHS)
@@ -325,24 +204,22 @@ public:
 
 	friend inline uint32 GetTypeHash(const FTypedElementHandle& InElementHandle)
 	{
-		return GetTypeHash(InElementHandle.Id);
+		return GetTypeHash(InElementHandle.GetId());
 	}
 
-	FORCEINLINE void Private_InitializeNoRef(const FTypedHandleTypeId InTypeId, const FTypedHandleElementId InElementId, const FTypedElementInternalData& InData)
+	FORCEINLINE void Private_InitializeNoRef(const FTypedElementInternalData& InData)
 	{
-		Id.Private_InitializeNoRef(InTypeId, InElementId);
 		DataPtr = &InData;
 	}
 
-	FORCEINLINE void Private_InitializeAddRef(const FTypedHandleTypeId InTypeId, const FTypedHandleElementId InElementId, const FTypedElementInternalData& InData)
+	FORCEINLINE void Private_InitializeAddRef(const FTypedElementInternalData& InData)
 	{
-		Private_InitializeNoRef(InTypeId, InElementId, InData);
+		Private_InitializeNoRef(InData);
 		RegisterRef();
 	}
 
 	FORCEINLINE void Private_DestroyNoRef()
 	{
-		Id.Private_DestroyNoRef();
 		DataPtr = nullptr;
 #if UE_TYPED_ELEMENT_HAS_REFTRACKING
 		ReferenceId = INDEX_NONE;
@@ -383,7 +260,6 @@ private:
 		}
 	}
 
-	FTypedElementId Id;
 	const FTypedElementInternalData* DataPtr = nullptr;
 #if UE_TYPED_ELEMENT_HAS_REFTRACKING
 	FTypedElementReferenceId ReferenceId = INDEX_NONE;
@@ -501,15 +377,15 @@ public:
 		return HashCombine(GetTypeHash(static_cast<const FTypedElementHandle&>(InElement)), GetTypeHash(InElement.InterfacePtr));
 	}
 
-	FORCEINLINE void Private_InitializeNoRef(const FTypedHandleTypeId InTypeId, const FTypedHandleElementId InElementId, const FTypedElementInternalData& InData, BaseInterfaceType* InInterfacePtr)
+	FORCEINLINE void Private_InitializeNoRef(const FTypedElementInternalData& InData, BaseInterfaceType* InInterfacePtr)
 	{
-		FTypedElementHandle::Private_InitializeNoRef(InTypeId, InElementId, InData);
+		FTypedElementHandle::Private_InitializeNoRef(InData);
 		InterfacePtr = InInterfacePtr;
 	}
 
-	FORCEINLINE void Private_InitializeAddRef(const FTypedHandleTypeId InTypeId, const FTypedHandleElementId InElementId, const FTypedElementInternalData& InData, BaseInterfaceType* InInterfacePtr)
+	FORCEINLINE void Private_InitializeAddRef(const FTypedElementInternalData& InData, BaseInterfaceType* InInterfacePtr)
 	{
-		FTypedElementHandle::Private_InitializeAddRef(InTypeId, InElementId, InData);
+		FTypedElementHandle::Private_InitializeAddRef(InData);
 		InterfacePtr = InInterfacePtr;
 	}
 
@@ -659,8 +535,7 @@ public:
 	TTypedElementOwner& operator=(const TTypedElementOwner&) = delete;
 
 	TTypedElementOwner(TTypedElementOwner&& InOther)
-		: Id(MoveTemp(InOther.Id))
-		, DataPtr(InOther.DataPtr)
+		: DataPtr(InOther.DataPtr)
 	{
 		InOther.DataPtr = nullptr;
 		checkSlow(!InOther.IsSet());
@@ -670,7 +545,6 @@ public:
 	{
 		if (this != &InOther)
 		{
-			Id = MoveTemp(InOther.Id);
 			DataPtr = InOther.DataPtr;
 
 			InOther.DataPtr = nullptr;
@@ -694,7 +568,7 @@ public:
 	 */
 	FORCEINLINE bool IsSet() const
 	{
-		return Id.IsSet();
+		return DataPtr != nullptr;
 	}
 
 	/**
@@ -702,7 +576,9 @@ public:
 	 */
 	FORCEINLINE const FTypedElementId& GetId() const
 	{
-		return Id;
+		return DataPtr
+			? DataPtr->GetId()
+			: FTypedElementId::Unset;
 	}
 
 	/**
@@ -738,7 +614,7 @@ public:
 		if (IsSet())
 		{
 			DataPtr->AddRef(/*bCanTrackReference*/false); // Cannot track element ID references as we have no space to store the reference ID
-			ElementId.Private_InitializeNoRef(Id.GetTypeId(), Id.GetElementId());
+			ElementId.Private_InitializeNoRef(DataPtr->GetId().GetTypeId(), DataPtr->GetId().GetElementId());
 		}
 		return ElementId;
 	}
@@ -749,7 +625,7 @@ public:
 	 */
 	void ReleaseId(FTypedElementId& InOutElementId) const
 	{
-		checkf(InOutElementId == Id, TEXT("Element ID does not match this owner!"));
+		checkf(InOutElementId == GetId(), TEXT("Element ID does not match this owner!"));
 		if (InOutElementId)
 		{
 			DataPtr->ReleaseRef(INDEX_NONE); // Cannot track element ID references as we have no space to store the reference ID
@@ -766,7 +642,7 @@ public:
 		FTypedElementHandle ElementHandle;
 		if (IsSet())
 		{
-			ElementHandle.Private_InitializeAddRef(Id.GetTypeId(), Id.GetElementId(), *DataPtr);
+			ElementHandle.Private_InitializeAddRef(*DataPtr);
 		}
 		return ElementHandle;
 	}
@@ -777,13 +653,13 @@ public:
 	 */
 	void ReleaseHandle(FTypedElementHandle& InOutElementHandle) const
 	{
-		checkf(InOutElementHandle.GetId() == Id, TEXT("Element handle ID does not match this owner!"));
+		checkf(InOutElementHandle.GetId() == GetId(), TEXT("Element handle ID does not match this owner!"));
 		InOutElementHandle.Release();
 	}
 
 	FORCEINLINE friend bool operator==(const TTypedElementOwner& InLHS, const TTypedElementOwner& InRHS)
 	{
-		return InLHS.Id == InRHS.Id;
+		return InLHS.DataPtr == InRHS.DataPtr;
 	}
 
 	FORCEINLINE friend bool operator!=(const TTypedElementOwner& InLHS, const TTypedElementOwner& InRHS)
@@ -793,24 +669,22 @@ public:
 
 	friend inline uint32 GetTypeHash(const TTypedElementOwner& InElementOwner)
 	{
-		return GetTypeHash(InElementOwner.Id);
+		return GetTypeHash(InElementOwner.GetId());
 	}
 
-	FORCEINLINE void Private_InitializeNoRef(const FTypedHandleTypeId InTypeId, const FTypedHandleElementId InElementId, TTypedElementInternalData<ElementDataType>& InData)
+	FORCEINLINE void Private_InitializeNoRef(TTypedElementInternalData<ElementDataType>& InData)
 	{
-		Id.Private_InitializeNoRef(InTypeId, InElementId);
 		DataPtr = &InData;
 	}
 
-	FORCEINLINE void Private_InitializeAddRef(const FTypedHandleTypeId InTypeId, const FTypedHandleElementId InElementId, TTypedElementInternalData<ElementDataType>& InData)
+	FORCEINLINE void Private_InitializeAddRef(TTypedElementInternalData<ElementDataType>& InData)
 	{
-		Private_InitializeNoRef(InTypeId, InElementId, InData);
+		Private_InitializeNoRef(InData);
 		RegisterRef();
 	}
 
 	FORCEINLINE void Private_DestroyNoRef()
 	{
-		Id.Private_DestroyNoRef();
 		DataPtr = nullptr;
 #if UE_TYPED_ELEMENT_HAS_REFTRACKING
 		ReferenceId = INDEX_NONE;
@@ -851,7 +725,6 @@ private:
 		}
 	}
 
-	FTypedElementId Id;
 	TTypedElementInternalData<ElementDataType>* DataPtr = nullptr;
 #if UE_TYPED_ELEMENT_HAS_REFTRACKING
 	FTypedElementReferenceId ReferenceId = INDEX_NONE;
diff --git a/Engine/Source/Runtime/TypedElementFramework/Public/Elements/Framework/TypedElementId.h b/Engine/Source/Runtime/TypedElementFramework/Public/Elements/Framework/TypedElementId.h
new file mode 100644
--- /dev/null
+++ b/Engine/Source/Runtime/TypedElementFramework/Public/Elements/Framework/TypedElementId.h
@@ -0,0 +1,127 @@
+// Copyright Epic Games, Inc. All Rights Reserved.
+
+#pragma once
+
+//#include "CoreMinimal.h"
+#include "Elements/Framework/TypedElementLimits.h"
+
+/**
+ * The most minimal representation of an element - its ID!
+ * This type is not immediately useful on its own, but can be used to find an element from the element registry or an element list.
+ * @note This is ref-counted like handles themselves are, so as long as an ID is available, the handle will be too.
+ * @note IDs lack the information needed to auto-release on destruction, so must be manually released, either via the corresponding handle or their owner element registry.
+ */
+struct TYPEDELEMENTFRAMEWORK_API FTypedElementId
+{
+public:
+	FTypedElementId()
+		: CombinedId(0)
+	{
+	}
+
+	FTypedElementId(const FTypedElementId&) = delete;
+	FTypedElementId& operator=(const FTypedElementId&) = delete;
+	
+	FTypedElementId(FTypedElementId&& InOther)
+		: CombinedId(InOther.CombinedId)
+	{
+		InOther.Private_DestroyNoRef();
+	}
+
+	FTypedElementId& operator=(FTypedElementId&& InOther)
+	{
+		if (this != &InOther)
+		{
+			CombinedId = InOther.CombinedId;
+
+			InOther.Private_DestroyNoRef();
+		}
+		return *this;
+	}
+
+	~FTypedElementId()
+	{
+		checkf(!IsSet(), TEXT("Element ID was still set during destruction! This will leak an element reference, and you should release this ID prior to destruction!"));
+	}
+
+	FORCEINLINE explicit operator bool() const
+	{
+		return IsSet();
+	}
+
+	/**
+	 * Has this ID been initialized to a valid element?
+	 */
+	FORCEINLINE bool IsSet() const
+	{
+		return TypeId != 0;
+	}
+
+	/**
+	 * Access the type ID portion of this element ID.
+	 */
+	FORCEINLINE FTypedHandleTypeId GetTypeId() const
+	{
+		return TypeId;
+	}
+
+	/**
+	 * Access the element ID portion of this element ID.
+	 */
+	FORCEINLINE FTypedHandleElementId GetElementId() const
+	{
+		return ElementId;
+	}
+
+	/**
+	 * Access the combined value of this element ID.
+	 * @note You typically don't want to store this directly as the element ID could be re-used.
+	 *       It is primarily useful as a secondary cache where something is keeping a reference to an element ID or element handle (eg, how UTypedElementList uses it internally).
+	 */
+	FORCEINLINE FTypedHandleCombinedId GetCombinedId() const
+	{
+		return CombinedId;
+	}
+
+	FORCEINLINE friend bool operator==(const FTypedElementId& InLHS, const FTypedElementId& InRHS)
+	{
+		return InLHS.CombinedId == InRHS.CombinedId;
+	}
+
+	FORCEINLINE friend bool operator!=(const FTypedElementId& InLHS, const FTypedElementId& InRHS)
+	{
+		return !(InLHS == InRHS);
+	}
+
+	friend inline uint32 GetTypeHash(const FTypedElementId& InElementId)
+	{
+		return GetTypeHash(InElementId.CombinedId);
+		
+	}
+
+	FORCEINLINE void Private_InitializeNoRef(const FTypedHandleTypeId InTypeId, const FTypedHandleElementId InElementId)
+	{
+		TypeId = InTypeId;
+		ElementId = InElementId;
+	}
+
+	FORCEINLINE void Private_DestroyNoRef()
+	{
+		CombinedId = 0;
+	}
+
+	/** An unset element ID */
+	static const FTypedElementId Unset;
+
+private:
+	union
+	{
+		struct
+		{
+			// Note: These are arranged in this order to give CombinedId better hash distribution for GetTypeHash!
+			FTypedHandleCombinedId ElementId : TypedHandleElementIdBits;
+			FTypedHandleCombinedId TypeId : TypedHandleTypeIdBits;
+		};
+		FTypedHandleCombinedId CombinedId;
+	};
+};
diff --git a/Engine/Source/Runtime/TypedElementFramework/Public/Elements/Framework/TypedElementRegistry.h b/Engine/Source/Runtime/TypedElementFramework/Public/Elements/Framework/TypedElementRegistry.h
--- a/Engine/Source/Runtime/TypedElementFramework/Public/Elements/Framework/TypedElementRegistry.h
+++ b/Engine/Source/Runtime/TypedElementFramework/Public/Elements/Framework/TypedElementRegistry.h
@@ -275,7 +275,7 @@ private:
 		
 		virtual FTypedElementInternalData& AddDataForElement(FTypedHandleElementId& InOutElementId) override
 		{
-			return HandleDataStore.AddDataForElement(InOutElementId);
+			return HandleDataStore.AddDataForElement(TypeId, InOutElementId);
 		}
 
 		virtual void RemoveDataForElement(const FTypedHandleElementId InElementId, const FTypedElementInternalData* InExpectedDataPtr) override
@@ -322,7 +322,7 @@ private:
 		FTypedElementInternalData& NewElementData = RegisteredElementType->AddDataForElement(NewElementId);
 
 		TTypedElementOwner<ElementDataType> ElementOwner;
-		ElementOwner.Private_InitializeAddRef(RegisteredElementType->TypeId, NewElementId, static_cast<TTypedElementInternalData<ElementDataType>&>(NewElementData));
+		ElementOwner.Private_InitializeAddRef(static_cast<TTypedElementInternalData<ElementDataType>&>(NewElementData));
 
 		return ElementOwner;
 	}
@@ -361,7 +361,7 @@ private:
 			FRegisteredElementType* RegisteredElementType = GetRegisteredElementTypeFromId(InElementId.GetTypeId());
 			checkf(RegisteredElementType, TEXT("Element type ID '%d' has not been registered!"), InElementId.GetTypeId());
 
-			OutElement.Private_InitializeAddRef(InElementId.GetTypeId(), InElementId.GetElementId(), RegisteredElementType->GetDataForElement(InElementId.GetElementId()), static_cast<BaseInterfaceType*>(RegisteredElementType->Interfaces.FindRef(InBaseInterfaceType->GetFName())));
+			OutElement.Private_InitializeAddRef(RegisteredElementType->GetDataForElement(InElementId.GetElementId()), static_cast<BaseInterfaceType*>(RegisteredElementType->Interfaces.FindRef(InBaseInterfaceType->GetFName())));
 		}
 	}
 
@@ -375,7 +375,7 @@ private:
 			FRegisteredElementType* RegisteredElementType = GetRegisteredElementTypeFromId(InElementHandle.GetId().GetTypeId());
 			checkf(RegisteredElementType, TEXT("Element type ID '%d' has not been registered!"), InElementHandle.GetId().GetTypeId());
 
-			OutElement.Private_InitializeAddRef(InElementHandle.GetId().GetTypeId(), InElementHandle.GetId().GetElementId(), *InElementHandle.Private_GetInternalData(), static_cast<BaseInterfaceType*>(RegisteredElementType->Interfaces.FindRef(InBaseInterfaceType->GetFName())));
+			OutElement.Private_InitializeAddRef(*InElementHandle.Private_GetInternalData(), static_cast<BaseInterfaceType*>(RegisteredElementType->Interfaces.FindRef(InBaseInterfaceType->GetFName())));
 		}
 	}
 
