Added support for deferred typed element destruction

This defers the actual destruction of the element until the end of the frame (or until UTypedElementRegistry::ProcessDeferredElementsToDestroy is called), or until GC runs (unless auto-GC destruction has been disabled).

This avoids firing the assert about dangling references that can be caused by transient typed elements lingering on the stack during an operation that leads to a typed element being destroyed.

#rnx

diff --git a/Engine/Source/Editor/EditorFramework/Private/Elements/Framework/TypedElementCommonActions.cpp b/Engine/Source/Editor/EditorFramework/Private/Elements/Framework/TypedElementCommonActions.cpp
--- a/Engine/Source/Editor/EditorFramework/Private/Elements/Framework/TypedElementCommonActions.cpp
+++ b/Engine/Source/Editor/EditorFramework/Private/Elements/Framework/TypedElementCommonActions.cpp
@@ -46,6 +46,8 @@ bool UTypedElementCommonActions::DeleteElements(TArrayView<const FTypedElementHa
 	bool bSuccess = false;
 
 	UTypedElementRegistry* Registry = UTypedElementRegistry::GetInstance();
+	UTypedElementRegistry::FDisableElementDestructionOnGC GCGuard(Registry);
+
 	for (const auto& ElementsByTypePair : ElementsToDuplicateByType)
 	{
 		FTypedElementCommonActionsCustomization* CommonActionsCustomization = GetInterfaceCustomizationByTypeId(ElementsByTypePair.Key);
@@ -67,6 +69,8 @@ bool UTypedElementCommonActions::DeleteElements(const UTypedElementList* Element
 	bool bSuccess = false;
 
 	UTypedElementRegistry* Registry = UTypedElementRegistry::GetInstance();
+	UTypedElementRegistry::FDisableElementDestructionOnGC GCGuard(Registry);
+
 	for (const auto& ElementsByTypePair : ElementsToDuplicateByType)
 	{
 		FTypedElementCommonActionsCustomization* CommonActionsCustomization = GetInterfaceCustomizationByTypeId(ElementsByTypePair.Key);
diff --git a/Engine/Source/Runtime/Engine/Private/Elements/Framework/EngineElementsLibrary.cpp b/Engine/Source/Runtime/Engine/Private/Elements/Framework/EngineElementsLibrary.cpp
--- a/Engine/Source/Runtime/Engine/Private/Elements/Framework/EngineElementsLibrary.cpp
+++ b/Engine/Source/Runtime/Engine/Private/Elements/Framework/EngineElementsLibrary.cpp
@@ -28,6 +28,7 @@ TTypedElementOwner<ElementDataType> CreateTypedElement(const ObjectClass* Object
 	TTypedElementOwner<ElementDataType> TypedElement;
 	if (ensureMsgf(Registry, TEXT("Typed element was requested for '%s' before the registry was available! This usually means that NewObject was used instead of CreateDefaultSubobject during CDO construction."), *Object->GetPathName()))
 	{
+		checkf(!Object->HasAnyFlags(RF_BeginDestroyed), TEXT("Typed element was requested for an object that is being destroyed!"));
 		TypedElement = Registry->CreateElement<ElementDataType>(ElementTypeName);
 	}
 
diff --git a/Engine/Source/Runtime/TypedElementFramework/Private/Elements/Framework/TypedElementRegistry.cpp b/Engine/Source/Runtime/TypedElementFramework/Private/Elements/Framework/TypedElementRegistry.cpp
--- a/Engine/Source/Runtime/TypedElementFramework/Private/Elements/Framework/TypedElementRegistry.cpp
+++ b/Engine/Source/Runtime/TypedElementFramework/Private/Elements/Framework/TypedElementRegistry.cpp
@@ -15,7 +15,9 @@ TStrongObjectPtr<UTypedElementRegistry>& GetTypedElementRegistryInstance()
 
 UTypedElementRegistry::UTypedElementRegistry()
 {
-	FCoreDelegates::OnEndFrame.AddUObject(this, &UTypedElementRegistry::NotifyElementListPendingChanges);
+	FCoreDelegates::OnBeginFrame.AddUObject(this, &UTypedElementRegistry::OnBeginFrame);
+	FCoreDelegates::OnEndFrame.AddUObject(this, &UTypedElementRegistry::OnEndFrame);
+	FCoreUObjectDelegates::GetPostGarbageCollect().AddUObject(this, &UTypedElementRegistry::OnPostGarbageCollect);
 }
 
 void UTypedElementRegistry::Private_InitializeInstance()
@@ -45,6 +47,8 @@ void UTypedElementRegistry::FinishDestroy()
 		Instance.Reset();
 	}
 
+	ProcessDeferredElementsToDestroy();
+
 	Super::FinishDestroy();
 }
 
@@ -53,6 +57,9 @@ void UTypedElementRegistry::AddReferencedObjects(UObject* InThis, FReferenceColl
 	Super::AddReferencedObjects(InThis, Collector);
 
 	UTypedElementRegistry* This = CastChecked<UTypedElementRegistry>(InThis);
+
+	FReadScopeLock RegisteredElementTypesLock(This->RegisteredElementTypesRW);
+	
 	for (TUniquePtr<FRegisteredElementType>& RegisteredElementType : This->RegisteredElementTypes)
 	{
 		if (RegisteredElementType)
@@ -112,6 +119,19 @@ UTypedElementInterface* UTypedElementRegistry::GetElementInterfaceImpl(const FTy
 	return RegisteredElementType->Interfaces.FindRef(InBaseInterfaceType->GetFName());
 }
 
+void UTypedElementRegistry::ProcessDeferredElementsToDestroy()
+{
+	FReadScopeLock RegisteredElementTypesLock(RegisteredElementTypesRW);
+
+	for (TUniquePtr<FRegisteredElementType>& RegisteredElementType : RegisteredElementTypes)
+	{
+		if (RegisteredElementType)
+		{
+			RegisteredElementType->ProcessDeferredElementsToRemove();
+		}
+	}
+}
+
 void UTypedElementRegistry::ReleaseElementId(FTypedElementId& InOutElementId)
 {
 	if (!InOutElementId)
@@ -175,3 +195,26 @@ void UTypedElementRegistry::NotifyElementListPendingChanges()
 		ActiveElementList->NotifyPendingChanges();
 	}
 }
+
+void UTypedElementRegistry::OnBeginFrame()
+{
+	// Prevent auto-GC reference collection during this frame
+	IncrementDisableElementDestructionOnGCCount();
+}
+
+void UTypedElementRegistry::OnEndFrame()
+{
+	NotifyElementListPendingChanges();
+	ProcessDeferredElementsToDestroy();
+
+	// Allow auto-GC reference collection until the start of the next frame
+	DecrementDisableElementDestructionOnGCCount();
+}
+
+void UTypedElementRegistry::OnPostGarbageCollect()
+{
+	if (DisableElementDestructionOnGCCount == 0)
+	{
+		ProcessDeferredElementsToDestroy();
+	}
+}
diff --git a/Engine/Source/Runtime/TypedElementFramework/Public/Elements/Framework/TypedElementData.h b/Engine/Source/Runtime/TypedElementFramework/Public/Elements/Framework/TypedElementData.h
--- a/Engine/Source/Runtime/TypedElementFramework/Public/Elements/Framework/TypedElementData.h
+++ b/Engine/Source/Runtime/TypedElementFramework/Public/Elements/Framework/TypedElementData.h
@@ -68,10 +68,10 @@ class FTypedElementInternalData
 {
 public:
 	FTypedElementInternalData() = default;
-	
+
 	FTypedElementInternalData(const FTypedElementInternalData&) = delete;
 	FTypedElementInternalData& operator=(const FTypedElementInternalData&) = delete;
-	
+
 	FTypedElementInternalData(FTypedElementInternalData&& InOther) = delete;
 	FTypedElementInternalData& operator=(FTypedElementInternalData&&) = delete;
 
@@ -96,6 +96,7 @@ public:
 		{
 			FScopeLock ReferencesLock(&ReferencesCS);
 			References.Reset();
+			DestructionRequestCallstack.Reset();
 		}
 #endif	// UE_TYPED_ELEMENT_HAS_REFTRACKING
 	}
@@ -151,6 +152,7 @@ public:
 #if UE_TYPED_ELEMENT_HAS_REFTRACKING
 		{
 			FScopeLock ReferencesLock(&ReferencesCS);
+			UE_LOG(LogCore, Error, TEXT("==============================================="));
 			UE_LOG(LogCore, Error, TEXT("External Element References:"));
 			for (const FTypedElementReference& Reference : References)
 			{
@@ -158,12 +160,48 @@ public:
 				Reference.LogReference();
 			}
 			UE_LOG(LogCore, Error, TEXT("==============================================="));
+			if (DestructionRequestCallstack)
+			{
+				UE_LOG(LogCore, Error, TEXT("Destruction requested by:"));
+				DestructionRequestCallstack->LogReference();
+				UE_LOG(LogCore, Error, TEXT("==============================================="));
+			}
 		}
 #else	// UE_TYPED_ELEMENT_HAS_REFTRACKING
 		UE_LOG(LogCore, Error, TEXT("UE_TYPED_ELEMENT_HAS_REFTRACKING is disabled. Enable it and recompile to see reference tracking."));
 #endif	// UE_TYPED_ELEMENT_HAS_REFTRACKING
 	}
 
+	void StoreDestructionRequestCallstack() const
+	{
+#if UE_TYPED_ELEMENT_HAS_REFTRACKING
+		FScopeLock ReferencesLock(&ReferencesCS);
+#if DO_CHECK
+		if (DestructionRequestCallstack)
+		{
+			UE_LOG(LogCore, Error, TEXT("==============================================="));
+			UE_LOG(LogCore, Error, TEXT("Destruction requested by:"));
+			DestructionRequestCallstack->LogReference();
+			UE_LOG(LogCore, Error, TEXT("==============================================="));
+			UE_LOG(LogCore, Fatal, TEXT("Element has already had its destruction callstack set! (see above)"));
+		}
+#endif	// DO_CHECK
+		DestructionRequestCallstack = MakeUnique<FTypedElementReference>();
+#endif	// UE_TYPED_ELEMENT_HAS_REFTRACKING
+	}
+
+	void CheckNoExternalReferencesOnDestruction() const
+	{
+#if DO_CHECK
+		const FTypedElementRefCount LocalRefCount = GetRefCount();
+		if (LocalRefCount > 1)
+		{
+			LogReferences();
+			UE_LOG(LogCore, Fatal, TEXT("Element is still externally referenced when being destroyed! Ref-count: %d; see above for reference information (if available)."), LocalRefCount);
+		}
+#endif	// DO_CHECK
+	}
+
 	virtual const void* GetUntypedData() const
 	{
 		return nullptr;
@@ -177,6 +215,7 @@ private:
 #if UE_TYPED_ELEMENT_HAS_REFTRACKING
 	mutable FCriticalSection ReferencesCS;
 	mutable TSparseArray<FTypedElementReference> References;
+	mutable TUniquePtr<FTypedElementReference> DestructionRequestCallstack;
 #endif	// UE_TYPED_ELEMENT_HAS_REFTRACKING
 };
 
@@ -253,6 +292,7 @@ public:
 		
 		TTypedElementInternalData<ElementDataType>& InternalData = InternalDataArray[InElementId];
 		checkf(InExpectedDataPtr == &InternalData, TEXT("Internal data pointer did not match the expected value! Does this handle belong to a different element registry?"));
+		InternalData.CheckNoExternalReferencesOnDestruction();
 		InternalData.Reset();
 		InternalDataFreeIndices.Add(InElementId);
 	}
@@ -326,6 +366,7 @@ public:
 
 		TTypedElementInternalData<void>& InternalData = InternalDataArray[InElementId];
 		checkf(InExpectedDataPtr == &InternalData, TEXT("Internal data pointer did not match the expected value! Does this handle belong to a different element registry?"));
+		InternalData.CheckNoExternalReferencesOnDestruction();
 		InternalData.Reset();
 		InternalDataFreeIndices.Add(InternalDataArrayIndex);
 	}
diff --git a/Engine/Source/Runtime/TypedElementFramework/Public/Elements/Framework/TypedElementRegistry.h b/Engine/Source/Runtime/TypedElementFramework/Public/Elements/Framework/TypedElementRegistry.h
--- a/Engine/Source/Runtime/TypedElementFramework/Public/Elements/Framework/TypedElementRegistry.h
+++ b/Engine/Source/Runtime/TypedElementFramework/Public/Elements/Framework/TypedElementRegistry.h
@@ -127,6 +127,7 @@ public:
 
 	/**
 	 * Destroy an element.
+	 * @note Destruction is deferred until the next call to ProcessDeferredElementsToDestroy.
 	 */
 	FORCEINLINE void DestroyElement(FTypedElementOwner& InOutElementOwner)
 	{
@@ -135,6 +136,7 @@ public:
 
 	/**
 	 * Destroy an element.
+	 * @note Destruction is deferred until the next call to ProcessDeferredElementsToDestroy.
 	 */
 	template <typename ElementDataType>
 	FORCEINLINE void DestroyElement(TTypedElementOwner<ElementDataType>& InOutElementOwner)
@@ -142,6 +144,13 @@ public:
 		return DestroyElementImpl<ElementDataType>(InOutElementOwner);
 	}
 
+	/**
+	 * Process any elements that were previously marked for deferred destruction.
+	 * @note This is automatically called at the end of each frame, but may also be called manually.
+	 * @note This is automatically called post-GC, unless auto-GC destruction has been disabled (@see FDisableElementDestructionOnGC).
+	 */
+	void ProcessDeferredElementsToDestroy();
+
 	/**
 	 * Release an element ID that was previously acquired from an existing handle.
 	 */
@@ -251,14 +260,41 @@ public:
 		GetElementImpl(InElementHandle, InBaseInterfaceType, OutElement);
 	}
 
+	/** Struct to disable auto-GC reference collection within a scope */
+	struct FDisableElementDestructionOnGC
+	{
+	public:
+		explicit FDisableElementDestructionOnGC(UTypedElementRegistry* InRegistry)
+			: Registry(InRegistry)
+		{
+			checkf(Registry, TEXT("FDisableElementDestructionOnGC must be used with a valid registry!"));
+			Registry->IncrementDisableElementDestructionOnGCCount();
+		}
+
+		~FDisableElementDestructionOnGC()
+		{
+			Registry->DecrementDisableElementDestructionOnGCCount();
+		}
+
+		FDisableElementDestructionOnGC(const FDisableElementDestructionOnGC&) = delete;
+		FDisableElementDestructionOnGC& operator=(const FDisableElementDestructionOnGC&) = delete;
+
+		FDisableElementDestructionOnGC(FDisableElementDestructionOnGC&&) = delete;
+		FDisableElementDestructionOnGC& operator=(FDisableElementDestructionOnGC&&) = delete;
+
+	private:
+		UTypedElementRegistry* Registry = nullptr;
+	};
+
 private:
 	struct FRegisteredElementType
 	{
 		virtual ~FRegisteredElementType() = default;
 
 		virtual FTypedElementInternalData& AddDataForElement(FTypedHandleElementId& InOutElementId) = 0;
-		virtual void RemoveDataForElement(const FTypedHandleElementId InElementId, const FTypedElementInternalData* InExpectedDataPtr) = 0;
+		virtual void RemoveDataForElement(const FTypedHandleElementId InElementId, const FTypedElementInternalData* InExpectedDataPtr, const bool bDefer) = 0;
 		virtual const FTypedElementInternalData& GetDataForElement(const FTypedHandleElementId InElementId) const = 0;
+		virtual void ProcessDeferredElementsToRemove() = 0;
 		virtual void SetDataTypeId(const FTypedHandleTypeId InTypeId) = 0;
 		virtual FTypedHandleTypeId GetDataTypeId() const = 0;
 		virtual FName GetDataTypeName() const = 0;
@@ -278,9 +314,19 @@ private:
 			return HandleDataStore.AddDataForElement(TypeId, InOutElementId);
 		}
 
-		virtual void RemoveDataForElement(const FTypedHandleElementId InElementId, const FTypedElementInternalData* InExpectedDataPtr) override
+		virtual void RemoveDataForElement(const FTypedHandleElementId InElementId, const FTypedElementInternalData* InExpectedDataPtr, const bool bDefer) override
 		{
-			HandleDataStore.RemoveDataForElement(InElementId, InExpectedDataPtr);
+			if (bDefer)
+			{
+				InExpectedDataPtr->StoreDestructionRequestCallstack();
+
+				FScopeLock DeferredElementsToRemoveLock(&DeferredElementsToRemoveCS);
+				DeferredElementsToRemove.Add(MakeTuple(InElementId, InExpectedDataPtr));
+			}
+			else
+			{
+				HandleDataStore.RemoveDataForElement(InElementId, InExpectedDataPtr);
+			}
 		}
 
 		virtual const FTypedElementInternalData& GetDataForElement(const FTypedHandleElementId InElementId) const override
@@ -288,6 +334,16 @@ private:
 			return HandleDataStore.GetDataForElement(InElementId);
 		}
 
+		virtual void ProcessDeferredElementsToRemove() override
+		{
+			FScopeLock DeferredElementsToRemoveLock(&DeferredElementsToRemoveCS);
+			for (const FDeferredElementToRemove& DeferredElementToRemove : DeferredElementsToRemove)
+			{
+				HandleDataStore.RemoveDataForElement(DeferredElementToRemove.Get<0>(), DeferredElementToRemove.Get<1>());
+			}
+			DeferredElementsToRemove.Reset();
+		}
+
 		virtual void SetDataTypeId(const FTypedHandleTypeId InTypeId) override
 		{
 			TTypedElementInternalDataStore<ElementDataType>::SetStaticDataTypeId(InTypeId);
@@ -304,6 +360,10 @@ private:
 		}
 
 		TTypedElementInternalDataStore<ElementDataType> HandleDataStore;
+
+		using FDeferredElementToRemove = TTuple<FTypedHandleElementId, const FTypedElementInternalData*>;
+		FCriticalSection DeferredElementsToRemoveCS;
+		TArray<FDeferredElementToRemove> DeferredElementsToRemove;
 	};
 
 	void RegisterElementTypeImpl(const FName InElementTypeName, TUniquePtr<FRegisteredElementType>&& InRegisteredElementType);
@@ -333,21 +393,7 @@ private:
 		FRegisteredElementType* RegisteredElementType = GetRegisteredElementTypeFromId(InOutElementOwner.GetId().GetTypeId());
 		checkf(RegisteredElementType, TEXT("Element type ID '%d' has not been registered!"), InOutElementOwner.GetId().GetTypeId());
 
-#if DO_CHECK
-		{
-			const FTypedElementInternalData& ElementData = RegisteredElementType->GetDataForElement(InOutElementOwner.GetId().GetElementId());
-			const FTypedElementRefCount RefCount = ElementData.GetRefCount();
-
-			const bool bIsExternallyReferenced = RefCount > 1;
-			if (bIsExternallyReferenced)
-			{
-				ElementData.LogReferences();
-				UE_LOG(LogCore, Fatal, TEXT("Element is still externally referenced when being destroyed! Ref-count: %d; see above for reference information (if available)."), RefCount);
-			}
-		}
-#endif	// DO_CHECK
-
-		RegisteredElementType->RemoveDataForElement(InOutElementOwner.GetId().GetElementId(), InOutElementOwner.Private_GetInternalData());
+		RegisteredElementType->RemoveDataForElement(InOutElementOwner.GetId().GetElementId(), InOutElementOwner.Private_GetInternalData(), /*bDefer*/true);
 		InOutElementOwner.Private_DestroyNoRef();
 	}
 
@@ -415,11 +461,29 @@ private:
 	}
 
 	void NotifyElementListPendingChanges();
-	
+
+	void OnBeginFrame();
+	void OnEndFrame();
+	void OnPostGarbageCollect();
+
+	FORCEINLINE void IncrementDisableElementDestructionOnGCCount()
+	{
+		checkf(DisableElementDestructionOnGCCount != MAX_uint8, TEXT("DisableElementDestructionOnGCCount overflow!"));
+		++DisableElementDestructionOnGCCount;
+	}
+
+	FORCEINLINE void DecrementDisableElementDestructionOnGCCount()
+	{
+		checkf(DisableElementDestructionOnGCCount != 0, TEXT("DisableElementDestructionOnGCCount underflow!"));
+		--DisableElementDestructionOnGCCount;
+	}
+
 	mutable FRWLock RegisteredElementTypesRW;
 	TUniquePtr<FRegisteredElementType> RegisteredElementTypes[TypedHandleMaxTypeId - 1];
 	TSortedMap<FName, FTypedHandleTypeId, FDefaultAllocator, FNameFastLess> RegisteredElementTypesNameToId;
 
 	mutable FRWLock ActiveElementListsRW;
 	TSet<UTypedElementList*> ActiveElementLists;
+
+	uint8 DisableElementDestructionOnGCCount = 0;
 };
