Fixed assert if a UTypedElementRegistry was created mid-frame

#rnx

diff --git a/Engine/Source/Runtime/TypedElementFramework/Private/Elements/Framework/TypedElementRegistry.cpp b/Engine/Source/Runtime/TypedElementFramework/Private/Elements/Framework/TypedElementRegistry.cpp
--- a/Engine/Source/Runtime/TypedElementFramework/Private/Elements/Framework/TypedElementRegistry.cpp
+++ b/Engine/Source/Runtime/TypedElementFramework/Private/Elements/Framework/TypedElementRegistry.cpp
@@ -15,9 +15,12 @@ TStrongObjectPtr<UTypedElementRegistry>& GetTypedElementRegistryInstance()
 
 UTypedElementRegistry::UTypedElementRegistry()
 {
-	FCoreDelegates::OnBeginFrame.AddUObject(this, &UTypedElementRegistry::OnBeginFrame);
-	FCoreDelegates::OnEndFrame.AddUObject(this, &UTypedElementRegistry::OnEndFrame);
-	FCoreUObjectDelegates::GetPostGarbageCollect().AddUObject(this, &UTypedElementRegistry::OnPostGarbageCollect);
+	if (!HasAnyFlags(RF_ClassDefaultObject))
+	{
+		FCoreDelegates::OnBeginFrame.AddUObject(this, &UTypedElementRegistry::OnBeginFrame);
+		FCoreDelegates::OnEndFrame.AddUObject(this, &UTypedElementRegistry::OnEndFrame);
+		FCoreUObjectDelegates::GetPostGarbageCollect().AddUObject(this, &UTypedElementRegistry::OnPostGarbageCollect);
+	}
 }
 
 void UTypedElementRegistry::Private_InitializeInstance()
@@ -200,6 +203,7 @@ void UTypedElementRegistry::OnBeginFrame()
 {
 	// Prevent auto-GC reference collection during this frame
 	IncrementDisableElementDestructionOnGCCount();
+	bIsWithinFrame = true;
 }
 
 void UTypedElementRegistry::OnEndFrame()
@@ -207,8 +211,12 @@ void UTypedElementRegistry::OnEndFrame()
 	NotifyElementListPendingChanges();
 	ProcessDeferredElementsToDestroy();
 
-	// Allow auto-GC reference collection until the start of the next frame
-	DecrementDisableElementDestructionOnGCCount();
+	if (bIsWithinFrame)
+	{
+		// Allow auto-GC reference collection until the start of the next frame
+		DecrementDisableElementDestructionOnGCCount();
+		bIsWithinFrame = false;
+	}
 }
 
 void UTypedElementRegistry::OnPostGarbageCollect()
diff --git a/Engine/Source/Runtime/TypedElementFramework/Public/Elements/Framework/TypedElementRegistry.h b/Engine/Source/Runtime/TypedElementFramework/Public/Elements/Framework/TypedElementRegistry.h
--- a/Engine/Source/Runtime/TypedElementFramework/Public/Elements/Framework/TypedElementRegistry.h
+++ b/Engine/Source/Runtime/TypedElementFramework/Public/Elements/Framework/TypedElementRegistry.h
@@ -486,4 +486,5 @@ private:
 	TSet<UTypedElementList*> ActiveElementLists;
 
 	uint8 DisableElementDestructionOnGCCount = 0;
+	bool bIsWithinFrame = false;
 };
