Added runtime options for controlling typed element features

UE_TYPED_ELEMENT_HAS_REFTRACKING is now compiled in to Debug and Development builds, and is controlled at runtime by the "TypedElements.EnableReferenceTracking" CVar. This can be queried in code via FTypedElementReferences::ReferenceTrackingEnabled.

UE_ENABLE_SMINSTANCE_ELEMENTS has been removed, and is controlled at runtime by the "TypedElements.EnableSMInstanceElements" CVar. This can be queried in code via SMInstanceElementDataUtil::SMInstanceElementsEnabled.

diff --git a/Engine/Source/Editor/AssetPlacementEdMode/Private/Tools/PlacementBrushToolBase.cpp b/Engine/Source/Editor/AssetPlacementEdMode/Private/Tools/PlacementBrushToolBase.cpp
--- a/Engine/Source/Editor/AssetPlacementEdMode/Private/Tools/PlacementBrushToolBase.cpp
+++ b/Engine/Source/Editor/AssetPlacementEdMode/Private/Tools/PlacementBrushToolBase.cpp
@@ -19,6 +19,7 @@
 #include "Editor.h"
 #include "Modes/PlacementModeSubsystem.h"
 #include "ActorFactories/ActorFactory.h"
+#include "Elements/SMInstance/SMInstanceElementData.h" // For SMInstanceElementDataUtil::SMInstanceElementsEnabled
 
 bool UPlacementToolBuilderBase::CanBuildTool(const FToolBuilderState& SceneState) const
 {	
@@ -280,33 +281,34 @@ TArray<FTypedElementHandle> UPlacementBrushToolBase::GetElementsInBrushRadius()
 		}
 	}
 
-#if !UE_ENABLE_SMINSTANCE_ELEMENTS
-	// Handle the IFA for the brush stroke level
-	if (UActorPartitionSubsystem* PartitionSubsystem = UWorld::GetSubsystem<UActorPartitionSubsystem>(GEditor->GetEditorWorldContext().World()))
+	if (!SMInstanceElementDataUtil::SMInstanceElementsEnabled())
 	{
-		constexpr bool bCreatePartitionActorIfMissing = false;
-		FActorPartitionGetParams PartitionActorFindParams(AInstancedFoliageActor::StaticClass(), bCreatePartitionActorIfMissing, GEditor->GetEditorWorldContext().World()->GetCurrentLevel(), LastBrushStamp.WorldPosition);
-		if (AInstancedFoliageActor* FoliageActor = Cast<AInstancedFoliageActor>(PartitionSubsystem->GetActor(PartitionActorFindParams)))
+		// Handle the IFA for the brush stroke level
+		if (UActorPartitionSubsystem* PartitionSubsystem = UWorld::GetSubsystem<UActorPartitionSubsystem>(GEditor->GetEditorWorldContext().World()))
 		{
-			for (const auto& FoliageInfo : FoliageActor->GetFoliageInfos())
+			constexpr bool bCreatePartitionActorIfMissing = false;
+			FActorPartitionGetParams PartitionActorFindParams(AInstancedFoliageActor::StaticClass(), bCreatePartitionActorIfMissing, GEditor->GetEditorWorldContext().World()->GetCurrentLevel(), LastBrushStamp.WorldPosition);
+			if (AInstancedFoliageActor* FoliageActor = Cast<AInstancedFoliageActor>(PartitionSubsystem->GetActor(PartitionActorFindParams)))
 			{
-				FTypedElementHandle SourceObjectHandle = UEngineElementsLibrary::AcquireEditorObjectElementHandle(FoliageInfo.Key->GetSource());
-				if (GEditor->GetEditorSubsystem<UPlacementModeSubsystem>()->DoesCurrentPaletteSupportElement(SourceObjectHandle))
+				for (const auto& FoliageInfo : FoliageActor->GetFoliageInfos())
 				{
-					TArray<int32> Instances;
-					FSphere SphereToCheck(LastBrushStamp.WorldPosition, LastBrushStamp.Radius);
-					FoliageInfo.Value->GetInstancesInsideSphere(SphereToCheck, Instances);
-					if (Instances.Num())
+					FTypedElementHandle SourceObjectHandle = UEngineElementsLibrary::AcquireEditorObjectElementHandle(FoliageInfo.Key->GetSource());
+					if (GEditor->GetEditorSubsystem<UPlacementModeSubsystem>()->DoesCurrentPaletteSupportElement(SourceObjectHandle))
 					{
-						// For now, return the whole foliage actor, and allow the calling code to drill down, since we do not have element handles at the instance level just yet
-						ElementHandles.Emplace(UEngineElementsLibrary::AcquireEditorActorElementHandle(FoliageActor));
-						break;
+						TArray<int32> Instances;
+						FSphere SphereToCheck(LastBrushStamp.WorldPosition, LastBrushStamp.Radius);
+						FoliageInfo.Value->GetInstancesInsideSphere(SphereToCheck, Instances);
+						if (Instances.Num())
+						{
+							// For now, return the whole foliage actor, and allow the calling code to drill down, since we do not have element handles at the instance level just yet
+							ElementHandles.Emplace(UEngineElementsLibrary::AcquireEditorActorElementHandle(FoliageActor));
+							break;
+						}
 					}
 				}
 			}
 		}
 	}
-#endif
 
 	return ElementHandles;
 }
diff --git a/Engine/Source/Editor/AssetPlacementEdMode/Private/Tools/PlacementPlaceSingleTool.cpp b/Engine/Source/Editor/AssetPlacementEdMode/Private/Tools/PlacementPlaceSingleTool.cpp
--- a/Engine/Source/Editor/AssetPlacementEdMode/Private/Tools/PlacementPlaceSingleTool.cpp
+++ b/Engine/Source/Editor/AssetPlacementEdMode/Private/Tools/PlacementPlaceSingleTool.cpp
@@ -12,6 +12,7 @@
 #include "Elements/Framework/TypedElementRegistry.h"
 #include "Elements/Framework/TypedElementSelectionSet.h"
 #include "Elements/Interfaces/TypedElementWorldInterface.h"
+#include "Elements/SMInstance/SMInstanceElementData.h" // For SMInstanceElementDataUtil::SMInstanceElementsEnabled
 #include "Modes/PlacementModeSubsystem.h"
 #include "Subsystems/PlacementSubsystem.h"
 #include "UObject/Object.h"
@@ -62,11 +63,7 @@ void UPlacementModePlaceSingleTool::OnClickPress(const FInputDeviceRay& PressPos
 	// Place the Preview data if we managed to get to a valid handled click.
 	FPlacementOptions PlacementOptions;
 	PlacementOptions.bPreferBatchPlacement = true;
-#if UE_ENABLE_SMINSTANCE_ELEMENTS
-	PlacementOptions.bPreferInstancedPlacement = true;
-#else
-	PlacementOptions.bPreferInstancedPlacement = false;
-#endif
+	PlacementOptions.bPreferInstancedPlacement = SMInstanceElementDataUtil::SMInstanceElementsEnabled();
 	
 	UPlacementSubsystem* PlacementSubsystem = GEditor->GetEditorSubsystem<UPlacementSubsystem>();
 	if (PlacementSubsystem)
diff --git a/Engine/Source/Editor/UnrealEd/Private/Factories/EditorStaticMeshFactory.cpp b/Engine/Source/Editor/UnrealEd/Private/Factories/EditorStaticMeshFactory.cpp
--- a/Engine/Source/Editor/UnrealEd/Private/Factories/EditorStaticMeshFactory.cpp
+++ b/Engine/Source/Editor/UnrealEd/Private/Factories/EditorStaticMeshFactory.cpp
@@ -63,11 +63,13 @@ TArray<FTypedElementHandle> UEditorStaticMeshFactory::PlaceAsset(const FAssetPla
 			FoliageInfo->AddInstance(FoundOrCreatedType, FoliagePlacementInfo);
 			CurrentPlacementScopedFoliageInfos.Add(FoliageInfo);
 
-#if UE_ENABLE_SMINSTANCE_ELEMENTS
-			return TArray<FTypedElementHandle>({ UEngineElementsLibrary::AcquireEditorSMInstanceElementHandle(FoliageInfo->GetComponent(), FoliageInfo->GetPlacedInstanceCount() - 1) });
-#else
-			return TArray<FTypedElementHandle>({ UEngineElementsLibrary::AcquireEditorComponentElementHandle(FoliageInfo->GetComponent()) });
-#endif
+			FTypedElementHandle PlacedElement = UEngineElementsLibrary::AcquireEditorSMInstanceElementHandle(FoliageInfo->GetComponent(), FoliageInfo->GetPlacedInstanceCount() - 1);
+			if (!PlacedElement)
+			{
+				PlacedElement = UEngineElementsLibrary::AcquireEditorComponentElementHandle(FoliageInfo->GetComponent());
+			}
+			check(PlacedElement);
+			return TArray<FTypedElementHandle>({ MoveTemp(PlacedElement) });
 		}
 	}
 
diff --git a/Engine/Source/Runtime/Engine/Engine.Build.cs b/Engine/Source/Runtime/Engine/Engine.Build.cs
--- a/Engine/Source/Runtime/Engine/Engine.Build.cs
+++ b/Engine/Source/Runtime/Engine/Engine.Build.cs
@@ -273,8 +273,6 @@ public class Engine : ModuleRules
 			}
 		}
 
-		PublicDefinitions.Add("UE_ENABLE_SMINSTANCE_ELEMENTS=0");
-
 		DynamicallyLoadedModuleNames.AddRange(
 			new string[] {
 				"NullNetworkReplayStreaming",
diff --git a/Engine/Source/Runtime/Engine/Private/Elements/Framework/EngineElementsLibrary.cpp b/Engine/Source/Runtime/Engine/Private/Elements/Framework/EngineElementsLibrary.cpp
--- a/Engine/Source/Runtime/Engine/Private/Elements/Framework/EngineElementsLibrary.cpp
+++ b/Engine/Source/Runtime/Engine/Private/Elements/Framework/EngineElementsLibrary.cpp
@@ -386,15 +386,13 @@ TTypedElementOwner<FSMInstanceElementData> UEngineElementsLibrary::CreateSMInsta
 	UTypedElementRegistry* Registry = UTypedElementRegistry::GetInstance();
 
 	TTypedElementOwner<FSMInstanceElementData> SMInstanceElement;
-	if (ensureAlways(Registry))
+	if (ensureAlways(Registry) && SMInstanceElementDataUtil::SMInstanceElementsEnabled())
 	{
-#if UE_ENABLE_SMINSTANCE_ELEMENTS
 		SMInstanceElement = Registry->CreateElement<FSMInstanceElementData>(NAME_SMInstance);
 		if (SMInstanceElement)
 		{
 			SMInstanceElement.GetDataChecked().InstanceElementId = InSMInstanceElementId;
 		}
-#endif	// UE_ENABLE_SMINSTANCE_ELEMENTS
 	}
 
 	return SMInstanceElement;
diff --git a/Engine/Source/Runtime/Engine/Private/Elements/SMInstance/SMInstanceElementData.cpp b/Engine/Source/Runtime/Engine/Private/Elements/SMInstance/SMInstanceElementData.cpp
--- a/Engine/Source/Runtime/Engine/Private/Elements/SMInstance/SMInstanceElementData.cpp
+++ b/Engine/Source/Runtime/Engine/Private/Elements/SMInstance/SMInstanceElementData.cpp
@@ -2,12 +2,25 @@
 
 #include "Elements/SMInstance/SMInstanceElementData.h"
 #include "Elements/Framework/TypedElementHandle.h"
+#include "HAL/IConsoleManager.h"
 
 UE_DEFINE_TYPED_ELEMENT_DATA_RTTI(FSMInstanceElementData);
 
 namespace SMInstanceElementDataUtil
 {
 
+static int32 GEnableSMInstanceElements = 0;
+static FAutoConsoleVariableRef CVarEnableSMInstanceElements(
+	TEXT("TypedElements.EnableSMInstanceElements"),
+	GEnableSMInstanceElements,
+	TEXT("Is support for static mesh instance elements enabled?")
+	);
+
+bool SMInstanceElementsEnabled()
+{
+	return GEnableSMInstanceElements != 0;
+}
+
 FSMInstanceId GetSMInstanceFromHandle(const FTypedElementHandle& InHandle, const bool bSilent)
 {
 	const FSMInstanceElementData* SMInstanceElement = InHandle.GetData<FSMInstanceElementData>(bSilent);
diff --git a/Engine/Source/Runtime/Engine/Public/Elements/SMInstance/SMInstanceElementData.h b/Engine/Source/Runtime/Engine/Public/Elements/SMInstance/SMInstanceElementData.h
--- a/Engine/Source/Runtime/Engine/Public/Elements/SMInstance/SMInstanceElementData.h
+++ b/Engine/Source/Runtime/Engine/Public/Elements/SMInstance/SMInstanceElementData.h
@@ -20,6 +20,12 @@ struct ENGINE_API FSMInstanceElementData
 namespace SMInstanceElementDataUtil
 {
 
+/**
+ * Test whether static mesh instance elements are currently enabled?
+ * @note Controlled by the CVar: "TypedElements.EnableSMInstanceElements".
+ */
+ENGINE_API bool SMInstanceElementsEnabled();
+
 /**
  * Attempt to get the static mesh instance ID from the given element handle.
  * @note This is not typically something you'd want to use outside of data access within an interface implementation.
diff --git a/Engine/Source/Runtime/TypedElementFramework/Private/Elements/Framework/TypedElementRegistry.cpp b/Engine/Source/Runtime/TypedElementFramework/Private/Elements/Framework/TypedElementRegistry.cpp
--- a/Engine/Source/Runtime/TypedElementFramework/Private/Elements/Framework/TypedElementRegistry.cpp
+++ b/Engine/Source/Runtime/TypedElementFramework/Private/Elements/Framework/TypedElementRegistry.cpp
@@ -3,10 +3,28 @@
 #include "Elements/Framework/TypedElementRegistry.h"
 #include "Misc/CoreDelegates.h"
 #include "Misc/ScopeLock.h"
+#include "HAL/IConsoleManager.h"
 #include "UObject/StrongObjectPtr.h"
 
 const FTypedElementId FTypedElementId::Unset;
 
+#if UE_TYPED_ELEMENT_HAS_REFTRACKING
+namespace TypedElementReferences
+{
+static int32 GEnableReferenceTracking = 0;
+static FAutoConsoleVariableRef CVarEnableReferenceTracking(
+	TEXT("TypedElements.EnableReferenceTracking"),
+	GEnableReferenceTracking,
+	TEXT("Is support for element reference tracking enabled?")
+	);
+} // namespace TypedElementReference
+
+bool FTypedElementReferences::ReferenceTrackingEnabled()
+{
+	return TypedElementReferences::GEnableReferenceTracking != 0;
+}
+#endif	// UE_TYPED_ELEMENT_HAS_REFTRACKING
+
 TStrongObjectPtr<UTypedElementRegistry>& GetTypedElementRegistryInstance()
 {
 	static TStrongObjectPtr<UTypedElementRegistry> TypedElementRegistryInstance;
diff --git a/Engine/Source/Runtime/TypedElementFramework/Public/Elements/Framework/TypedElementData.h b/Engine/Source/Runtime/TypedElementFramework/Public/Elements/Framework/TypedElementData.h
--- a/Engine/Source/Runtime/TypedElementFramework/Public/Elements/Framework/TypedElementData.h
+++ b/Engine/Source/Runtime/TypedElementFramework/Public/Elements/Framework/TypedElementData.h
@@ -59,6 +59,84 @@ private:
 	uint64 Callstack[TypedHandleRefTrackingDepth + TypedHandleRefTrackingSkipCount];
 	uint32 CallstackDepth = 0;
 };
+
+/**
+ * Debugging information used to locate reference leaks.
+ */
+class FTypedElementReferences
+{
+public:
+	static TUniquePtr<FTypedElementReferences> Create()
+	{
+		return ReferenceTrackingEnabled()
+			? MakeUnique<FTypedElementReferences>()
+			: nullptr;
+	}
+
+	void Reset()
+	{
+		FScopeLock ReferencesLock(&ReferencesCS);
+		References.Reset();
+		DestructionRequestCallstack.Reset();
+	}
+
+	FTypedElementReferenceId AddRef()
+	{
+		FScopeLock ReferencesLock(&ReferencesCS);
+		return References.Add(FTypedElementReference());
+	}
+
+	void ReleaseRef(const FTypedElementReferenceId InReferenceId)
+	{
+		if (InReferenceId != INDEX_NONE)
+		{
+			FScopeLock ReferencesLock(&ReferencesCS);
+			References.RemoveAt(InReferenceId);
+		}
+	}
+
+	void LogReferences() const
+	{
+		FScopeLock ReferencesLock(&ReferencesCS);
+		UE_LOG(LogCore, Error, TEXT("==============================================="));
+		UE_LOG(LogCore, Error, TEXT("External Element References:"));
+		for (const FTypedElementReference& Reference : References)
+		{
+			UE_LOG(LogCore, Error, TEXT("-----------------------------------------------"));
+			Reference.LogReference();
+		}
+		UE_LOG(LogCore, Error, TEXT("==============================================="));
+		if (DestructionRequestCallstack)
+		{
+			UE_LOG(LogCore, Error, TEXT("Destruction requested by:"));
+			DestructionRequestCallstack->LogReference();
+			UE_LOG(LogCore, Error, TEXT("==============================================="));
+		}
+	}
+
+	void StoreDestructionRequestCallstack()
+	{
+		FScopeLock ReferencesLock(&ReferencesCS);
+#if DO_CHECK
+		if (DestructionRequestCallstack)
+		{
+			UE_LOG(LogCore, Error, TEXT("==============================================="));
+			UE_LOG(LogCore, Error, TEXT("Destruction requested by:"));
+			DestructionRequestCallstack->LogReference();
+			UE_LOG(LogCore, Error, TEXT("==============================================="));
+			UE_LOG(LogCore, Fatal, TEXT("Element has already had its destruction callstack set! (see above)"));
+		}
+#endif	// DO_CHECK
+		DestructionRequestCallstack = MakeUnique<FTypedElementReference>();
+	}
+
+private:
+	TYPEDELEMENTFRAMEWORK_API static bool ReferenceTrackingEnabled();
+
+	mutable FCriticalSection ReferencesCS;
+	TSparseArray<FTypedElementReference> References;
+	TUniquePtr<FTypedElementReference> DestructionRequestCallstack;
+};
 #endif	// UE_TYPED_ELEMENT_HAS_REFTRACKING
 
 /**
@@ -84,6 +162,14 @@ public:
 	{
 		checkSlow(!Id.IsSet());
 		Id.Private_InitializeNoRef(InTypeId, InElementId);
+#if UE_TYPED_ELEMENT_HAS_REFTRACKING
+		if (!References)
+		{
+			// Do this in Initialize rather than the constructor, as the CVar value 
+			// may change while the already constructed instances are re-used
+			References = FTypedElementReferences::Create();
+		}
+#endif	// UE_TYPED_ELEMENT_HAS_REFTRACKING
 	}
 
 	virtual void Reset()
@@ -93,10 +179,9 @@ public:
 		FPlatformAtomics::InterlockedExchange(&RefCount, 0);
 #endif	// UE_TYPED_ELEMENT_HAS_REFCOUNTING
 #if UE_TYPED_ELEMENT_HAS_REFTRACKING
+		if (References)
 		{
-			FScopeLock ReferencesLock(&ReferencesCS);
-			References.Reset();
-			DestructionRequestCallstack.Reset();
+			References->Reset();
 		}
 #endif	// UE_TYPED_ELEMENT_HAS_REFTRACKING
 	}
@@ -114,10 +199,9 @@ public:
 		FPlatformAtomics::InterlockedIncrement(&RefCount);
 #endif	// UE_TYPED_ELEMENT_HAS_REFCOUNTING
 #if UE_TYPED_ELEMENT_HAS_REFTRACKING
-		if (bCanTrackReference)
+		if (bCanTrackReference && References)
 		{
-			FScopeLock ReferencesLock(&ReferencesCS);
-			ReferenceId = References.Add(FTypedElementReference());
+			ReferenceId = References->AddRef();
 		}
 #endif	// UE_TYPED_ELEMENT_HAS_REFTRACKING
 		return ReferenceId;
@@ -130,10 +214,9 @@ public:
 		FPlatformAtomics::InterlockedDecrement(&RefCount);
 #endif	// UE_TYPED_ELEMENT_HAS_REFCOUNTING
 #if UE_TYPED_ELEMENT_HAS_REFTRACKING
-		if (InReferenceId != INDEX_NONE)
+		if (References)
 		{
-			FScopeLock ReferencesLock(&ReferencesCS);
-			References.RemoveAt(InReferenceId);
+			References->ReleaseRef(InReferenceId);
 		}
 #endif	// UE_TYPED_ELEMENT_HAS_REFTRACKING
 	}
@@ -150,22 +233,13 @@ public:
 	void LogReferences() const
 	{
 #if UE_TYPED_ELEMENT_HAS_REFTRACKING
+		if (References)
 		{
-			FScopeLock ReferencesLock(&ReferencesCS);
-			UE_LOG(LogCore, Error, TEXT("==============================================="));
-			UE_LOG(LogCore, Error, TEXT("External Element References:"));
-			for (const FTypedElementReference& Reference : References)
-			{
-				UE_LOG(LogCore, Error, TEXT("-----------------------------------------------"));
-				Reference.LogReference();
-			}
-			UE_LOG(LogCore, Error, TEXT("==============================================="));
-			if (DestructionRequestCallstack)
-			{
-				UE_LOG(LogCore, Error, TEXT("Destruction requested by:"));
-				DestructionRequestCallstack->LogReference();
-				UE_LOG(LogCore, Error, TEXT("==============================================="));
-			}
+			References->LogReferences();
+		}
+		else
+		{
+			UE_LOG(LogCore, Error, TEXT("CVar 'TypedElements.EnableReferenceTracking' is disabled. Enable it to see reference tracking."));
 		}
 #else	// UE_TYPED_ELEMENT_HAS_REFTRACKING
 		UE_LOG(LogCore, Error, TEXT("UE_TYPED_ELEMENT_HAS_REFTRACKING is disabled. Enable it and recompile to see reference tracking."));
@@ -175,18 +249,10 @@ public:
 	void StoreDestructionRequestCallstack() const
 	{
 #if UE_TYPED_ELEMENT_HAS_REFTRACKING
-		FScopeLock ReferencesLock(&ReferencesCS);
-#if DO_CHECK
-		if (DestructionRequestCallstack)
+		if (References)
 		{
-			UE_LOG(LogCore, Error, TEXT("==============================================="));
-			UE_LOG(LogCore, Error, TEXT("Destruction requested by:"));
-			DestructionRequestCallstack->LogReference();
-			UE_LOG(LogCore, Error, TEXT("==============================================="));
-			UE_LOG(LogCore, Fatal, TEXT("Element has already had its destruction callstack set! (see above)"));
+			References->StoreDestructionRequestCallstack();
 		}
-#endif	// DO_CHECK
-		DestructionRequestCallstack = MakeUnique<FTypedElementReference>();
 #endif	// UE_TYPED_ELEMENT_HAS_REFTRACKING
 	}
 
@@ -213,9 +279,7 @@ private:
 	mutable FTypedElementRefCount RefCount = 0;
 #endif	// UE_TYPED_ELEMENT_HAS_REFCOUNTING
 #if UE_TYPED_ELEMENT_HAS_REFTRACKING
-	mutable FCriticalSection ReferencesCS;
-	mutable TSparseArray<FTypedElementReference> References;
-	mutable TUniquePtr<FTypedElementReference> DestructionRequestCallstack;
+	TUniquePtr<FTypedElementReferences> References;
 #endif	// UE_TYPED_ELEMENT_HAS_REFTRACKING
 };
 
diff --git a/Engine/Source/Runtime/TypedElementFramework/Public/Elements/Framework/TypedElementLimits.h b/Engine/Source/Runtime/TypedElementFramework/Public/Elements/Framework/TypedElementLimits.h
--- a/Engine/Source/Runtime/TypedElementFramework/Public/Elements/Framework/TypedElementLimits.h
+++ b/Engine/Source/Runtime/TypedElementFramework/Public/Elements/Framework/TypedElementLimits.h
@@ -6,7 +6,7 @@
 #include "Math/NumericLimits.h"
 
 #define UE_TYPED_ELEMENT_HAS_REFCOUNTING (1)
-#define UE_TYPED_ELEMENT_HAS_REFTRACKING (0)
+#define UE_TYPED_ELEMENT_HAS_REFTRACKING (UE_BUILD_DEBUG || UE_BUILD_DEVELOPMENT)
 #define UE_TYPED_ELEMENT_HAS_REFERENCING (UE_TYPED_ELEMENT_HAS_REFCOUNTING || UE_TYPED_ELEMENT_HAS_REFTRACKING)
 
 /**
