Don't allow empty element owners to be added to an element store

diff --git a/Engine/Source/Runtime/TypedElementFramework/Public/Elements/Framework/TypedElementOwnerStore.h b/Engine/Source/Runtime/TypedElementFramework/Public/Elements/Framework/TypedElementOwnerStore.h
--- a/Engine/Source/Runtime/TypedElementFramework/Public/Elements/Framework/TypedElementOwnerStore.h
+++ b/Engine/Source/Runtime/TypedElementFramework/Public/Elements/Framework/TypedElementOwnerStore.h
@@ -128,6 +128,7 @@ public:
 	{
 		TTypedElementOwnerScopedAccess<ElementDataType> ScopedAccess(&ElementOwnerMapCS);
 		checkf(!ElementOwnerMap.Contains(InKey), TEXT("Element owner has already been registered for this key (%s)! This will leak elements, and you should unregister and destroy the old element owner for this instance before adding a new one!"), *GetTypedElementOwnerStoreKeyDebugString(InKey));
+		checkf(InElementOwner, TEXT("Element owner passed to RegisterElementOwner was null for key (%s)! Element owners must be valid!"), *GetTypedElementOwnerStoreKeyDebugString(InKey));
 		ScopedAccess.Private_SetElementOwner(&ElementOwnerMap.Add(InKey, MoveTemp(InElementOwner)));
 		return ScopedAccess;
 	}
@@ -163,8 +164,8 @@ public:
 	}
 
 	/**
-	 * Find the given element owner, or call the register callback to create one if none is found.
-	 * @note This must be paired with a call to UnregisterElementOwner.
+	 * Find the given element owner, or call the register callback to attempt to create one if none is found.
+	 * @note This must be paired with a call to UnregisterElementOwner if a new element owner is created.
 	 */
 	TTypedElementOwnerScopedAccess<ElementDataType> FindOrRegisterElementOwner(const KeyDataType& InKey, TFunctionRef<TTypedElementOwner<ElementDataType>()> InCreateElement)
 	{
@@ -173,9 +174,9 @@ public:
 		{
 			ScopedAccess.Private_SetElementOwner(ExistingElement);
 		}
-		else
+		else if (TTypedElementOwner<ElementDataType> NewElement = InCreateElement())
 		{
-			ScopedAccess.Private_SetElementOwner(&ElementOwnerMap.Add(InKey, InCreateElement()));
+			ScopedAccess.Private_SetElementOwner(&ElementOwnerMap.Add(InKey, MoveTemp(NewElement)));
 		}
 		return ScopedAccess;
 	}
