Cleaned-up lingering asset selection references during shutdown

diff --git a/Engine/Source/Editor/StatusBar/Private/StatusBarSubsystem.cpp b/Engine/Source/Editor/StatusBar/Private/StatusBarSubsystem.cpp
--- a/Engine/Source/Editor/StatusBar/Private/StatusBarSubsystem.cpp
+++ b/Engine/Source/Editor/StatusBar/Private/StatusBarSubsystem.cpp
@@ -215,6 +215,8 @@ void UStatusBarSubsystem::Deinitialize()
 	FSourceControlCommands::Unregister();
 
 	FSlateNotificationManager::Get().SetProgressNotificationHandler(nullptr);
+
+	StatusBarContentBrowser.Reset();
 }
 
 bool UStatusBarSubsystem::FocusDebugConsole(TSharedRef<SWindow> ParentWindow)
diff --git a/Engine/Source/Runtime/TypedElementFramework/Private/Elements/Framework/TypedElementList.cpp b/Engine/Source/Runtime/TypedElementFramework/Private/Elements/Framework/TypedElementList.cpp
--- a/Engine/Source/Runtime/TypedElementFramework/Private/Elements/Framework/TypedElementList.cpp
+++ b/Engine/Source/Runtime/TypedElementFramework/Private/Elements/Framework/TypedElementList.cpp
@@ -8,7 +8,14 @@ namespace TypedElementList_Private
 
 void GetElementImpl(const UTypedElementRegistry* InRegistry, const FTypedElementHandle& InElementHandle, const UClass* InBaseInterfaceType, FTypedElement& OutElement)
 {
-	InRegistry->Private_GetElementImpl(InElementHandle, InBaseInterfaceType, OutElement);
+	if (InRegistry)
+	{
+		InRegistry->Private_GetElementImpl(InElementHandle, InBaseInterfaceType, OutElement);
+	}
+	else
+	{
+		OutElement.Release();
+	}
 }
 
 } // namespace TypedElementList_Private
@@ -125,6 +132,7 @@ void UTypedElementList::BeginDestroy()
 {
 	Super::BeginDestroy();
 
+	Empty();
 	LegacySync.Reset();
 	if (UTypedElementRegistry* RegistryPtr = Registry.Get())
 	{
