Added a general value counter for typed element lists

This is used to count the number of element types that are present in the list (for quick queries), but can also be expanded via the UTypedElementCounterInterface to allow specific element types to count extra data.

diff --git a/Engine/Source/Runtime/Engine/Private/Elements/Actor/ActorElementSelectionInterface.cpp b/Engine/Source/Runtime/Engine/Private/Elements/Actor/ActorElementSelectionInterface.cpp
--- a/Engine/Source/Runtime/Engine/Private/Elements/Actor/ActorElementSelectionInterface.cpp
+++ b/Engine/Source/Runtime/Engine/Private/Elements/Actor/ActorElementSelectionInterface.cpp
@@ -7,25 +7,10 @@
 
 int32 UActorElementSelectionInterface::GetNumSelectedActors(const UTypedElementList* InCurrentSelection)
 {
-	int32 NumSelected = 0;
-	InCurrentSelection->ForEachElementHandle([&NumSelected](const FTypedElementHandle& InSelectedElement)
-	{
-		if (ActorElementDataUtil::GetActorFromHandle(InSelectedElement, /*bSilent*/true))
-		{
-			++NumSelected;
-		}
-		return true;
-	});
-	return NumSelected;
+	return InCurrentSelection->CountElementsOfType(NAME_Actor);
 }
 
 bool UActorElementSelectionInterface::HasSelectedActors(const UTypedElementList* InCurrentSelection)
 {
-	bool bHasSelectedActors = false;
-	InCurrentSelection->ForEachElementHandle([&bHasSelectedActors](const FTypedElementHandle& InSelectedElement)
-	{
-		bHasSelectedActors = ActorElementDataUtil::GetActorFromHandle(InSelectedElement, /*bSilent*/true) != nullptr;
-		return !bHasSelectedActors;
-	});
-	return bHasSelectedActors;
+	return InCurrentSelection->HasElementsOfType(NAME_Actor);
 }
diff --git a/Engine/Source/Runtime/Engine/Private/Elements/Component/ComponentElementSelectionInterface.cpp b/Engine/Source/Runtime/Engine/Private/Elements/Component/ComponentElementSelectionInterface.cpp
--- a/Engine/Source/Runtime/Engine/Private/Elements/Component/ComponentElementSelectionInterface.cpp
+++ b/Engine/Source/Runtime/Engine/Private/Elements/Component/ComponentElementSelectionInterface.cpp
@@ -7,25 +7,10 @@
 
 int32 UComponentElementSelectionInterface::GetNumSelectedComponents(const UTypedElementList* InCurrentSelection)
 {
-	int32 NumSelected = 0;
-	InCurrentSelection->ForEachElementHandle([&NumSelected](const FTypedElementHandle& InSelectedElement)
-	{
-		if (ComponentElementDataUtil::GetComponentFromHandle(InSelectedElement, /*bSilent*/true))
-		{
-			++NumSelected;
-		}
-		return true;
-	});
-	return NumSelected;
+	return InCurrentSelection->CountElementsOfType(NAME_Components);
 }
 
 bool UComponentElementSelectionInterface::HasSelectedComponents(const UTypedElementList* InCurrentSelection)
 {
-	bool bHasSelectedComponents = false;
-	InCurrentSelection->ForEachElementHandle([&bHasSelectedComponents](const FTypedElementHandle& InSelectedElement)
-	{
-		bHasSelectedComponents = ComponentElementDataUtil::GetComponentFromHandle(InSelectedElement, /*bSilent*/true) != nullptr;
-		return !bHasSelectedComponents;
-	});
-	return bHasSelectedComponents;
+	return InCurrentSelection->HasElementsOfType(NAME_Components);
 }
diff --git a/Engine/Source/Runtime/Engine/Private/Elements/SMInstance/SMInstanceElementSelectionInterface.cpp b/Engine/Source/Runtime/Engine/Private/Elements/SMInstance/SMInstanceElementSelectionInterface.cpp
--- a/Engine/Source/Runtime/Engine/Private/Elements/SMInstance/SMInstanceElementSelectionInterface.cpp
+++ b/Engine/Source/Runtime/Engine/Private/Elements/SMInstance/SMInstanceElementSelectionInterface.cpp
@@ -7,25 +7,10 @@
 
 int32 USMInstanceElementSelectionInterface::GetNumSelectedSMInstances(const UTypedElementList* InCurrentSelection)
 {
-	int32 NumSelected = 0;
-	InCurrentSelection->ForEachElementHandle([&NumSelected](const FTypedElementHandle& InSelectedElement)
-	{
-		if (SMInstanceElementDataUtil::GetSMInstanceFromHandle(InSelectedElement, /*bSilent*/true))
-		{
-			++NumSelected;
-		}
-		return true;
-	});
-	return NumSelected;
+	return InCurrentSelection->CountElementsOfType(NAME_SMInstance);
 }
 
 bool USMInstanceElementSelectionInterface::HasSelectedSMInstances(const UTypedElementList* InCurrentSelection)
 {
-	bool bHasSelectedSMInstances = false;
-	InCurrentSelection->ForEachElementHandle([&bHasSelectedSMInstances](const FTypedElementHandle& InSelectedElement)
-	{
-		bHasSelectedSMInstances = !!SMInstanceElementDataUtil::GetSMInstanceFromHandle(InSelectedElement, /*bSilent*/true);
-		return !bHasSelectedSMInstances;
-	});
-	return bHasSelectedSMInstances;
+	return InCurrentSelection->HasElementsOfType(NAME_SMInstance);
 }
diff --git a/Engine/Source/Runtime/TypedElementFramework/Private/Elements/Framework/TypedElementCounter.cpp b/Engine/Source/Runtime/TypedElementFramework/Private/Elements/Framework/TypedElementCounter.cpp
new file mode 100644
--- /dev/null
+++ b/Engine/Source/Runtime/TypedElementFramework/Private/Elements/Framework/TypedElementCounter.cpp
@@ -0,0 +1,64 @@
+// Copyright Epic Games, Inc. All Rights Reserved.
+
+#include "Elements/Framework/TypedElementCounter.h"
+#include "Elements/Framework/TypedElementRegistry.h"
+
+FTypedElementCounter::FTypedElementCounter(UTypedElementRegistry* InRegistry)
+{
+	Initialize(InRegistry);
+}
+
+void FTypedElementCounter::Initialize(UTypedElementRegistry* InRegistry)
+{
+	checkf(!Registry.Get(), TEXT("Initialize has already been called!"));
+	Registry = InRegistry;
+	checkf(InRegistry, TEXT("Registry is null!"));
+}
+
+void FTypedElementCounter::AddElement(const FTypedElementHandle& InElementHandle)
+{
+	UTypedElementRegistry* RegistryPtr = Registry.Get();
+	if (!RegistryPtr || !InElementHandle)
+	{
+		return;
+	}
+
+	IncrementCounter(GetElementTypeCategoryName(), InElementHandle.GetId().GetTypeId());
+	
+	if (const TTypedElement<UTypedElementCounterInterface> ElementCounterHandle = RegistryPtr->GetElement<UTypedElementCounterInterface>(InElementHandle))
+	{
+		ElementCounterHandle.IncrementCountersForElement(*this);
+	}
+}
+
+void FTypedElementCounter::RemoveElement(const FTypedElementHandle& InElementHandle)
+{
+	UTypedElementRegistry* RegistryPtr = Registry.Get();
+	if (!RegistryPtr || !InElementHandle)
+	{
+		return;
+	}
+
+	DecrementCounter(GetElementTypeCategoryName(), InElementHandle.GetId().GetTypeId());
+
+	if (const TTypedElement<UTypedElementCounterInterface> ElementCounterHandle = RegistryPtr->GetElement<UTypedElementCounterInterface>(InElementHandle))
+	{
+		ElementCounterHandle.DecrementCountersForElement(*this);
+	}
+}
+
+void FTypedElementCounter::ClearCounters(const FName InCategory)
+{
+	CounterCategories.Remove(InCategory);
+}
+
+void FTypedElementCounter::ClearCounters()
+{
+	CounterCategories.Reset();
+}
+
+FName FTypedElementCounter::GetElementTypeCategoryName()
+{
+	static const FName NAME_ElementType = "ElementType";
+	return NAME_ElementType;
+}
diff --git a/Engine/Source/Runtime/TypedElementFramework/Private/Elements/Framework/TypedElementList.cpp b/Engine/Source/Runtime/TypedElementFramework/Private/Elements/Framework/TypedElementList.cpp
--- a/Engine/Source/Runtime/TypedElementFramework/Private/Elements/Framework/TypedElementList.cpp
+++ b/Engine/Source/Runtime/TypedElementFramework/Private/Elements/Framework/TypedElementList.cpp
@@ -125,6 +125,9 @@ void UTypedElementList::Initialize(UTypedElementRegistry* InRegistry)
 	checkf(!Registry.Get(), TEXT("Initialize has already been called!"));
 	Registry = InRegistry;
 	checkf(InRegistry, TEXT("Registry is null!"));
+
+	ElementCounts.Initialize(InRegistry);
+
 	InRegistry->Private_OnElementListCreated(this);
 }
 
@@ -144,14 +147,19 @@ void UTypedElementList::BeginDestroy()
 UTypedElementList* UTypedElementList::Clone() const
 {
 	UTypedElementList* ClonedElementList = Private_CreateElementList(Registry.Get());
-	ClonedElementList->ElementCombinedIds = ElementCombinedIds;
-	ClonedElementList->ElementHandles = ElementHandles;
+	for (const FTypedElementHandle& ElementHandle : ElementHandles)
+	{
+		ClonedElementList->Add(ElementHandle);
+	}
 	return ClonedElementList;
 }
 
 UTypedElementInterface* UTypedElementList::GetElementInterface(const FTypedElementHandle& InElementHandle, const TSubclassOf<UTypedElementInterface>& InBaseInterfaceType) const
 {
-	return Registry->GetElementInterface(InElementHandle, InBaseInterfaceType);
+	UTypedElementRegistry* RegistryPtr = Registry.Get();
+	return RegistryPtr
+		? Registry->GetElementInterface(InElementHandle, InBaseInterfaceType)
+		: nullptr;
 }
 
 bool UTypedElementList::HasElements(const TSubclassOf<UTypedElementInterface>& InBaseInterfaceType) const
@@ -194,6 +202,34 @@ int32 UTypedElementList::CountElements(const TSubclassOf<UTypedElementInterface>
 	return NumFilteredElements;
 }
 
+bool UTypedElementList::HasElementsOfType(const FName InElementTypeName) const
+{
+	return CountElementsOfType(InElementTypeName) > 0;
+}
+
+bool UTypedElementList::HasElementsOfType(const FTypedHandleTypeId InElementTypeId) const
+{
+	return CountElementsOfType(InElementTypeId) > 0;
+}
+
+int32 UTypedElementList::CountElementsOfType(const FName InElementTypeName) const
+{
+	if (UTypedElementRegistry* RegistryPtr = Registry.Get())
+	{
+		const FTypedHandleTypeId ElementTypeId = RegistryPtr->GetRegisteredElementTypeId(InElementTypeName);
+		if (ElementTypeId > 0)
+		{
+			return CountElementsOfType(ElementTypeId);
+		}
+	}
+	return 0;
+}
+
+int32 UTypedElementList::CountElementsOfType(const FTypedHandleTypeId InElementTypeId) const
+{
+	return ElementCounts.GetCounterValue(FTypedElementCounter::GetElementTypeCategoryName(), InElementTypeId);
+}
+
 TArray<FTypedElementHandle> UTypedElementList::GetElementHandles(const TSubclassOf<UTypedElementInterface>& InBaseInterfaceType) const
 {
 	TArray<FTypedElementHandle> FilteredElementHandles;
@@ -237,6 +273,7 @@ bool UTypedElementList::AddElementImpl(FTypedElementHandle&& InElementHandle)
 	if (!bAlreadyAdded)
 	{
 		const FTypedElementHandle& AddedElementHandle = ElementHandles.Add_GetRef(MoveTemp(InElementHandle));
+		ElementCounts.AddElement(AddedElementHandle);
 		NoteListChanged(EChangeType::Added, AddedElementHandle);
 	}
 
@@ -264,6 +301,7 @@ bool UTypedElementList::RemoveElementImpl(const FTypedElementId& InElementId)
 
 		FTypedElementHandle RemovedElementHandle = MoveTemp(ElementHandles[ElementHandleIndexToRemove]);
 		ElementHandles.RemoveAt(ElementHandleIndexToRemove, 1, /*bAllowShrinking*/false);
+		ElementCounts.RemoveElement(RemovedElementHandle);
 
 		NoteListChanged(EChangeType::Removed, RemovedElementHandle);
 	}
@@ -288,6 +326,7 @@ int32 UTypedElementList::RemoveAllElementsImpl(TFunctionRef<bool(const FTypedEle
 				FTypedElementHandle RemovedElementHandle = MoveTemp(ElementHandles[Index]);
 				ElementCombinedIds.Remove(RemovedElementHandle.GetId().GetCombinedId());
 				ElementHandles.RemoveAt(Index, 1, /*bAllowShrinking*/false);
+				ElementCounts.RemoveElement(RemovedElementHandle);
 
 				NoteListChanged(EChangeType::Removed, RemovedElementHandle);
 
diff --git a/Engine/Source/Runtime/TypedElementFramework/Public/Elements/Framework/TypedElementCounter.h b/Engine/Source/Runtime/TypedElementFramework/Public/Elements/Framework/TypedElementCounter.h
new file mode 100644
--- /dev/null
+++ b/Engine/Source/Runtime/TypedElementFramework/Public/Elements/Framework/TypedElementCounter.h
@@ -0,0 +1,296 @@
+// Copyright Epic Games, Inc. All Rights Reserved.
+
+#pragma once
+
+#include <type_traits>
+#include "CoreMinimal.h"
+#include "Containers/SortedMap.h"
+#include "Elements/Framework/TypedElementHandle.h"
+#include "TypedElementCounter.generated.h"
+
+class UTypedElementRegistry;
+
+template <typename KeyType>
+inline FString GetTypedElementCounterKeyName()
+{
+	if constexpr (std::is_convertible_v<KeyType, const UObject*>)
+	{
+		return std::remove_pointer_t<KeyType>::StaticClass()->GetName();
+	}
+	else
+	{
+		return TNameOf<KeyType>::GetName();
+	}
+}
+
+/**
+ * Tracks various categories of counters for the typed elements (eg, the number of elements of a given type).
+ * The categories counted may be expanded by element types that implement the UTypedElementCounterInterface API.
+ */
+class TYPEDELEMENTFRAMEWORK_API FTypedElementCounter
+{
+public:
+	using FCounterValue = int32;
+
+	FTypedElementCounter() = default;
+	explicit FTypedElementCounter(UTypedElementRegistry* InRegistry);
+
+	FTypedElementCounter(const FTypedElementCounter&) = delete;
+	FTypedElementCounter& operator=(const FTypedElementCounter&) = delete;
+
+	FTypedElementCounter(FTypedElementCounter&&) = default;
+	FTypedElementCounter& operator=(FTypedElementCounter&&) = default;
+
+	/**
+	 * Explicitly initialize this instance, if it was previously default constructed.
+	 */
+	void Initialize(UTypedElementRegistry* InRegistry);
+
+	/**
+	 * Update the counter categories to include the given element.
+	 */
+	void AddElement(const FTypedElementHandle& InElementHandle);
+
+	/**
+	 * Update the counter categories to remove the given element.
+	 */
+	void RemoveElement(const FTypedElementHandle& InElementHandle);
+
+	/**
+	 * Increment the counter for the given category and key.
+	 * @note For use by UTypedElementCounterInterface.
+	 */
+	template <typename KeyType>
+	void IncrementCounter(const FName InCategory, const KeyType& InKey)
+	{
+		TUniquePtr<ICounterCategory>& CounterCategory = CounterCategories.FindOrAdd(InCategory);
+		if (!CounterCategory)
+		{
+			CounterCategory = MakeUnique<TCounterCategory<KeyType>>();
+		}
+		CounterCategory->IncrementCounter(InKey);
+	}
+
+	/**
+	 * Decrement the counter for the given category and key.
+	 * @note For use by UTypedElementCounterInterface.
+	 */
+	template <typename KeyType>
+	void DecrementCounter(const FName InCategory, const KeyType& InKey)
+	{
+		if (TUniquePtr<ICounterCategory>* CounterCategory = CounterCategories.Find(InCategory))
+		{
+			(*CounterCategory)->DecrementCounter(InKey);
+		}
+	}
+	
+	/**
+	 * Get the value of the counter for the given category and key, or zero if the category or key are unknown.
+	 */
+	template <typename KeyType>
+	FCounterValue GetCounterValue(const FName InCategory, const KeyType& InKey) const
+	{
+		if (const TUniquePtr<ICounterCategory>* CounterCategory = CounterCategories.Find(InCategory))
+		{
+			return (*CounterCategory)->GetCounterValue(InKey);
+		}
+		return 0;
+	}
+
+	/**
+	 * Enumerate the value of the counters for the given category.
+	 */
+	template <typename KeyType>
+	void ForEachCounterValue(const FName InCategory, TFunctionRef<bool(const KeyType&, FCounterValue)> InCallback) const
+	{
+		if (const TUniquePtr<ICounterCategory>* CounterCategory = CounterCategories.Find(InCategory))
+		{
+			(*CounterCategory)->ForEachCounterValue<KeyType>(InCallback);
+		}
+	}
+
+	/**
+	 * Clear the value of the counter for the given category and key.
+	 */
+	template <typename KeyType>
+	void ClearCounter(const FName InCategory, const KeyType& InKey)
+	{
+		if (TUniquePtr<ICounterCategory>* CounterCategory = CounterCategories.Find(InCategory))
+		{
+			(*CounterCategory)->ClearCounter(InKey);
+		}
+	}
+
+	/**
+	 * Clear the value of the counters for the given category.
+	 */
+	void ClearCounters(const FName InCategory);
+
+	/**
+	 * Clear all counters.
+	 */
+	void ClearCounters();
+
+	/**
+	 * Get the category name used to count element types (by type ID).
+	 */
+	static FName GetElementTypeCategoryName();
+
+private:
+	class ICounterCategory
+	{
+	public:
+		virtual ~ICounterCategory() = default;
+
+		template <typename KeyType>
+		FCounterValue GetCounterValue(const KeyType& InKey) const
+		{
+			return GetCounterValueImpl(&InKey, GetTypedElementCounterKeyName<KeyType>());
+		}
+
+		template <typename KeyType>
+		void ForEachCounterValue(TFunctionRef<bool(const KeyType&, FCounterValue)> InCallback) const
+		{
+			ForEachCounterValueImpl(&InCallback, GetTypedElementCounterKeyName<KeyType>());
+		}
+
+		template <typename KeyType>
+		void IncrementCounter(const KeyType& InKey)
+		{
+			IncrementCounterImpl(&InKey, GetTypedElementCounterKeyName<KeyType>());
+		}
+
+		template <typename KeyType>
+		void DecrementCounter(const KeyType& InKey)
+		{
+			DecrementCounterImpl(&InKey, GetTypedElementCounterKeyName<KeyType>());
+		}
+
+		template <typename KeyType>
+		void ClearCounter(const KeyType& InKey)
+		{
+			ClearCounterImpl(&InKey, GetTypedElementCounterKeyName<KeyType>());
+		}
+
+		void ClearCounters()
+		{
+			ClearCountersImpl();
+		}
+
+	protected:
+		virtual FCounterValue GetCounterValueImpl(const void* InKey, const FStringView InKeyTypeName) const = 0;
+		virtual void ForEachCounterValueImpl(const void* InCallback, const FStringView InKeyTypeName) const = 0;
+		virtual void IncrementCounterImpl(const void* InKey, const FStringView InKeyTypeName) = 0;
+		virtual void DecrementCounterImpl(const void* InKey, const FStringView InKeyTypeName) = 0;
+		virtual void ClearCounterImpl(const void* InKey, const FStringView InKeyTypeName) = 0;
+		virtual void ClearCountersImpl() = 0;
+	};
+
+	template <typename KeyType>
+	class TCounterCategory : public ICounterCategory
+	{
+	public:
+		TCounterCategory()
+			: KeyTypeName(GetTypedElementCounterKeyName<KeyType>())
+		{
+		}
+
+	private:
+		virtual FCounterValue GetCounterValueImpl(const void* InKey, const FStringView InKeyTypeName) const override
+		{
+			checkf(KeyTypeName == InKeyTypeName, TEXT("Invalid counter access! Counter is keyed against '%s' but '%s' was given."), *KeyTypeName, *FString(InKeyTypeName));
+
+			const KeyType* KeyPtr = static_cast<const KeyType*>(InKey);
+			return Counters.FindRef(*KeyPtr);
+		}
+
+		virtual void ForEachCounterValueImpl(const void* InCallback, const FStringView InKeyTypeName) const override
+		{
+			checkf(KeyTypeName == InKeyTypeName, TEXT("Invalid counter access! Counter is keyed against '%s' but '%s' was given."), *KeyTypeName, *FString(InKeyTypeName));
+
+			const TFunctionRef<bool(const KeyType&, FCounterValue)>* CallbackPtr = static_cast<const TFunctionRef<bool(const KeyType&, FCounterValue)>*>(InCallback);
+			for (const TTuple<KeyType, FCounterValue>& CounterValuePair : Counters)
+			{
+				if (!(*CallbackPtr)(CounterValuePair.Key, CounterValuePair.Value))
+				{
+					break;
+				}
+			}
+		}
+
+		virtual void IncrementCounterImpl(const void* InKey, const FStringView InKeyTypeName) override
+		{
+			checkf(KeyTypeName == InKeyTypeName, TEXT("Invalid counter access! Counter is keyed against '%s' but '%s' was given."), *KeyTypeName, *FString(InKeyTypeName));
+
+			const KeyType* KeyPtr = static_cast<const KeyType*>(InKey);
+			FCounterValue& CounterValue = Counters.FindOrAdd(*KeyPtr);
+			++CounterValue;
+		}
+
+		virtual void DecrementCounterImpl(const void* InKey, const FStringView InKeyTypeName) override
+		{
+			checkf(KeyTypeName == InKeyTypeName, TEXT("Invalid counter access! Counter is keyed against '%s' but '%s' was given."), *KeyTypeName, *FString(InKeyTypeName));
+
+			const KeyType* KeyPtr = static_cast<const KeyType*>(InKey);
+			if (FCounterValue* CounterValue = Counters.Find(*KeyPtr))
+			{
+				checkf(*CounterValue > 0, TEXT("Counter value unflow!"));
+				if (--(*CounterValue) == 0)
+				{
+					Counters.Remove(*KeyPtr);
+				}
+			}
+		}
+
+		virtual void ClearCounterImpl(const void* InKey, const FStringView InKeyTypeName) override
+		{
+			checkf(KeyTypeName == InKeyTypeName, TEXT("Invalid counter access! Counter is keyed against '%s' but '%s' was given."), *KeyTypeName, *FString(InKeyTypeName));
+
+			const KeyType* KeyPtr = static_cast<const KeyType*>(InKey);
+			Counters.Remove(*KeyPtr);
+		}
+
+		virtual void ClearCountersImpl() override
+		{
+			Counters.Empty();
+		}
+
+	private:
+		FString KeyTypeName;
+		TMap<KeyType, FCounterValue> Counters;
+	};
+
+	/**
+	 * Element registry this element counter is associated with.
+	 */
+	TWeakObjectPtr<UTypedElementRegistry> Registry;
+
+	/**
+	 * Counter categories.
+	 */
+	TSortedMap<FName, TUniquePtr<ICounterCategory>, FDefaultAllocator, FNameFastLess> CounterCategories;
+};
+
+UCLASS(Abstract)
+class TYPEDELEMENTFRAMEWORK_API UTypedElementCounterInterface : public UTypedElementInterface
+{
+	GENERATED_BODY()
+
+public:
+	/**
+	 * Increment additional counters for the given element.
+	 */
+	virtual void IncrementCountersForElement(const FTypedElementHandle& InElementHandle, FTypedElementCounter& InOutCounter) {}
+
+	/**
+	 * Decrement additional counters for the given element.
+	 */
+	virtual void DecrementCountersForElement(const FTypedElementHandle& InElementHandle, FTypedElementCounter& InOutCounter) {}
+};
+
+template <>
+struct TTypedElement<UTypedElementCounterInterface> : public TTypedElementBase<UTypedElementCounterInterface>
+{
+	void IncrementCountersForElement(FTypedElementCounter& InOutCounter) const { InterfacePtr->IncrementCountersForElement(*this, InOutCounter); }
+	void DecrementCountersForElement(FTypedElementCounter& InOutCounter) const { InterfacePtr->DecrementCountersForElement(*this, InOutCounter); }
+};
diff --git a/Engine/Source/Runtime/TypedElementFramework/Public/Elements/Framework/TypedElementList.h b/Engine/Source/Runtime/TypedElementFramework/Public/Elements/Framework/TypedElementList.h
--- a/Engine/Source/Runtime/TypedElementFramework/Public/Elements/Framework/TypedElementList.h
+++ b/Engine/Source/Runtime/TypedElementFramework/Public/Elements/Framework/TypedElementList.h
@@ -9,6 +9,7 @@
 #include "Containers/ArrayView.h"
 #include "Templates/UniquePtr.h"
 #include "Elements/Framework/TypedElementHandle.h"
+#include "Elements/Framework/TypedElementCounter.h"
 #include "TypedElementList.generated.h"
 
 class UTypedElementList;
@@ -269,6 +270,20 @@ public:
 	UFUNCTION(BlueprintCallable, BlueprintPure=false, Category="TypedElementFramework|List")
 	int32 CountElements(const TSubclassOf<UTypedElementInterface>& InBaseInterfaceType = nullptr) const;
 
+	/**
+	 * Test whether there are elements in this list of the given type.
+	 */
+	UFUNCTION(BlueprintCallable, BlueprintPure=false, Category="TypedElementFramework|List")
+	bool HasElementsOfType(const FName InElementTypeName) const;
+	bool HasElementsOfType(const FTypedHandleTypeId InElementTypeId) const;
+
+	/**
+	 * Count the number of elements in this list of the given type.
+	 */
+	UFUNCTION(BlueprintCallable, BlueprintPure=false, Category="TypedElementFramework|List")
+	int32 CountElementsOfType(const FName InElementTypeName) const;
+	int32 CountElementsOfType(const FTypedHandleTypeId InElementTypeId) const;
+
 	/**
 	 * Get the handle of every element in this list, optionally filtering to elements that implement the given interface.
 	 */
@@ -361,6 +376,7 @@ public:
 		NoteListMayChange();
 		ElementCombinedIds.Empty(InSlack);
 		ElementHandles.Empty(InSlack);
+		ElementCounts.ClearCounters();
 		NoteListChanged(EChangeType::Cleared);
 	}
 
@@ -373,6 +389,7 @@ public:
 		NoteListMayChange();
 		ElementCombinedIds.Reset();
 		ElementHandles.Reset();
+		ElementCounts.ClearCounters();
 		NoteListChanged(EChangeType::Cleared);
 	}
 
@@ -531,6 +548,14 @@ public:
 		});
 	}
 
+	/**
+	 * Get the counter for the elements within the list.
+	 */
+	const FTypedElementCounter& GetCounter() const
+	{
+		return ElementCounts;
+	}
+
 	/**
 	 * Access the delegate that is invoked whenever this element list is potentially about to change.
 	 * @note This may be called even if no actual change happens, though once a change does happen it won't be called again until after the next call to NotifyPendingChanges.
@@ -622,6 +647,11 @@ private:
 	 */
 	TArray<FTypedElementHandle> ElementHandles;
 
+	/**
+	 * Tracks various categories of counters for the elements within the list (eg, the number of elements of a given type).
+	 */
+	FTypedElementCounter ElementCounts;
+
 	/**
 	 * Delegate that is invoked whenever this element list is potentially about to change.
 	 */
diff --git a/Engine/Source/Runtime/TypedElementFramework/Public/Elements/Framework/TypedElementRegistry.h b/Engine/Source/Runtime/TypedElementFramework/Public/Elements/Framework/TypedElementRegistry.h
--- a/Engine/Source/Runtime/TypedElementFramework/Public/Elements/Framework/TypedElementRegistry.h
+++ b/Engine/Source/Runtime/TypedElementFramework/Public/Elements/Framework/TypedElementRegistry.h
@@ -75,10 +75,18 @@ public:
 	 */
 	FORCEINLINE FTypedHandleTypeId GetRegisteredElementTypeId(const FName InElementTypeName) const
 	{
-		FReadScopeLock RegisteredElementTypesLock(RegisteredElementTypesRW);
+		const FRegisteredElementType* RegisteredElementType = GetRegisteredElementTypeFromName(InElementTypeName);
+		return RegisteredElementType ? RegisteredElementType->TypeId : 0;
+	}
 
-		const FTypedHandleTypeId* TypeId = RegisteredElementTypesNameToId.Find(InElementTypeName);
-		return TypeId ? *TypeId : 0;
+	/**
+	 * Get the element type name for the associated element type ID, if any.
+	 * @return The element type name, or None if the given ID wasn't registered.
+	 */
+	FORCEINLINE FName GetRegisteredElementTypeName(const FTypedHandleTypeId InElementTypeId) const
+	{
+		const FRegisteredElementType* RegisteredElementType = GetRegisteredElementTypeFromId(InElementTypeId);
+		return RegisteredElementType ? RegisteredElementType->TypeName : FName();
 	}
 
 	/**
