Updated FTypedElementCounter to use FName for its type name checks

This saves 100ms when adding 2,000,000 elements to the set, taking the append time in the perf test from 210ms to 110ms. The counter still accounts for ~60ms of the append time though, which isn't ideal, but there's nothing more to be gained outside of maybe ~10ms spent resolving a weak pointer.

diff --git a/Engine/Source/Runtime/TypedElementFramework/Public/Elements/Framework/TypedElementCounter.h b/Engine/Source/Runtime/TypedElementFramework/Public/Elements/Framework/TypedElementCounter.h
--- a/Engine/Source/Runtime/TypedElementFramework/Public/Elements/Framework/TypedElementCounter.h
+++ b/Engine/Source/Runtime/TypedElementFramework/Public/Elements/Framework/TypedElementCounter.h
@@ -11,15 +11,16 @@
 class UTypedElementRegistry;
 
 template <typename KeyType>
-inline FString GetTypedElementCounterKeyName()
+inline FName GetTypedElementCounterKeyName()
 {
 	if constexpr (std::is_convertible_v<KeyType, const UObject*>)
 	{
-		return std::remove_pointer_t<KeyType>::StaticClass()->GetName();
+		return std::remove_pointer_t<KeyType>::StaticClass()->GetFName();
 	}
 	else
 	{
-		return TNameOf<KeyType>::GetName();
+		static const FName KeyTypeName = TNameOf<KeyType>::GetName();
+		return KeyTypeName;
 	}
 }
 
@@ -178,11 +179,11 @@ private:
 		}
 
 	protected:
-		virtual FCounterValue GetCounterValueImpl(const void* InKey, const FStringView InKeyTypeName) const = 0;
-		virtual void ForEachCounterValueImpl(const void* InCallback, const FStringView InKeyTypeName) const = 0;
-		virtual void IncrementCounterImpl(const void* InKey, const FStringView InKeyTypeName) = 0;
-		virtual void DecrementCounterImpl(const void* InKey, const FStringView InKeyTypeName) = 0;
-		virtual void ClearCounterImpl(const void* InKey, const FStringView InKeyTypeName) = 0;
+		virtual FCounterValue GetCounterValueImpl(const void* InKey, const FName InKeyTypeName) const = 0;
+		virtual void ForEachCounterValueImpl(const void* InCallback, const FName InKeyTypeName) const = 0;
+		virtual void IncrementCounterImpl(const void* InKey, const FName InKeyTypeName) = 0;
+		virtual void DecrementCounterImpl(const void* InKey, const FName InKeyTypeName) = 0;
+		virtual void ClearCounterImpl(const void* InKey, const FName InKeyTypeName) = 0;
 		virtual void ClearCountersImpl() = 0;
 	};
 
@@ -196,17 +197,17 @@ private:
 		}
 
 	private:
-		virtual FCounterValue GetCounterValueImpl(const void* InKey, const FStringView InKeyTypeName) const override
+		virtual FCounterValue GetCounterValueImpl(const void* InKey, const FName InKeyTypeName) const override
 		{
-			checkf(KeyTypeName == InKeyTypeName, TEXT("Invalid counter access! Counter is keyed against '%s' but '%s' was given."), *KeyTypeName, *FString(InKeyTypeName));
+			checkf(KeyTypeName == InKeyTypeName, TEXT("Invalid counter access! Counter is keyed against '%s' but '%s' was given."), *KeyTypeName.ToString(), *InKeyTypeName.ToString());
 
 			const KeyType* KeyPtr = static_cast<const KeyType*>(InKey);
 			return Counters.FindRef(*KeyPtr);
 		}
 
-		virtual void ForEachCounterValueImpl(const void* InCallback, const FStringView InKeyTypeName) const override
+		virtual void ForEachCounterValueImpl(const void* InCallback, const FName InKeyTypeName) const override
 		{
-			checkf(KeyTypeName == InKeyTypeName, TEXT("Invalid counter access! Counter is keyed against '%s' but '%s' was given."), *KeyTypeName, *FString(InKeyTypeName));
+			checkf(KeyTypeName == InKeyTypeName, TEXT("Invalid counter access! Counter is keyed against '%s' but '%s' was given."), *KeyTypeName.ToString(), *InKeyTypeName.ToString());
 
 			const TFunctionRef<bool(const KeyType&, FCounterValue)>* CallbackPtr = static_cast<const TFunctionRef<bool(const KeyType&, FCounterValue)>*>(InCallback);
 			for (const TTuple<KeyType, FCounterValue>& CounterValuePair : Counters)
@@ -218,18 +219,18 @@ private:
 			}
 		}
 
-		virtual void IncrementCounterImpl(const void* InKey, const FStringView InKeyTypeName) override
+		virtual void IncrementCounterImpl(const void* InKey, const FName InKeyTypeName) override
 		{
-			checkf(KeyTypeName == InKeyTypeName, TEXT("Invalid counter access! Counter is keyed against '%s' but '%s' was given."), *KeyTypeName, *FString(InKeyTypeName));
+			checkf(KeyTypeName == InKeyTypeName, TEXT("Invalid counter access! Counter is keyed against '%s' but '%s' was given."), *KeyTypeName.ToString(), *InKeyTypeName.ToString());
 
 			const KeyType* KeyPtr = static_cast<const KeyType*>(InKey);
 			FCounterValue& CounterValue = Counters.FindOrAdd(*KeyPtr);
 			++CounterValue;
 		}
 
-		virtual void DecrementCounterImpl(const void* InKey, const FStringView InKeyTypeName) override
+		virtual void DecrementCounterImpl(const void* InKey, const FName InKeyTypeName) override
 		{
-			checkf(KeyTypeName == InKeyTypeName, TEXT("Invalid counter access! Counter is keyed against '%s' but '%s' was given."), *KeyTypeName, *FString(InKeyTypeName));
+			checkf(KeyTypeName == InKeyTypeName, TEXT("Invalid counter access! Counter is keyed against '%s' but '%s' was given."), *KeyTypeName.ToString(), *InKeyTypeName.ToString());
 
 			const KeyType* KeyPtr = static_cast<const KeyType*>(InKey);
 			if (FCounterValue* CounterValue = Counters.Find(*KeyPtr))
@@ -242,9 +243,9 @@ private:
 			}
 		}
 
-		virtual void ClearCounterImpl(const void* InKey, const FStringView InKeyTypeName) override
+		virtual void ClearCounterImpl(const void* InKey, const FName InKeyTypeName) override
 		{
-			checkf(KeyTypeName == InKeyTypeName, TEXT("Invalid counter access! Counter is keyed against '%s' but '%s' was given."), *KeyTypeName, *FString(InKeyTypeName));
+			checkf(KeyTypeName == InKeyTypeName, TEXT("Invalid counter access! Counter is keyed against '%s' but '%s' was given."), *KeyTypeName.ToString(), *InKeyTypeName.ToString());
 
 			const KeyType* KeyPtr = static_cast<const KeyType*>(InKey);
 			Counters.Remove(*KeyPtr);
@@ -256,7 +257,7 @@ private:
 		}
 
 	private:
-		FString KeyTypeName;
+		FName KeyTypeName;
 		TMap<KeyType, FCounterValue> Counters;
 	};
 
