Added low-level debug info to typed element reference leaks

This lets us see what kind of elements have leaked, even if full reference tracking is disabled

diff --git a/Engine/Source/Runtime/Engine/Public/Elements/Actor/ActorElementData.h b/Engine/Source/Runtime/Engine/Public/Elements/Actor/ActorElementData.h
--- a/Engine/Source/Runtime/Engine/Public/Elements/Actor/ActorElementData.h
+++ b/Engine/Source/Runtime/Engine/Public/Elements/Actor/ActorElementData.h
@@ -3,6 +3,7 @@
 #pragma once
 
 #include "Elements/Framework/TypedElementData.h"
+#include "UObject/Object.h"
 
 class AActor;
 struct FTypedElementHandle;
@@ -17,6 +18,15 @@ struct ENGINE_API FActorElementData
 	AActor* Actor = nullptr;
 };
 
+template <>
+inline FString GetTypedElementDebugId<FActorElementData>(const FActorElementData& InElementData)
+{
+	UObject* Object = (UObject*)InElementData.Actor;
+	return Object
+		? Object->GetFullName()
+		: TEXT("null");
+}
+
 namespace ActorElementDataUtil
 {
 
diff --git a/Engine/Source/Runtime/Engine/Public/Elements/Component/ComponentElementData.h b/Engine/Source/Runtime/Engine/Public/Elements/Component/ComponentElementData.h
--- a/Engine/Source/Runtime/Engine/Public/Elements/Component/ComponentElementData.h
+++ b/Engine/Source/Runtime/Engine/Public/Elements/Component/ComponentElementData.h
@@ -3,6 +3,7 @@
 #pragma once
 
 #include "Elements/Framework/TypedElementData.h"
+#include "UObject/Object.h"
 
 class UActorComponent;
 struct FTypedElementHandle;
@@ -17,6 +18,15 @@ struct ENGINE_API FComponentElementData
 	UActorComponent* Component = nullptr;
 };
 
+template <>
+inline FString GetTypedElementDebugId<FComponentElementData>(const FComponentElementData& InElementData)
+{
+	UObject* Object = (UObject*)InElementData.Component;
+	return Object
+		? Object->GetFullName()
+		: TEXT("null");
+}
+
 namespace ComponentElementDataUtil
 {
 
diff --git a/Engine/Source/Runtime/Engine/Public/Elements/Object/ObjectElementData.h b/Engine/Source/Runtime/Engine/Public/Elements/Object/ObjectElementData.h
--- a/Engine/Source/Runtime/Engine/Public/Elements/Object/ObjectElementData.h
+++ b/Engine/Source/Runtime/Engine/Public/Elements/Object/ObjectElementData.h
@@ -3,8 +3,8 @@
 #pragma once
 
 #include "Elements/Framework/TypedElementData.h"
+#include "UObject/Object.h"
 
-class UObject;
 struct FTypedElementHandle;
 
 /**
@@ -17,6 +17,15 @@ struct ENGINE_API FObjectElementData
 	UObject* Object = nullptr;
 };
 
+template <>
+inline FString GetTypedElementDebugId<FObjectElementData>(const FObjectElementData& InElementData)
+{
+	UObject* Object = InElementData.Object;
+	return Object
+		? Object->GetFullName()
+		: TEXT("null");
+}
+
 namespace ObjectElementDataUtil
 {
 
diff --git a/Engine/Source/Runtime/Engine/Public/Elements/SMInstance/SMInstanceElementData.h b/Engine/Source/Runtime/Engine/Public/Elements/SMInstance/SMInstanceElementData.h
--- a/Engine/Source/Runtime/Engine/Public/Elements/SMInstance/SMInstanceElementData.h
+++ b/Engine/Source/Runtime/Engine/Public/Elements/SMInstance/SMInstanceElementData.h
@@ -4,6 +4,7 @@
 
 #include "Elements/Framework/TypedElementData.h"
 #include "Elements/SMInstance/SMInstanceManager.h"
+#include "UObject/Object.h"
 
 struct FTypedElementHandle;
 
@@ -17,6 +18,15 @@ struct ENGINE_API FSMInstanceElementData
 	FSMInstanceElementId InstanceElementId;
 };
 
+template <>
+inline FString GetTypedElementDebugId<FSMInstanceElementData>(const FSMInstanceElementData& InElementData)
+{
+	UObject* Object = (UObject*)InElementData.InstanceElementId.ISMComponent;
+	return Object
+		? Object->GetFullName()
+		: TEXT("null");
+}
+
 namespace SMInstanceElementDataUtil
 {
 
diff --git a/Engine/Source/Runtime/TypedElementFramework/Public/Elements/Framework/TypedElementData.h b/Engine/Source/Runtime/TypedElementFramework/Public/Elements/Framework/TypedElementData.h
--- a/Engine/Source/Runtime/TypedElementFramework/Public/Elements/Framework/TypedElementData.h
+++ b/Engine/Source/Runtime/TypedElementFramework/Public/Elements/Framework/TypedElementData.h
@@ -26,6 +26,16 @@
 #define UE_DEFINE_TYPED_ELEMENT_DATA_RTTI(ELEMENT_DATA_TYPE)						\
 	FTypedHandleTypeId ELEMENT_DATA_TYPE::Private_RegisteredTypeId = 0;
 
+/**
+ * Templated util to get the low-level debug ID for an element data instance.
+ * Specialize this to provide exta/custom data for your element data type.
+ */
+template <typename ElementDataType>
+inline FString GetTypedElementDebugId(const ElementDataType& InElementData)
+{
+	return FString();
+}
+
 #if UE_TYPED_ELEMENT_HAS_REFTRACKING
 /**
  * Debugging information used to locate reference leaks.
@@ -263,7 +273,7 @@ public:
 		if (LocalRefCount > 1)
 		{
 			LogReferences();
-			UE_LOG(LogCore, Fatal, TEXT("Element is still externally referenced when being destroyed! Ref-count: %d; see above for reference information (if available)."), LocalRefCount);
+			UE_LOG(LogCore, Fatal, TEXT("Element '%s' is still externally referenced when being destroyed! Ref-count: %d; see above for reference information (if available)."), *GetDebugId(), LocalRefCount);
 		}
 #endif	// DO_CHECK
 	}
@@ -273,6 +283,11 @@ public:
 		return nullptr;
 	}
 
+	virtual FString GetDebugId() const
+	{
+		return FString::Printf(TEXT("ID: %d"), Id.GetElementId());
+	}
+
 private:
 	FTypedElementId Id;
 #if UE_TYPED_ELEMENT_HAS_REFCOUNTING
@@ -311,6 +326,14 @@ public:
 		return &Data;
 	}
 
+	virtual FString GetDebugId() const override
+	{
+		const FString ElementDebugId = GetTypedElementDebugId(Data);
+		return ElementDebugId.IsEmpty()
+			? FTypedElementInternalData::GetDebugId()
+			: FString::Printf(TEXT("%s - %s"), *ElementDebugId, *FTypedElementInternalData::GetDebugId());
+	}
+
 private:
 	ElementDataType Data;
 };
