Fixed a crash when asking for a TypedElementHandle with a interface that is not suported for that element. (Added some nullptr validation to the function involved)

Added a step to theTyped Element smoke test that can reproduce the issue.

diff --git a/Engine/Source/Runtime/TypedElementFramework/Private/TypedElementFrameworkTests.cpp b/Engine/Source/Runtime/TypedElementFramework/Private/TypedElementFrameworkTests.cpp
--- a/Engine/Source/Runtime/TypedElementFramework/Private/TypedElementFrameworkTests.cpp
+++ b/Engine/Source/Runtime/TypedElementFramework/Private/TypedElementFrameworkTests.cpp
@@ -188,6 +188,13 @@ bool FTypedElementRegistrySmokeTest::RunTest(const FString& Parameters)
 		return true;
 	});
 
+	// Test some elements with an interface they don't have
+	ElementList->ForEachElement<ITestTypedElementInterfaceC>([&TestInterfaceCHandle](const TTypedElement<ITestTypedElementInterfaceC>& InElementHandle)
+	{
+		TestInterfaceCHandle(InElementHandle);
+		return true;
+	});
+
 	ElementList.Reset();
 
 	Registry->DestroyElement(TypedElement1);
diff --git a/Engine/Source/Runtime/TypedElementFramework/Public/Elements/Framework/TypedElementRegistry.h b/Engine/Source/Runtime/TypedElementFramework/Public/Elements/Framework/TypedElementRegistry.h
--- a/Engine/Source/Runtime/TypedElementFramework/Public/Elements/Framework/TypedElementRegistry.h
+++ b/Engine/Source/Runtime/TypedElementFramework/Public/Elements/Framework/TypedElementRegistry.h
@@ -448,15 +448,17 @@ private:
 			FRegisteredElementType* RegisteredElementType = GetRegisteredElementTypeFromId(InElementId.GetTypeId());
 			checkf(RegisteredElementType, TEXT("Element type ID '%d' has not been registered!"), InElementId.GetTypeId());
 
-			UObject* InterfaceObject = RegisteredElementType->Interfaces.FindRef(InBaseInterfaceType->GetFName());
-			const FTypedElementInternalData& TypedElementInternalData = RegisteredElementType->GetDataForElement(InElementId.GetElementId());
-			if constexpr (std::is_void<BaseInterfaceType>::value)
-			{ 
-				OutElement.Private_InitializeAddRef(TypedElementInternalData, InterfaceObject->GetInterfaceAddress(InBaseInterfaceType));
-			}
-			else
+			if (UObject* InterfaceObject = RegisteredElementType->Interfaces.FindRef(InBaseInterfaceType->GetFName()))
 			{
-				OutElement.Private_InitializeAddRef(TypedElementInternalData, Cast<BaseInterfaceType>(InterfaceObject));
+				const FTypedElementInternalData& TypedElementInternalData = RegisteredElementType->GetDataForElement(InElementId.GetElementId());
+				if constexpr (std::is_void<BaseInterfaceType>::value)
+				{ 
+					OutElement.Private_InitializeAddRef(TypedElementInternalData, InterfaceObject->GetInterfaceAddress(InBaseInterfaceType));
+				}
+				else
+				{
+					OutElement.Private_InitializeAddRef(TypedElementInternalData, Cast<BaseInterfaceType>(InterfaceObject));
+				}
 			}
 		}
 	}
@@ -471,15 +473,17 @@ private:
 			FRegisteredElementType* RegisteredElementType = GetRegisteredElementTypeFromId(InElementHandle.GetId().GetTypeId());
 			checkf(RegisteredElementType, TEXT("Element type ID '%d' has not been registered!"), InElementHandle.GetId().GetTypeId());
 
-			UObject* InterfaceObject = RegisteredElementType->Interfaces.FindRef(InBaseInterfaceType->GetFName());
-			const FTypedElementInternalData& TypedElementInternalData = *InElementHandle.Private_GetInternalData();
-			if constexpr (std::is_void<BaseInterfaceType>::value)
-			{
-				OutElement.Private_InitializeAddRef(TypedElementInternalData, InterfaceObject->GetInterfaceAddress(InBaseInterfaceType));
-			}
-			else
+			if (UObject* InterfaceObject = RegisteredElementType->Interfaces.FindRef(InBaseInterfaceType->GetFName()))
 			{
-				OutElement.Private_InitializeAddRef(TypedElementInternalData, Cast<BaseInterfaceType>(InterfaceObject));
+				const FTypedElementInternalData& TypedElementInternalData = *InElementHandle.Private_GetInternalData();
+				if constexpr (std::is_void<BaseInterfaceType>::value)
+				{
+					OutElement.Private_InitializeAddRef(TypedElementInternalData, InterfaceObject->GetInterfaceAddress(InBaseInterfaceType));
+				}
+				else
+				{
+					OutElement.Private_InitializeAddRef(TypedElementInternalData, Cast<BaseInterfaceType>(InterfaceObject));
+				}
 			}
 		}
 	}
