Added a commend to copy some debug info on the typed element registry into the clipboard.

diff --git a/Engine/Source/Editor/UnrealEd/Private/Elements/Framework/EditorElements.cpp b/Engine/Source/Editor/UnrealEd/Private/Elements/Framework/EditorElements.cpp
--- a/Engine/Source/Editor/UnrealEd/Private/Elements/Framework/EditorElements.cpp
+++ b/Engine/Source/Editor/UnrealEd/Private/Elements/Framework/EditorElements.cpp
@@ -24,6 +24,11 @@
 #include "Elements/SMInstance/SMInstanceElementEditorSelectionInterface.h"
 #include "Elements/SMInstance/SMInstanceElementDetailsProxyObject.h"
 
+#include "Elements/TypedElementEditorLog.h"
+
+#include "HAL/IConsoleManager.h"
+#include "HAL/PlatformApplicationMisc.h"
+
 #include "Modules/ModuleManager.h"
 #include "EditorWidgetsModule.h"
 #include "ObjectNameEditSinkRegistry.h"
@@ -76,5 +81,18 @@ void RegisterEditorElements()
 	RegisterEditorComponentElements();
 	RegisterEditorSMInstanceElements();
 
+	IConsoleManager::Get().RegisterConsoleCommand(
+		TEXT("TypedElements.OutputRegistredTypeElementsToClipboard"),
+		TEXT("Output a debug string to the clipboard and to the log./n\
+			It contains the names of the Typed Elements registred and their Interfaces./n\
+			For each Interface it also provide the path of the class that implements it."),
+			FConsoleCommandDelegate::CreateLambda([]()
+				{
+					FString DebugString = UTypedElementRegistry::GetInstance()->RegistredElementTypesAndInterfacesToString();
+					FPlatformApplicationMisc::ClipboardCopy(*DebugString);
+					UE_LOG(LogTypedElementEditor, Log, TEXT("%s"), *DebugString);
+				})
+		);
+
 	OnRegisterEditorElementsDelegate.Broadcast();
 }
diff --git a/Engine/Source/Editor/UnrealEd/Public/Elements/TypedElementEditorLog.cpp b/Engine/Source/Editor/UnrealEd/Public/Elements/TypedElementEditorLog.cpp
new file mode 100644
--- /dev/null
+++ b/Engine/Source/Editor/UnrealEd/Public/Elements/TypedElementEditorLog.cpp
@@ -0,0 +1,5 @@
+// Copyright Epic Games, Inc. All Rights Reserved.
+
+#include "TypedElementEditorLog.h"
+
+DEFINE_LOG_CATEGORY(LogTypedElementEditor)
\ No newline at end of file
diff --git a/Engine/Source/Editor/UnrealEd/Public/Elements/TypedElementEditorLog.h b/Engine/Source/Editor/UnrealEd/Public/Elements/TypedElementEditorLog.h
new file mode 100644
--- /dev/null
+++ b/Engine/Source/Editor/UnrealEd/Public/Elements/TypedElementEditorLog.h
@@ -0,0 +1,7 @@
+// Copyright Epic Games, Inc. All Rights Reserved.
+
+#pragma once
+
+#include "Logging/LogMacros.h"
+
+DECLARE_LOG_CATEGORY_EXTERN(LogTypedElementEditor, Log, All)
diff --git a/Engine/Source/Runtime/TypedElementFramework/Private/Elements/Framework/TypedElementRegistry.cpp b/Engine/Source/Runtime/TypedElementFramework/Private/Elements/Framework/TypedElementRegistry.cpp
--- a/Engine/Source/Runtime/TypedElementFramework/Private/Elements/Framework/TypedElementRegistry.cpp
+++ b/Engine/Source/Runtime/TypedElementFramework/Private/Elements/Framework/TypedElementRegistry.cpp
@@ -1,9 +1,9 @@
 // Copyright Epic Games, Inc. All Rights Reserved.
 
 #include "Elements/Framework/TypedElementRegistry.h"
+
 #include "Misc/CoreDelegates.h"
 #include "Misc/ScopeLock.h"
-#include "HAL/IConsoleManager.h"
 #include "UObject/StrongObjectPtr.h"
 
 const FTypedElementId FTypedElementId::Unset;
@@ -209,6 +209,62 @@ FTypedElementListRef UTypedElementRegistry::CreateElementList(TArrayView<const F
 	return ElementList;
 }
 
+FString UTypedElementRegistry::RegistredElementTypesAndInterfacesToString() const
+{
+	TArray<FString> Lines;
+
+	{
+		FReadScopeLock RegisteredElementTypesLock(RegisteredElementTypesRW);
+		Lines.Reserve(RegisteredElementTypesNameToId.Num());
+
+		{
+			FStringFormatOrderedArguments FormatArguments;
+			FormatArguments.Add(GetPathName());
+			Lines.Add(FString::Format(TEXT("Registred Type Elements and their Interfaces for the registry: {0}"), FormatArguments));
+		}
+
+		for (const TPair<FName, FTypedHandleTypeId>& TypePair : RegisteredElementTypesNameToId)
+		{
+			if (FRegisteredElementType* RegistredElementType = GetRegisteredElementTypeFromId(TypePair.Value))
+			{
+				Lines.Reserve(RegistredElementType->Interfaces.Num() + 1 + Lines.Num());
+				{
+					FStringFormatOrderedArguments FormatArguments;
+					FormatArguments.Add(TypePair.Key.ToString());
+					Lines.Add(FString::Format(TEXT("	Type: {0}"), FormatArguments));
+				}
+				for (const TPair<FName, UObject*>& InterfacePair : RegistredElementType->Interfaces)
+				{
+					FStringFormatOrderedArguments FormatArguments;
+					FormatArguments.Reserve(2);
+					FormatArguments.Add(InterfacePair.Key.ToString());
+					FormatArguments.Add(InterfacePair.Value->GetClass()->GetPathName());
+					Lines.Add(FString::Format(TEXT("		Interface {0} is implemented by {1}"), FormatArguments));
+				}
+			}
+		}
+
+		FString Output;
+		int32 OutputSize = 0;
+		for (const FString& Line : Lines)
+		{
+			OutputSize += Line.Len();
+			// for the line termination char
+			++OutputSize;
+		}
+
+		Output.Reserve(OutputSize);
+		for (const FString& Line : Lines)
+		{
+			Output += Line;
+			Output.AppendChar(TEXT('\n'));
+		}
+
+		return Output;
+	}
+
+}
+
 void UTypedElementRegistry::NotifyElementListPendingChanges()
 {
 	/** 
diff --git a/Engine/Source/Runtime/TypedElementFramework/Public/Elements/Framework/TypedElementRegistry.h b/Engine/Source/Runtime/TypedElementFramework/Public/Elements/Framework/TypedElementRegistry.h
--- a/Engine/Source/Runtime/TypedElementFramework/Public/Elements/Framework/TypedElementRegistry.h
+++ b/Engine/Source/Runtime/TypedElementFramework/Public/Elements/Framework/TypedElementRegistry.h
@@ -283,6 +283,12 @@ public:
 		return ElementList;
 	}
 
+	/**
+	 * Extract a string that contains name of the type elements registered.
+	 * For each type it also says what interface are registered and the path of the class that implement them
+	 */
+	FString RegistredElementTypesAndInterfacesToString() const;
+
 	void Private_OnElementListCreated(FTypedElementList* InElementList)
 	{
 		FScopeLock ActiveElementListsLock(&ActiveElementListsCS);
