Fix a crash in typed element get bottom element from passing the wrong element handle to the predicate

#rnx

diff --git a/Engine/Source/Runtime/TypedElementFramework/Public/Elements/Framework/TypedElementList.h b/Engine/Source/Runtime/TypedElementFramework/Public/Elements/Framework/TypedElementList.h
--- a/Engine/Source/Runtime/TypedElementFramework/Public/Elements/Framework/TypedElementList.h
+++ b/Engine/Source/Runtime/TypedElementFramework/Public/Elements/Framework/TypedElementList.h
@@ -270,7 +270,7 @@ public:
 		for (int32 ElementIndex = Num() - 1; ElementIndex >= 0; --ElementIndex)
 		{
 			GetElementAt(ElementIndex, TempElement);
-			if (TempElement && Predicate(ElementToReturn))
+			if (TempElement && Predicate(TempElement))
 			{
 				ElementToReturn = MoveTemp(TempElement);
 				break;
