Fixed the typed element scripting api issues.

Added the FScriptTypedElementHandle. These handles are disarmed when the element they pointing to is destroyed instead of crashing the engine. They do have a performance overhead so using these should be restricted to exposing stuff to blueprint/python.

Reworked the TypedElementList to be a template so that we have both a FTypedElementList and a FScriptTypedElementList from the same source code.

Changed the api of the interfaces so that can now accept and use the scripted version of the handle and list instead of the native ones.

diff --git a/Engine/Source/Editor/LevelEditor/Private/Elements/Actor/ActorElementLevelEditorSelectionCustomization.cpp b/Engine/Source/Editor/LevelEditor/Private/Elements/Actor/ActorElementLevelEditorSelectionCustomization.cpp
--- a/Engine/Source/Editor/LevelEditor/Private/Elements/Actor/ActorElementLevelEditorSelectionCustomization.cpp
+++ b/Engine/Source/Editor/LevelEditor/Private/Elements/Actor/ActorElementLevelEditorSelectionCustomization.cpp
@@ -256,7 +256,7 @@ bool FActorElementLevelEditorSelectionCustomization::DeselectActorElement(const
 	
 	// Deselect and unbind the override delegates for the components on the selected actor
 	{
-		FTypedElementListLegacySyncScopedBatch LegacySyncBatch(*InSelectionSet, InSelectionOptions.AllowLegacyNotifications());
+		FTypedElementList::FLegacySyncScopedBatch LegacySyncBatch(*InSelectionSet, InSelectionOptions.AllowLegacyNotifications());
 
 		for (UActorComponent* Component : Actor->GetComponents())
 		{
@@ -297,7 +297,7 @@ bool FActorElementLevelEditorSelectionCustomization::SelectActorGroup(AGroupActo
 	// Skip if the group is already selected, since this logic will have already run
 	if ((bForce || InGroupActor->IsLocked()) && !GroupSelectionHandle.IsElementSelected(InSelectionSet, FTypedElementIsSelectedOptions()) && CanSelectActorElement(GroupSelectionHandle, GroupSelectionOptions))
 	{
-		FTypedElementListLegacySyncScopedBatch LegacySyncBatch(*InSelectionSet, InSelectionOptions.AllowLegacyNotifications());
+		FTypedElementList::FLegacySyncScopedBatch LegacySyncBatch(*InSelectionSet, InSelectionOptions.AllowLegacyNotifications());
 
 		bSelectionChanged |= SelectActorElement(GroupSelectionHandle, InSelectionSet, GroupSelectionOptions);
 
@@ -329,7 +329,7 @@ bool FActorElementLevelEditorSelectionCustomization::DeselectActorGroup(AGroupAc
 	// Skip if the group is already deselected, since this logic will have already run
 	if ((bForce || InGroupActor->IsLocked()) && GroupSelectionHandle.IsElementSelected(InSelectionSet, FTypedElementIsSelectedOptions()) && CanDeselectActorElement(GroupSelectionHandle, GroupSelectionOptions))
 	{
-		FTypedElementListLegacySyncScopedBatch LegacySyncBatch(*InSelectionSet, InSelectionOptions.AllowLegacyNotifications());
+		FTypedElementList::FLegacySyncScopedBatch LegacySyncBatch(*InSelectionSet, InSelectionOptions.AllowLegacyNotifications());
 
 		bSelectionChanged |= DeselectActorElement(GroupSelectionHandle, InSelectionSet, GroupSelectionOptions);
 
diff --git a/Engine/Source/Editor/LevelEditor/Private/LevelEditorMenuContext.cpp b/Engine/Source/Editor/LevelEditor/Private/LevelEditorMenuContext.cpp
new file mode 100644
--- /dev/null
+++ b/Engine/Source/Editor/LevelEditor/Private/LevelEditorMenuContext.cpp
@@ -0,0 +1,10 @@
+// Copyright Epic Games, Inc. All Rights Reserved.
+
+#include "LevelEditorMenuContext.h"
+
+#include "Elements/Framework/TypedElementRegistry.h"
+
+FScriptTypedElementHandle ULevelEditorContextMenuContext::GetScriptHitProxyElement()
+{
+	return UTypedElementRegistry::GetInstance()->CreateScriptHandle(HitProxyElement.GetId());
+}
\ No newline at end of file
diff --git a/Engine/Source/Editor/LevelEditor/Public/LevelEditorMenuContext.h b/Engine/Source/Editor/LevelEditor/Public/LevelEditorMenuContext.h
--- a/Engine/Source/Editor/LevelEditor/Public/LevelEditorMenuContext.h
+++ b/Engine/Source/Editor/LevelEditor/Public/LevelEditorMenuContext.h
@@ -42,7 +42,6 @@ class LEVELEDITOR_API ULevelEditorContextMenuContext : public UObject
 {
 	GENERATED_BODY()
 public:
-
 	TWeakPtr<ILevelEditor> LevelEditor;
 
 	UPROPERTY(VisibleAnywhere, BlueprintReadOnly, Category="LevelEditor|Menu")
@@ -52,7 +51,6 @@ public:
 	TObjectPtr<const UTypedElementSelectionSet> CurrentSelection;
 	
 	/** If the ContextType is Viewport this property can be set to the HitProxy element that triggered the ContextMenu. */
-	UPROPERTY(VisibleAnywhere, BlueprintReadOnly, Category="LevelEditor|Menu")
 	FTypedElementHandle HitProxyElement;
 
 	UPROPERTY(VisibleAnywhere, BlueprintReadOnly, Category="LevelEditor|Menu")
@@ -64,6 +62,9 @@ public:
 	/** If the ContextType is Viewport this property can be set to the HitProxy actor that triggered the ContextMenu. */
 	UPROPERTY(VisibleAnywhere, BlueprintReadOnly, Category="LevelEditor|Menu")
 	TObjectPtr<AActor> HitProxyActor = nullptr;
+
+	UFUNCTION(BlueprintCallable, Category = "LevelEditor | Menu", DisplayName="Get Hit Proxy Element", meta=(ScriptName="GetHitProxyElement"))
+	FScriptTypedElementHandle GetScriptHitProxyElement();
 };
 
 UCLASS()
diff --git a/Engine/Source/Editor/UnrealEd/Private/EditorActor.cpp b/Engine/Source/Editor/UnrealEd/Private/EditorActor.cpp
--- a/Engine/Source/Editor/UnrealEd/Private/EditorActor.cpp
+++ b/Engine/Source/Editor/UnrealEd/Private/EditorActor.cpp
@@ -795,7 +795,7 @@ bool UUnrealEdEngine::DeleteComponents(const TArray<UActorComponent*>& InCompone
 		return false;
 	}
 
-	TUniquePtr<FTypedElementListLegacySyncScopedBatch> LegacySyncBatch = MakeUnique<FTypedElementListLegacySyncScopedBatch>(*InSelectionSet->GetElementList());
+	TUniquePtr<FTypedElementList::FLegacySyncScopedBatch> LegacySyncBatch = MakeUnique<FTypedElementList::FLegacySyncScopedBatch>(*InSelectionSet->GetElementList());
 
 	const FTypedElementSelectionOptions SelectionOptions = FTypedElementSelectionOptions()
 		.SetAllowHidden(true)
@@ -927,7 +927,7 @@ bool UUnrealEdEngine::DeleteActors(const TArray<AActor*>& InActorsToDelete, UWor
 	bool bRequestedDeleteAllBySoftReference = false;
 	int32 DeleteCount = 0;
 
-	TUniquePtr<FTypedElementListLegacySyncScopedBatch> LegacySyncBatch = MakeUnique<FTypedElementListLegacySyncScopedBatch>(*InSelectionSet->GetElementList());
+	TUniquePtr<FTypedElementList::FLegacySyncScopedBatch> LegacySyncBatch = MakeUnique<FTypedElementList::FLegacySyncScopedBatch>(*InSelectionSet->GetElementList());
 
 	const FTypedElementSelectionOptions SelectionOptions = FTypedElementSelectionOptions()
 		.SetAllowHidden(true)
diff --git a/Engine/Source/Editor/UnrealEd/Private/Elements/Actor/ActorElementEditorSelectionInterface.cpp b/Engine/Source/Editor/UnrealEd/Private/Elements/Actor/ActorElementEditorSelectionInterface.cpp
--- a/Engine/Source/Editor/UnrealEd/Private/Elements/Actor/ActorElementEditorSelectionInterface.cpp
+++ b/Engine/Source/Editor/UnrealEd/Private/Elements/Actor/ActorElementEditorSelectionInterface.cpp
@@ -38,9 +38,8 @@ private:
 	TWeakObjectPtr<AActor> ActorPtr;
 };
 
-bool UActorElementEditorSelectionInterface::IsElementSelected(const FTypedElementHandle& InElementHandle, const FTypedElementListProxy InSelectionSet, const FTypedElementIsSelectedOptions& InSelectionOptions)
+bool UActorElementEditorSelectionInterface::IsElementSelected(const FTypedElementHandle& InElementHandle, const FTypedElementListConstPtr& SelectionSetPtr, const FTypedElementIsSelectedOptions& InSelectionOptions)
 {
-	FTypedElementListConstPtr SelectionSetPtr = InSelectionSet.GetElementList();
 	const AActor* Actor = ActorElementDataUtil::GetActorFromHandle(InElementHandle);
 	return SelectionSetPtr && Actor && IsActorSelected(Actor, SelectionSetPtr.ToSharedRef(), InSelectionOptions);
 }
diff --git a/Engine/Source/Editor/UnrealEd/Private/Elements/Component/ComponentElementEditorSelectionInterface.cpp b/Engine/Source/Editor/UnrealEd/Private/Elements/Component/ComponentElementEditorSelectionInterface.cpp
--- a/Engine/Source/Editor/UnrealEd/Private/Elements/Component/ComponentElementEditorSelectionInterface.cpp
+++ b/Engine/Source/Editor/UnrealEd/Private/Elements/Component/ComponentElementEditorSelectionInterface.cpp
@@ -39,9 +39,8 @@ private:
 	TWeakObjectPtr<UActorComponent> ComponentPtr;
 };
 
-bool UComponentElementEditorSelectionInterface::IsElementSelected(const FTypedElementHandle& InElementHandle, const FTypedElementListProxy InSelectionSet, const FTypedElementIsSelectedOptions& InSelectionOptions)
+bool UComponentElementEditorSelectionInterface::IsElementSelected(const FTypedElementHandle& InElementHandle, const FTypedElementListConstPtr& SelectionSetPtr, const FTypedElementIsSelectedOptions& InSelectionOptions)
 {
-	FTypedElementListConstPtr SelectionSetPtr = InSelectionSet.GetElementList();
 	const UActorComponent* Component = ComponentElementDataUtil::GetComponentFromHandle(InElementHandle);
 	return SelectionSetPtr && Component && IsComponentSelected(Component, SelectionSetPtr.ToSharedRef(), InSelectionOptions);
 }
diff --git a/Engine/Source/Editor/UnrealEd/Private/Elements/SMInstance/SMInstanceElementEditorSelectionInterface.cpp b/Engine/Source/Editor/UnrealEd/Private/Elements/SMInstance/SMInstanceElementEditorSelectionInterface.cpp
--- a/Engine/Source/Editor/UnrealEd/Private/Elements/SMInstance/SMInstanceElementEditorSelectionInterface.cpp
+++ b/Engine/Source/Editor/UnrealEd/Private/Elements/SMInstance/SMInstanceElementEditorSelectionInterface.cpp
@@ -41,9 +41,8 @@ private:
 	uint64 InstanceId = 0;
 };
 
-bool USMInstanceElementEditorSelectionInterface::IsElementSelected(const FTypedElementHandle& InElementHandle, const FTypedElementListProxy InSelectionSet, const FTypedElementIsSelectedOptions& InSelectionOptions)
+bool USMInstanceElementEditorSelectionInterface::IsElementSelected(const FTypedElementHandle& InElementHandle, const FTypedElementListConstPtr& SelectionSetPtr, const FTypedElementIsSelectedOptions& InSelectionOptions)
 {
-	FTypedElementListConstPtr SelectionSetPtr = InSelectionSet.GetElementList();
 	const FSMInstanceManager SMInstance = SMInstanceElementDataUtil::GetSMInstanceFromHandle(InElementHandle);
 
 	if (SMInstance && SelectionSetPtr && SelectionSetPtr->Num() > 0)
@@ -74,9 +73,8 @@ bool USMInstanceElementEditorSelectionInterface::IsElementSelected(const FTypedE
 	return false;
 }
 
-bool USMInstanceElementEditorSelectionInterface::SelectElement(const FTypedElementHandle& InElementHandle, FTypedElementListProxy InSelectionSet, const FTypedElementSelectionOptions& InSelectionOptions)
+bool USMInstanceElementEditorSelectionInterface::SelectElement(const FTypedElementHandle& InElementHandle, const FTypedElementListPtr& SelectionSetPtr, const FTypedElementSelectionOptions& InSelectionOptions)
 {
-	FTypedElementListPtr SelectionSetPtr = InSelectionSet.GetElementList();
 	FSMInstanceManager SMInstance = SMInstanceElementDataUtil::GetSMInstanceFromHandle(InElementHandle);
 
 	if (SMInstance && SelectionSetPtr && SelectionSetPtr->Num() > 0)
@@ -109,9 +107,8 @@ bool USMInstanceElementEditorSelectionInterface::SelectElement(const FTypedEleme
 	return false;
 }
 
-bool USMInstanceElementEditorSelectionInterface::DeselectElement(const FTypedElementHandle& InElementHandle, FTypedElementListProxy InSelectionSet, const FTypedElementSelectionOptions& InSelectionOptions)
+bool USMInstanceElementEditorSelectionInterface::DeselectElement(const FTypedElementHandle& InElementHandle, const FTypedElementListPtr& SelectionSetPtr, const FTypedElementSelectionOptions& InSelectionOptions)
 {
-	FTypedElementListPtr SelectionSetPtr = InSelectionSet.GetElementList();
 	FSMInstanceManager SMInstance = SMInstanceElementDataUtil::GetSMInstanceFromHandle(InElementHandle);
 
 	if (SMInstance && SelectionSetPtr && SelectionSetPtr->Num() > 0)
diff --git a/Engine/Source/Editor/UnrealEd/Private/Elements/SMInstance/SMInstanceElementEditorWorldInterface.cpp b/Engine/Source/Editor/UnrealEd/Private/Elements/SMInstance/SMInstanceElementEditorWorldInterface.cpp
--- a/Engine/Source/Editor/UnrealEd/Private/Elements/SMInstance/SMInstanceElementEditorWorldInterface.cpp
+++ b/Engine/Source/Editor/UnrealEd/Private/Elements/SMInstance/SMInstanceElementEditorWorldInterface.cpp
@@ -46,7 +46,7 @@ bool USMInstanceElementEditorWorldInterface::DeleteElements(TArrayView<const FTy
 
 		bool bDidDelete = false;
 
-		FTypedElementListLegacySyncScopedBatch LegacySyncBatch(*InSelectionSet->GetElementList());
+		FTypedElementList::FLegacySyncScopedBatch LegacySyncBatch(*InSelectionSet->GetElementList());
 		for (TTuple<ISMInstanceManager*, TArray<FSMInstanceId>>& BatchedInstancesToDeletePair : BatchedInstancesToDelete)
 		{
 			for (const FSMInstanceId& InstanceId : BatchedInstancesToDeletePair.Value)
diff --git a/Engine/Source/Editor/UnrealEd/Private/Selection.cpp b/Engine/Source/Editor/UnrealEd/Private/Selection.cpp
--- a/Engine/Source/Editor/UnrealEd/Private/Selection.cpp
+++ b/Engine/Source/Editor/UnrealEd/Private/Selection.cpp
@@ -198,7 +198,7 @@ UObject* USelection::GetObjectForElementHandle(const FTypedElementHandle& InElem
 	return nullptr;
 }
 
-void USelection::OnElementListSyncEvent(const FTypedElementList& InElementList, FTypedElementListLegacySync::ESyncType InSyncType, const FTypedElementHandle& InElementHandle, bool bIsWithinBatchOperation)
+void USelection::OnElementListSyncEvent(const FTypedElementList& InElementList, FTypedElementList::FLegacySync::ESyncType InSyncType, const FTypedElementHandle& InElementHandle, bool bIsWithinBatchOperation)
 {
 	check(&InElementList == &ElementSelectionSet->GetElementList().Get());
 
@@ -208,8 +208,8 @@ void USelection::OnElementListSyncEvent(const FTypedElementList& InElementList,
 		switch (InSyncType)
 		{
 		// Emit a general object selection changed notification for the following events
-		case FTypedElementListLegacySync::ESyncType::Added:
-		case FTypedElementListLegacySync::ESyncType::Removed:
+		case FTypedElementList::FLegacySync::ESyncType::Added:
+		case FTypedElementList::FLegacySync::ESyncType::Removed:
 			if (UObject* Object = GetObjectForElementHandle(InElementHandle))
 			{
 				USelection::SelectObjectEvent.Broadcast(Object);
@@ -217,9 +217,9 @@ void USelection::OnElementListSyncEvent(const FTypedElementList& InElementList,
 			break;
 
 		// Emit a general selection changed notification for the following events
-		case FTypedElementListLegacySync::ESyncType::Modified:
-		case FTypedElementListLegacySync::ESyncType::Cleared:
-		case FTypedElementListLegacySync::ESyncType::BatchComplete:
+		case FTypedElementList::FLegacySync::ESyncType::Modified:
+		case FTypedElementList::FLegacySync::ESyncType::Cleared:
+		case FTypedElementList::FLegacySync::ESyncType::BatchComplete:
 			USelection::SelectionChangedEvent.Broadcast(this);
 			break;
 
diff --git a/Engine/Source/Editor/UnrealEd/Public/Elements/Actor/ActorElementEditorSelectionInterface.h b/Engine/Source/Editor/UnrealEd/Public/Elements/Actor/ActorElementEditorSelectionInterface.h
--- a/Engine/Source/Editor/UnrealEd/Public/Elements/Actor/ActorElementEditorSelectionInterface.h
+++ b/Engine/Source/Editor/UnrealEd/Public/Elements/Actor/ActorElementEditorSelectionInterface.h
@@ -13,7 +13,7 @@ class UNREALED_API UActorElementEditorSelectionInterface : public UActorElementS
 	GENERATED_BODY()
 
 public:
-	virtual bool IsElementSelected(const FTypedElementHandle& InElementHandle, const FTypedElementListProxy InSelectionSet, const FTypedElementIsSelectedOptions& InSelectionOptions) override;
+	virtual bool IsElementSelected(const FTypedElementHandle& InElementHandle, const FTypedElementListConstPtr& SelectionSetPtr, const FTypedElementIsSelectedOptions& InSelectionOptions) override;
 	virtual bool ShouldPreventTransactions(const FTypedElementHandle& InElementHandle) override;
 	virtual TUniquePtr<ITypedElementTransactedElement> CreateTransactedElementImpl() override;
 
diff --git a/Engine/Source/Editor/UnrealEd/Public/Elements/Component/ComponentElementEditorSelectionInterface.h b/Engine/Source/Editor/UnrealEd/Public/Elements/Component/ComponentElementEditorSelectionInterface.h
--- a/Engine/Source/Editor/UnrealEd/Public/Elements/Component/ComponentElementEditorSelectionInterface.h
+++ b/Engine/Source/Editor/UnrealEd/Public/Elements/Component/ComponentElementEditorSelectionInterface.h
@@ -13,7 +13,7 @@ class UNREALED_API UComponentElementEditorSelectionInterface : public UComponent
 	GENERATED_BODY()
 
 public:
-	virtual bool IsElementSelected(const FTypedElementHandle& InElementHandle, const FTypedElementListProxy InSelectionSet, const FTypedElementIsSelectedOptions& InSelectionOptions) override;
+	virtual bool IsElementSelected(const FTypedElementHandle& InElementHandle, const FTypedElementListConstPtr& InSelectionSet, const FTypedElementIsSelectedOptions& InSelectionOptions) override;
 	virtual bool ShouldPreventTransactions(const FTypedElementHandle& InElementHandle) override;
 	virtual TUniquePtr<ITypedElementTransactedElement> CreateTransactedElementImpl() override;
 
diff --git a/Engine/Source/Editor/UnrealEd/Public/Elements/SMInstance/SMInstanceElementEditorSelectionInterface.h b/Engine/Source/Editor/UnrealEd/Public/Elements/SMInstance/SMInstanceElementEditorSelectionInterface.h
--- a/Engine/Source/Editor/UnrealEd/Public/Elements/SMInstance/SMInstanceElementEditorSelectionInterface.h
+++ b/Engine/Source/Editor/UnrealEd/Public/Elements/SMInstance/SMInstanceElementEditorSelectionInterface.h
@@ -11,9 +11,9 @@ class UNREALED_API USMInstanceElementEditorSelectionInterface : public USMInstan
 	GENERATED_BODY()
 
 public:
-	virtual bool IsElementSelected(const FTypedElementHandle& InElementHandle, const FTypedElementListProxy InSelectionSet, const FTypedElementIsSelectedOptions& InSelectionOptions) override;
-	virtual bool SelectElement(const FTypedElementHandle& InElementHandle, FTypedElementListProxy InSelectionSet, const FTypedElementSelectionOptions& InSelectionOptions) override;
-	virtual bool DeselectElement(const FTypedElementHandle& InElementHandle, FTypedElementListProxy InSelectionSet, const FTypedElementSelectionOptions& InSelectionOptions) override;
+	virtual bool IsElementSelected(const FTypedElementHandle& InElementHandle, const FTypedElementListConstPtr& InSelectionSet, const FTypedElementIsSelectedOptions& InSelectionOptions) override;
+	virtual bool SelectElement(const FTypedElementHandle& InElementHandle, const FTypedElementListPtr& InSelectionSet, const FTypedElementSelectionOptions& InSelectionOptions) override;
+	virtual bool DeselectElement(const FTypedElementHandle& InElementHandle, const FTypedElementListPtr& InSelectionSet, const FTypedElementSelectionOptions& InSelectionOptions) override;
 	virtual bool ShouldPreventTransactions(const FTypedElementHandle& InElementHandle) override;
 	virtual TUniquePtr<ITypedElementTransactedElement> CreateTransactedElementImpl() override;
 };
diff --git a/Engine/Source/Editor/UnrealEd/Public/Selection.h b/Engine/Source/Editor/UnrealEd/Public/Selection.h
--- a/Engine/Source/Editor/UnrealEd/Public/Selection.h
+++ b/Engine/Source/Editor/UnrealEd/Public/Selection.h
@@ -325,7 +325,7 @@ private:
 	bool IsValidObjectToSelect(const UObject* InObject) const;
 	UObject* GetObjectForElementHandle(const FTypedElementHandle& InElementHandle) const;
 
-	void OnElementListSyncEvent(const FTypedElementList& InElementList, FTypedElementListLegacySync::ESyncType InSyncType, const FTypedElementHandle& InElementHandle, bool bIsWithinBatchOperation);
+	void OnElementListSyncEvent(const FTypedElementList& InElementList, FTypedElementList::FLegacySync::ESyncType InSyncType, const FTypedElementHandle& InElementHandle, bool bIsWithinBatchOperation);
 
 	/** Bridge from UObjects to their corresponding typed elements. */
 	TSharedPtr<ISelectionElementBridge> SelectionElementBridge;
diff --git a/Engine/Source/Runtime/Engine/Private/Elements/Actor/ActorElementHierarchyInterface.cpp b/Engine/Source/Runtime/Engine/Private/Elements/Actor/ActorElementHierarchyInterface.cpp
--- a/Engine/Source/Runtime/Engine/Private/Elements/Actor/ActorElementHierarchyInterface.cpp
+++ b/Engine/Source/Runtime/Engine/Private/Elements/Actor/ActorElementHierarchyInterface.cpp
@@ -1,7 +1,9 @@
 // Copyright Epic Games, Inc. All Rights Reserved.
 
 #include "Elements/Actor/ActorElementHierarchyInterface.h"
+
 #include "Elements/Actor/ActorElementData.h"
+#include "Elements/Framework/TypedElementRegistry.h"
 #include "GameFramework/Actor.h"
 
 #include "Elements/Framework/EngineElementsLibrary.h"
diff --git a/Engine/Source/Runtime/Engine/Private/Elements/Actor/ActorElementWorldInterface.cpp b/Engine/Source/Runtime/Engine/Private/Elements/Actor/ActorElementWorldInterface.cpp
--- a/Engine/Source/Runtime/Engine/Private/Elements/Actor/ActorElementWorldInterface.cpp
+++ b/Engine/Source/Runtime/Engine/Private/Elements/Actor/ActorElementWorldInterface.cpp
@@ -1,10 +1,12 @@
 // Copyright Epic Games, Inc. All Rights Reserved.
 
 #include "Elements/Actor/ActorElementWorldInterface.h"
-#include "Elements/Actor/ActorElementData.h"
-#include "GameFramework/Actor.h"
 
+#include "Elements/Actor/ActorElementData.h"
 #include "Elements/Component/ComponentElementData.h"
+#include "Elements/Framework/TypedElementRegistry.h"
+
+#include "GameFramework/Actor.h"
 #include "Components/PrimitiveComponent.h"
 
 bool UActorElementWorldInterface::IsTemplateElement(const FTypedElementHandle& InElementHandle)
diff --git a/Engine/Source/Runtime/Engine/Private/Elements/Component/ComponentElementHierarchyInterface.cpp b/Engine/Source/Runtime/Engine/Private/Elements/Component/ComponentElementHierarchyInterface.cpp
--- a/Engine/Source/Runtime/Engine/Private/Elements/Component/ComponentElementHierarchyInterface.cpp
+++ b/Engine/Source/Runtime/Engine/Private/Elements/Component/ComponentElementHierarchyInterface.cpp
@@ -1,7 +1,9 @@
 // Copyright Epic Games, Inc. All Rights Reserved.
 
 #include "Elements/Component/ComponentElementHierarchyInterface.h"
+
 #include "Elements/Component/ComponentElementData.h"
+#include "Elements/Framework/TypedElementRegistry.h"
 #include "Components/ActorComponent.h"
 
 #include "Elements/Framework/EngineElementsLibrary.h"
diff --git a/Engine/Source/Runtime/Engine/Private/Elements/Component/ComponentElementWorldInterface.cpp b/Engine/Source/Runtime/Engine/Private/Elements/Component/ComponentElementWorldInterface.cpp
--- a/Engine/Source/Runtime/Engine/Private/Elements/Component/ComponentElementWorldInterface.cpp
+++ b/Engine/Source/Runtime/Engine/Private/Elements/Component/ComponentElementWorldInterface.cpp
@@ -1,10 +1,12 @@
 // Copyright Epic Games, Inc. All Rights Reserved.
 
 #include "Elements/Component/ComponentElementWorldInterface.h"
-#include "Elements/Component/ComponentElementData.h"
-#include "Components/PrimitiveComponent.h"
 
 #include "Elements/Actor/ActorElementWorldInterface.h"
+#include "Elements/Component/ComponentElementData.h"
+#include "Elements/Framework/TypedElementRegistry.h"
+
+#include "Components/PrimitiveComponent.h"
 
 bool UComponentElementWorldInterface::CanEditElement(const FTypedElementHandle& InElementHandle)
 {
diff --git a/Engine/Source/Runtime/Engine/Private/Elements/Framework/EngineElements.cpp b/Engine/Source/Runtime/Engine/Private/Elements/Framework/EngineElements.cpp
--- a/Engine/Source/Runtime/Engine/Private/Elements/Framework/EngineElements.cpp
+++ b/Engine/Source/Runtime/Engine/Private/Elements/Framework/EngineElements.cpp
@@ -38,7 +38,7 @@ void RegisterEngineObjectElements()
 {
 	UTypedElementRegistry* Registry = UTypedElementRegistry::GetInstance();
 
-	Registry->RegisterElementType<FObjectElementData>(NAME_Object);
+	Registry->RegisterElementType<FObjectElementData, true>(NAME_Object);
 	Registry->RegisterElementInterface<ITypedElementAssetDataInterface>(NAME_Object, NewObject<UObjectElementAssetDataInterface>());
 	Registry->RegisterElementInterface<ITypedElementObjectInterface>(NAME_Object, NewObject<UObjectElementObjectInterface>());
 	Registry->RegisterElementInterface<ITypedElementCounterInterface>(NAME_Object, NewObject<UObjectElementCounterInterface>());
@@ -49,7 +49,7 @@ void RegisterEngineActorElements()
 {
 	UTypedElementRegistry* Registry = UTypedElementRegistry::GetInstance();
 
-	Registry->RegisterElementType<FActorElementData>(NAME_Actor);
+	Registry->RegisterElementType<FActorElementData, true>(NAME_Actor);
 	Registry->RegisterElementInterface<ITypedElementAssetDataInterface>(NAME_Actor, NewObject<UActorElementAssetDataInterface>());
 	Registry->RegisterElementInterface<ITypedElementHierarchyInterface>(NAME_Actor, NewObject<UActorElementHierarchyInterface>());
 	Registry->RegisterElementInterface<ITypedElementObjectInterface>(NAME_Actor, NewObject<UActorElementObjectInterface>());
@@ -62,7 +62,7 @@ void RegisterEngineComponentElements()
 {
 	UTypedElementRegistry* Registry = UTypedElementRegistry::GetInstance();
 
-	Registry->RegisterElementType<FComponentElementData>(NAME_Components);
+	Registry->RegisterElementType<FComponentElementData, true>(NAME_Components);
 	Registry->RegisterElementInterface<ITypedElementHierarchyInterface>(NAME_Components, NewObject<UComponentElementHierarchyInterface>());
 	Registry->RegisterElementInterface<ITypedElementObjectInterface>(NAME_Components, NewObject<UComponentElementObjectInterface>());
 	Registry->RegisterElementInterface<ITypedElementCounterInterface>(NAME_Components, NewObject<UComponentElementCounterInterface>());
@@ -74,7 +74,7 @@ void RegisterEngineSMInstanceElements()
 {
 	UTypedElementRegistry* Registry = UTypedElementRegistry::GetInstance();
 
-	Registry->RegisterElementType<FSMInstanceElementData>(NAME_SMInstance);
+	Registry->RegisterElementType<FSMInstanceElementData, true>(NAME_SMInstance);
 	Registry->RegisterElementInterface<ITypedElementHierarchyInterface>(NAME_SMInstance, NewObject<USMInstanceElementHierarchyInterface>());
 	Registry->RegisterElementInterface<ITypedElementWorldInterface>(NAME_SMInstance, NewObject<USMInstanceElementWorldInterface>());
 	Registry->RegisterElementInterface<ITypedElementSelectionInterface>(NAME_SMInstance, NewObject<USMInstanceElementSelectionInterface>());
diff --git a/Engine/Source/Runtime/Engine/Private/Elements/Framework/EngineElementsLibrary.cpp b/Engine/Source/Runtime/Engine/Private/Elements/Framework/EngineElementsLibrary.cpp
--- a/Engine/Source/Runtime/Engine/Private/Elements/Framework/EngineElementsLibrary.cpp
+++ b/Engine/Source/Runtime/Engine/Private/Elements/Framework/EngineElementsLibrary.cpp
@@ -334,15 +334,15 @@ FTypedElementHandle UEngineElementsLibrary::AcquireEditorObjectElementHandle(con
 	return EngineElementsLibraryUtil::AcquireEditorTypedElementHandle<UObject, FObjectElementData>(Object, GObjectElementOwnerStore, &UEngineElementsLibrary::CreateObjectElement, bAllowCreate);
 }
 
-FTypedElementHandle UEngineElementsLibrary::K2_AcquireEditorObjectElementHandle(const UObject* Object, const bool bAllowCreate)
+FScriptTypedElementHandle UEngineElementsLibrary::K2_AcquireEditorObjectElementHandle(const UObject* Object, const bool bAllowCreate)
 {
 	if (!Object)
 	{
 		FFrame::KismetExecutionMessage(TEXT("Cannot acquire an editor object element handle from a nullptr object."), ELogVerbosity::Error);
-		return FTypedElementHandle();
+		return FScriptTypedElementHandle();
 	}
 
-	return AcquireEditorObjectElementHandle(Object , bAllowCreate);
+	return UTypedElementRegistry::GetInstance()->CreateScriptHandle(AcquireEditorObjectElementHandle(Object , bAllowCreate).GetId());
 }
 
 void UEngineElementsLibrary::ReplaceEditorObjectElementHandles(const TMap<const UObject*, const UObject*>& ReplacementObjects)
@@ -391,15 +391,15 @@ FTypedElementHandle UEngineElementsLibrary::AcquireEditorActorElementHandle(cons
 	return EngineElementsLibraryUtil::AcquireEditorTypedElementHandle<AActor, FActorElementData>(Actor, GActorElementOwnerStore, &UEngineElementsLibrary::CreateActorElement, bAllowCreate);
 }
 
-FTypedElementHandle UEngineElementsLibrary::K2_AcquireEditorActorElementHandle(const AActor* Actor, const bool bAllowCreate)
+FScriptTypedElementHandle UEngineElementsLibrary::K2_AcquireEditorActorElementHandle(const AActor* Actor, const bool bAllowCreate)
 {
 	if (!Actor)
 	{
 		FFrame::KismetExecutionMessage(TEXT("Cannot acquire an editor actor element handle from a nullptr actor."), ELogVerbosity::Error);
-		return FTypedElementHandle();
+		return FScriptTypedElementHandle();
 	}
 
-	return AcquireEditorActorElementHandle(Actor, bAllowCreate);
+	return UTypedElementRegistry::GetInstance()->CreateScriptHandle(AcquireEditorActorElementHandle(Actor, bAllowCreate).GetId());
 }
 
 void UEngineElementsLibrary::ReplaceEditorActorElementHandles(const TMap<const AActor*, const AActor*>& ReplacementActors)
@@ -448,15 +448,15 @@ FTypedElementHandle UEngineElementsLibrary::AcquireEditorComponentElementHandle(
 	return EngineElementsLibraryUtil::AcquireEditorTypedElementHandle<UActorComponent, FComponentElementData>(Component, GComponentElementOwnerStore, &UEngineElementsLibrary::CreateComponentElement, bAllowCreate);
 }
 
-FTypedElementHandle UEngineElementsLibrary::K2_AcquireEditorComponentElementHandle(const UActorComponent* Component, const bool bAllowCreate)
+FScriptTypedElementHandle UEngineElementsLibrary::K2_AcquireEditorComponentElementHandle(const UActorComponent* Component, const bool bAllowCreate)
 {
 	if (!Component)
 	{
 		FFrame::KismetExecutionMessage(TEXT("Cannot acquire an editor component element handle from a nullptr component."), ELogVerbosity::Error);
-		return FTypedElementHandle();
+		return FScriptTypedElementHandle();
 	}
 
-	return AcquireEditorComponentElementHandle(Component, bAllowCreate);
+	return UTypedElementRegistry::GetInstance()->CreateScriptHandle(AcquireEditorComponentElementHandle(Component, bAllowCreate).GetId());
 }
 
 void UEngineElementsLibrary::ReplaceEditorComponentElementHandles(const TMap<const UActorComponent*, const UActorComponent*>& ReplacementComponents)
@@ -529,12 +529,12 @@ FTypedElementHandle UEngineElementsLibrary::AcquireEditorSMInstanceElementHandle
 	return AcquireEditorSMInstanceElementHandle(FSMInstanceId{ const_cast<UInstancedStaticMeshComponent*>(ISMComponent), InstanceIndex }, bAllowCreate);
 }
 
-FTypedElementHandle UEngineElementsLibrary::K2_AcquireEditorSMInstanceElementHandle(const UInstancedStaticMeshComponent* ISMComponent, const int32 InstanceIndex, const bool bAllowCreate)
+FScriptTypedElementHandle UEngineElementsLibrary::K2_AcquireEditorSMInstanceElementHandle(const UInstancedStaticMeshComponent* ISMComponent, const int32 InstanceIndex, const bool bAllowCreate)
 {
 	if (!ISMComponent)
 	{
 		FFrame::KismetExecutionMessage(TEXT("Cannot acquire an editor SMInstance element handle from a nullptr ISMComponent."), ELogVerbosity::Error);
-		return FTypedElementHandle();
+		return FScriptTypedElementHandle();
 	}
 
 	if (ISMComponent->GetInstanceCount() <= InstanceIndex || InstanceIndex < 0)
@@ -545,10 +545,10 @@ FTypedElementHandle UEngineElementsLibrary::K2_AcquireEditorSMInstanceElementHan
 		FormatOrderedArguments.Add(ISMComponent->GetInstanceCount());
 		const FString Message = FString::Format(TEXT("Cannot acquire an editor SMInstance element handle from a instance that doesn't exist in the ISMComponent. (InstanceIndex: {0}, Instance Count: {1})"), FormatOrderedArguments);
 		FFrame::KismetExecutionMessage(*Message, ELogVerbosity::Error);
-		return FTypedElementHandle();
+		return FScriptTypedElementHandle();
 	}
 
-	return AcquireEditorSMInstanceElementHandle(ISMComponent, InstanceIndex, bAllowCreate);
+	return UTypedElementRegistry::GetInstance()->CreateScriptHandle(AcquireEditorSMInstanceElementHandle(ISMComponent, InstanceIndex, bAllowCreate).GetId());
 }
 
 FTypedElementHandle UEngineElementsLibrary::AcquireEditorSMInstanceElementHandle(const FSMInstanceId& SMInstanceId, const bool bAllowCreate)
diff --git a/Engine/Source/Runtime/Engine/Private/Elements/Framework/TypedElementCommonActions.cpp b/Engine/Source/Runtime/Engine/Private/Elements/Framework/TypedElementCommonActions.cpp
--- a/Engine/Source/Runtime/Engine/Private/Elements/Framework/TypedElementCommonActions.cpp
+++ b/Engine/Source/Runtime/Engine/Private/Elements/Framework/TypedElementCommonActions.cpp
@@ -6,6 +6,7 @@
 #include "Elements/Framework/TypedElementUtil.h"
 
 #include "UObject/GCObjectScopeGuard.h"
+#include "UObject/Stack.h"
 
 bool FTypedElementCommonActionsCustomization::DeleteElements(ITypedElementWorldInterface* InWorldInterface, TArrayView<const FTypedElementHandle> InElementHandles, UWorld* InWorld, UTypedElementSelectionSet* InSelectionSet, const FTypedElementDeletionOptions& InDeletionOptions)
 {
@@ -20,15 +21,20 @@ void FTypedElementCommonActionsCustomization::DuplicateElements(ITypedElementWor
 
 bool UTypedElementCommonActions::DeleteSelectedElements(UTypedElementSelectionSet* SelectionSet, UWorld* World, const FTypedElementDeletionOptions& DeletionOptions)
 {
+	if (!SelectionSet)
+	{
+		FFrame::KismetExecutionMessage(TEXT("SelectionSet is null."), ELogVerbosity::Error);
+	}
+
 	FTypedElementListRef NormalizedElements = SelectionSet->GetNormalizedSelection(FTypedElementSelectionNormalizationOptions());
 	return DeleteNormalizedElements(NormalizedElements, World, SelectionSet, DeletionOptions);
 }
 
-bool UTypedElementCommonActions::DeleteNormalizedElements(const FTypedElementListProxy ElementList, UWorld* World, UTypedElementSelectionSet* SelectionSet, const FTypedElementDeletionOptions& DeletionOptions)
+bool UTypedElementCommonActions::DeleteNormalizedElements(const FTypedElementListConstPtr& ElementListPtr, UWorld* World, UTypedElementSelectionSet* SelectionSet, const FTypedElementDeletionOptions& DeletionOptions)
 {
 	bool bSuccess = false;
 
-	if (FTypedElementListConstPtr ElementListPtr = ElementList.GetElementList())
+	if (ElementListPtr)
 	{
 		TMap<FTypedHandleTypeId, TArray<FTypedElementHandle>> ElementsToDeleteByType;
 		TypedElementUtil::BatchElementsByType(ElementListPtr.ToSharedRef(), ElementsToDeleteByType);
@@ -56,10 +62,10 @@ TArray<FTypedElementHandle> UTypedElementCommonActions::DuplicateSelectedElement
 	return DuplicateNormalizedElements(NormalizedElements, World, LocationOffset);
 }
 
-TArray<FTypedElementHandle> UTypedElementCommonActions::DuplicateNormalizedElements(const FTypedElementListProxy ElementList, UWorld* World, const FVector& LocationOffset)
+TArray<FTypedElementHandle> UTypedElementCommonActions::DuplicateNormalizedElements(const FTypedElementListConstPtr& ElementListPtr, UWorld* World, const FVector& LocationOffset)
 {
 	TArray<FTypedElementHandle> NewElements;
-	if (FTypedElementListConstPtr ElementListPtr = ElementList.GetElementList())
+	if (ElementListPtr)
 	{
 		NewElements.Reserve(ElementListPtr->Num());
 
@@ -81,6 +87,40 @@ TArray<FTypedElementHandle> UTypedElementCommonActions::DuplicateNormalizedEleme
 	return NewElements;
 }
 
+bool UTypedElementCommonActions::DeleteNormalizedElements(const FManagedTypedElementListProxy ElementList, UWorld* World, UTypedElementSelectionSet* InSelectionSet, const FTypedElementDeletionOptions& DeletionOptions)
+{
+	FTypedElementListPtr NativeList = UE::TypedElementFramework::ConvertToNativeTypedElementList(ElementList.GetElementList());
+	if (!NativeList)
+	{
+		FFrame::KismetExecutionMessage(TEXT("ElementList is in a invalid state."), ELogVerbosity::Error);
+		return false;
+	}
+
+	return DeleteNormalizedElements(NativeList, World, InSelectionSet, DeletionOptions);
+}
+
+TArray<FScriptTypedElementHandle> UTypedElementCommonActions::K2_DuplicateSelectedElements(const UTypedElementSelectionSet* SelectionSet, UWorld* World, const FVector& LocationOffset)
+{
+	if (!SelectionSet)
+	{
+		FFrame::KismetExecutionMessage(TEXT("SelectionSet is null."), ELogVerbosity::Error);
+	}
+
+	return TypedElementUtil::ConvertToScriptElementArray(DuplicateSelectedElements(SelectionSet, World, LocationOffset), SelectionSet->GetElementList()->GetRegistry());
+}
+
+TArray<FScriptTypedElementHandle> UTypedElementCommonActions::DuplicateNormalizedElements(const FManagedTypedElementListProxy ElementList, UWorld* World, const FVector& LocationOffset)
+{
+	FTypedElementListPtr NativeList = UE::TypedElementFramework::ConvertToNativeTypedElementList(ElementList.GetElementList());
+	if (!NativeList)
+	{
+		FFrame::KismetExecutionMessage(TEXT("ElementList is in a invalid state."), ELogVerbosity::Error);
+		return {};
+	}
+
+	return TypedElementUtil::ConvertToScriptElementArray(DuplicateNormalizedElements(NativeList, World, LocationOffset), NativeList->GetRegistry());
+}
+
 FTypedElementCommonActionsElement UTypedElementCommonActions::ResolveCommonActionsElement(const FTypedElementHandle& InElementHandle) const
 {
 	return InElementHandle
diff --git a/Engine/Source/Runtime/Engine/Private/Elements/Interfaces/TypedElementWorldInterface.cpp b/Engine/Source/Runtime/Engine/Private/Elements/Interfaces/TypedElementWorldInterface.cpp
new file mode 100644
--- /dev/null
+++ b/Engine/Source/Runtime/Engine/Private/Elements/Interfaces/TypedElementWorldInterface.cpp
@@ -0,0 +1,261 @@
+// Copyright Epic Games, Inc. All Rights Reserved.
+
+#include "Elements/Interfaces/TypedElementWorldInterface.h"
+
+#include "Elements/Framework/TypedElementRegistry.h"
+
+#include "UObject/Stack.h"
+
+bool ITypedElementWorldInterface::IsTemplateElement(const FScriptTypedElementHandle& InElementHandle)
+{
+	FTypedElementHandle NativeHandle = InElementHandle.GetTypedElementHandle();
+	if (!NativeHandle)
+	{
+		FFrame::KismetExecutionMessage(TEXT("InElementHandle is not a valid handle."), ELogVerbosity::Error);
+		return false;
+	}
+
+	return IsTemplateElement(NativeHandle);
+}
+
+bool ITypedElementWorldInterface::CanEditElement(const FScriptTypedElementHandle& InElementHandle)
+{
+	FTypedElementHandle NativeHandle = InElementHandle.GetTypedElementHandle();
+	if (!NativeHandle)
+	{
+		FFrame::KismetExecutionMessage(TEXT("InElementHandle is not a valid handle."), ELogVerbosity::Error);
+		return false;
+	}
+
+	return CanEditElement(NativeHandle);
+}
+
+ULevel* ITypedElementWorldInterface::GetOwnerLevel(const FScriptTypedElementHandle& InElementHandle)
+{
+	FTypedElementHandle NativeHandle = InElementHandle.GetTypedElementHandle();
+	if (!NativeHandle)
+	{
+		FFrame::KismetExecutionMessage(TEXT("InElementHandle is not a valid handle."), ELogVerbosity::Error);
+		return nullptr;
+	}
+
+	return GetOwnerLevel(NativeHandle);
+}
+
+UWorld* ITypedElementWorldInterface::GetOwnerWorld(const FScriptTypedElementHandle& InElementHandle)
+{
+	FTypedElementHandle NativeHandle = InElementHandle.GetTypedElementHandle();
+	if (!NativeHandle)
+	{
+		FFrame::KismetExecutionMessage(TEXT("InElementHandle is not a valid handle."), ELogVerbosity::Error);
+		return nullptr;
+	}
+
+	return GetOwnerWorld(NativeHandle);
+}
+
+bool ITypedElementWorldInterface::GetBounds(const FScriptTypedElementHandle& InElementHandle, FBoxSphereBounds& OutBounds)
+{
+	FTypedElementHandle NativeHandle = InElementHandle.GetTypedElementHandle();
+	if (!NativeHandle)
+	{
+		FFrame::KismetExecutionMessage(TEXT("InElementHandle is not a valid handle."), ELogVerbosity::Error);
+		return false;
+	}
+
+	return GetBounds(NativeHandle, OutBounds);
+}
+
+
+bool ITypedElementWorldInterface::CanMoveElement(const FScriptTypedElementHandle& InElementHandle, const ETypedElementWorldType InWorldType)
+{
+	FTypedElementHandle NativeHandle = InElementHandle.GetTypedElementHandle();
+	if (!NativeHandle)
+	{
+		FFrame::KismetExecutionMessage(TEXT("InElementHandle is not a valid handle."), ELogVerbosity::Error);
+		return false;
+	}
+
+	return CanMoveElement(NativeHandle, InWorldType);
+}
+
+bool ITypedElementWorldInterface::GetWorldTransform(const FScriptTypedElementHandle& InElementHandle, FTransform& OutTransform)
+{
+	FTypedElementHandle NativeHandle = InElementHandle.GetTypedElementHandle();
+	if (!NativeHandle)
+	{
+		FFrame::KismetExecutionMessage(TEXT("InElementHandle is not a valid handle."), ELogVerbosity::Error);
+		return false;
+	}
+
+	return GetWorldTransform(NativeHandle, OutTransform);
+}
+
+bool ITypedElementWorldInterface::SetWorldTransform(const FScriptTypedElementHandle& InElementHandle, const FTransform& InTransform)
+{
+	FTypedElementHandle NativeHandle = InElementHandle.GetTypedElementHandle();
+	if (!NativeHandle)
+	{
+		FFrame::KismetExecutionMessage(TEXT("InElementHandle is not a valid handle."), ELogVerbosity::Error);
+		return false;
+	}
+
+	return SetWorldTransform(NativeHandle, InTransform);
+}
+
+bool ITypedElementWorldInterface::GetRelativeTransform(const FScriptTypedElementHandle& InElementHandle, FTransform& OutTransform)
+{
+	FTypedElementHandle NativeHandle = InElementHandle.GetTypedElementHandle();
+	if (!NativeHandle)
+	{
+		FFrame::KismetExecutionMessage(TEXT("InElementHandle is not a valid handle."), ELogVerbosity::Error);
+		return false;
+	}
+
+	return GetRelativeTransform(NativeHandle, OutTransform);
+}
+
+bool ITypedElementWorldInterface::SetRelativeTransform(const FScriptTypedElementHandle& InElementHandle, const FTransform& InTransform)
+{
+	FTypedElementHandle NativeHandle = InElementHandle.GetTypedElementHandle();
+	if (!NativeHandle)
+	{
+		FFrame::KismetExecutionMessage(TEXT("InElementHandle is not a valid handle."), ELogVerbosity::Error);
+		return false;
+	}
+
+	return SetRelativeTransform(NativeHandle, InTransform);
+}
+
+bool ITypedElementWorldInterface::GetPivotOffset(const FScriptTypedElementHandle& InElementHandle, FVector& OutPivotOffset)
+{
+	FTypedElementHandle NativeHandle = InElementHandle.GetTypedElementHandle();
+	if (!NativeHandle)
+	{
+		FFrame::KismetExecutionMessage(TEXT("InElementHandle is not a valid handle."), ELogVerbosity::Error);
+		return false;
+	}
+
+	return GetPivotOffset(NativeHandle, OutPivotOffset);
+}
+
+bool ITypedElementWorldInterface::SetPivotOffset(const FScriptTypedElementHandle& InElementHandle, const FVector& InPivotOffset)
+{
+	FTypedElementHandle NativeHandle = InElementHandle.GetTypedElementHandle();
+	if (!NativeHandle)
+	{
+		FFrame::KismetExecutionMessage(TEXT("InElementHandle is not a valid handle."), ELogVerbosity::Error);
+		return false;
+	}
+
+	return SetPivotOffset(NativeHandle, InPivotOffset);
+}
+
+
+void ITypedElementWorldInterface::NotifyMovementStarted(const FScriptTypedElementHandle& InElementHandle)
+{
+	FTypedElementHandle NativeHandle = InElementHandle.GetTypedElementHandle();
+	if (!NativeHandle)
+	{
+		FFrame::KismetExecutionMessage(TEXT("InElementHandle is not a valid handle."), ELogVerbosity::Error);
+		return;
+	}
+
+	NotifyMovementStarted(NativeHandle);
+}
+
+
+void ITypedElementWorldInterface::NotifyMovementOngoing(const FScriptTypedElementHandle& InElementHandle)
+{
+	FTypedElementHandle NativeHandle = InElementHandle.GetTypedElementHandle();
+	if (!NativeHandle)
+	{
+		FFrame::KismetExecutionMessage(TEXT("InElementHandle is not a valid handle."), ELogVerbosity::Error);
+		return;
+	}
+
+	NotifyMovementOngoing(NativeHandle);
+}
+
+void ITypedElementWorldInterface::NotifyMovementEnded(const FScriptTypedElementHandle& InElementHandle)
+{
+	FTypedElementHandle NativeHandle = InElementHandle.GetTypedElementHandle();
+	if (!NativeHandle)
+	{
+		FFrame::KismetExecutionMessage(TEXT("InElementHandle is not a valid handle."), ELogVerbosity::Error);
+		return;
+	}
+
+	NotifyMovementEnded(NativeHandle);
+}
+
+bool ITypedElementWorldInterface::CanDeleteElement(const FScriptTypedElementHandle& InElementHandle)
+{
+	FTypedElementHandle NativeHandle = InElementHandle.GetTypedElementHandle();
+	if (!NativeHandle)
+	{
+		FFrame::KismetExecutionMessage(TEXT("InElementHandle is not a valid handle."), ELogVerbosity::Error);
+		return false;
+	}
+
+	return CanDeleteElement(NativeHandle);
+}
+
+bool ITypedElementWorldInterface::DeleteElement(const FScriptTypedElementHandle& InElementHandle, UWorld* InWorld, UTypedElementSelectionSet* InSelectionSet, const FTypedElementDeletionOptions& InDeletionOptions)
+{
+	FTypedElementHandle NativeHandle = InElementHandle.GetTypedElementHandle();
+	if (!NativeHandle)
+	{
+		FFrame::KismetExecutionMessage(TEXT("InElementHandle is not a valid handle."), ELogVerbosity::Error);
+		return false;
+	}
+
+	if (!InWorld)
+	{
+		FFrame::KismetExecutionMessage(TEXT("InWorld is null"), ELogVerbosity::Error);
+		return false;
+	}
+
+	if (!InSelectionSet)
+	{
+		FFrame::KismetExecutionMessage(TEXT("InSelectionSet is null"), ELogVerbosity::Error);
+		return false;
+	}
+
+	return DeleteElement(NativeHandle, InWorld, InSelectionSet, InDeletionOptions);
+}
+
+bool ITypedElementWorldInterface::CanDuplicateElement(const FScriptTypedElementHandle& InElementHandle)
+{
+	FTypedElementHandle NativeHandle = InElementHandle.GetTypedElementHandle();
+	if (!NativeHandle)
+	{
+		FFrame::KismetExecutionMessage(TEXT("InElementHandle is not a valid handle."), ELogVerbosity::Error);
+		return false;
+	}
+
+	return CanDuplicateElement(NativeHandle);
+}
+
+FScriptTypedElementHandle ITypedElementWorldInterface::DuplicateElement(const FScriptTypedElementHandle& InElementHandle, UWorld* InWorld, const FVector& InLocationOffset)
+{
+	FTypedElementHandle NativeHandle = InElementHandle.GetTypedElementHandle();
+	if (!NativeHandle)
+	{
+		FFrame::KismetExecutionMessage(TEXT("InElementHandle is not a valid handle."), ELogVerbosity::Error);
+		return FScriptTypedElementHandle();
+	}
+
+	if (!InWorld)
+	{
+		FFrame::KismetExecutionMessage(TEXT("InWorld is null"), ELogVerbosity::Error);
+		return FScriptTypedElementHandle();
+	}
+
+	return GetRegistry().CreateScriptHandle(DuplicateElement(NativeHandle, InWorld, InLocationOffset).GetId());
+}
+
+class UTypedElementRegistry& ITypedElementWorldInterface::GetRegistry() const
+{
+	return *UTypedElementRegistry::GetInstance();
+}
diff --git a/Engine/Source/Runtime/Engine/Private/Elements/SMInstance/SMInstanceElementHierarchyInterface.cpp b/Engine/Source/Runtime/Engine/Private/Elements/SMInstance/SMInstanceElementHierarchyInterface.cpp
--- a/Engine/Source/Runtime/Engine/Private/Elements/SMInstance/SMInstanceElementHierarchyInterface.cpp
+++ b/Engine/Source/Runtime/Engine/Private/Elements/SMInstance/SMInstanceElementHierarchyInterface.cpp
@@ -1,7 +1,9 @@
 // Copyright Epic Games, Inc. All Rights Reserved.
 
 #include "Elements/SMInstance/SMInstanceElementHierarchyInterface.h"
+
 #include "Elements/SMInstance/SMInstanceElementData.h"
+#include "Elements/Framework/TypedElementRegistry.h"
 #include "Components/InstancedStaticMeshComponent.h"
 
 #include "Elements/Framework/EngineElementsLibrary.h"
diff --git a/Engine/Source/Runtime/Engine/Private/Elements/SMInstance/SMInstanceElementSelectionInterface.cpp b/Engine/Source/Runtime/Engine/Private/Elements/SMInstance/SMInstanceElementSelectionInterface.cpp
--- a/Engine/Source/Runtime/Engine/Private/Elements/SMInstance/SMInstanceElementSelectionInterface.cpp
+++ b/Engine/Source/Runtime/Engine/Private/Elements/SMInstance/SMInstanceElementSelectionInterface.cpp
@@ -5,7 +5,7 @@
 #include "Components/InstancedStaticMeshComponent.h"
 #include "Elements/Framework/TypedElementList.h"
 
-bool USMInstanceElementSelectionInterface::SelectElement(const FTypedElementHandle& InElementHandle, FTypedElementListProxy InSelectionSet, const FTypedElementSelectionOptions& InSelectionOptions)
+bool USMInstanceElementSelectionInterface::SelectElement(const FTypedElementHandle& InElementHandle, const FTypedElementListPtr& InSelectionSet, const FTypedElementSelectionOptions& InSelectionOptions)
 {
 	FSMInstanceManager SMInstance = SMInstanceElementDataUtil::GetSMInstanceFromHandle(InElementHandle);
 	if (SMInstance && ITypedElementSelectionInterface::SelectElement(InElementHandle, InSelectionSet, InSelectionOptions))
@@ -16,7 +16,7 @@ bool USMInstanceElementSelectionInterface::SelectElement(const FTypedElementHand
 	return false;
 }
 
-bool USMInstanceElementSelectionInterface::DeselectElement(const FTypedElementHandle& InElementHandle, FTypedElementListProxy InSelectionSet, const FTypedElementSelectionOptions& InSelectionOptions)
+bool USMInstanceElementSelectionInterface::DeselectElement(const FTypedElementHandle& InElementHandle, const FTypedElementListPtr& InSelectionSet, const FTypedElementSelectionOptions& InSelectionOptions)
 {
 	FSMInstanceManager SMInstance = SMInstanceElementDataUtil::GetSMInstanceFromHandle(InElementHandle);
 	if (SMInstance && ITypedElementSelectionInterface::DeselectElement(InElementHandle, InSelectionSet, InSelectionOptions))
diff --git a/Engine/Source/Runtime/Engine/Private/Elements/SMInstance/SMInstanceElementWorldInterface.cpp b/Engine/Source/Runtime/Engine/Private/Elements/SMInstance/SMInstanceElementWorldInterface.cpp
--- a/Engine/Source/Runtime/Engine/Private/Elements/SMInstance/SMInstanceElementWorldInterface.cpp
+++ b/Engine/Source/Runtime/Engine/Private/Elements/SMInstance/SMInstanceElementWorldInterface.cpp
@@ -1,11 +1,12 @@
 // Copyright Epic Games, Inc. All Rights Reserved.
 
 #include "Elements/SMInstance/SMInstanceElementWorldInterface.h"
-#include "Elements/SMInstance/SMInstanceElementData.h"
-#include "Components/InstancedStaticMeshComponent.h"
 
+#include "Components/InstancedStaticMeshComponent.h"
 #include "Elements/Framework/EngineElementsLibrary.h"
+#include "Elements/Framework/TypedElementRegistry.h"
 #include "Elements/Framework/TypedElementSelectionSet.h"
+#include "Elements/SMInstance/SMInstanceElementData.h"
 
 #include "Engine/StaticMesh.h"
 
diff --git a/Engine/Source/Runtime/Engine/Public/Elements/Framework/EngineElementsLibrary.h b/Engine/Source/Runtime/Engine/Public/Elements/Framework/EngineElementsLibrary.h
--- a/Engine/Source/Runtime/Engine/Public/Elements/Framework/EngineElementsLibrary.h
+++ b/Engine/Source/Runtime/Engine/Public/Elements/Framework/EngineElementsLibrary.h
@@ -36,7 +36,7 @@ public:
 	static void DestroyEditorObjectElement(const UObject* Object);
 	static FTypedElementHandle AcquireEditorObjectElementHandle(const UObject* Object, const bool bAllowCreate = true);
 	UFUNCTION(BlueprintCallable, BlueprintPure=false, Category="TypedElementFramework|Object", meta=(DisplayName="Acquire Editor Object Element Handle", ScriptMethod="AcquireEditorElementHandle"))
-	static FTypedElementHandle K2_AcquireEditorObjectElementHandle(const UObject* Object, const bool bAllowCreate = true);
+	static FScriptTypedElementHandle K2_AcquireEditorObjectElementHandle(const UObject* Object, const bool bAllowCreate = true);
 	static void ReplaceEditorObjectElementHandles(const TMap<const UObject*, const UObject*>& ReplacementObjects);
 #endif
 
@@ -47,7 +47,7 @@ public:
 	static void DestroyEditorActorElement(const AActor* Actor);
 	static FTypedElementHandle AcquireEditorActorElementHandle(const AActor* Actor, const bool bAllowCreate = true);
 	UFUNCTION(BlueprintCallable, Category="TypedElementFramework|Actor", meta=(DisplayName="Acquire Editor Actor Element Handle", ScriptMethod="AcquireEditorElementHandle"))
-	static FTypedElementHandle K2_AcquireEditorActorElementHandle(const AActor* Actor, const bool bAllowCreate = true);
+	static FScriptTypedElementHandle K2_AcquireEditorActorElementHandle(const AActor* Actor, const bool bAllowCreate = true);
 	static void ReplaceEditorActorElementHandles(const TMap<const AActor*, const AActor*>& ReplacementActors);
 #endif
 
@@ -58,7 +58,7 @@ public:
 	static void DestroyEditorComponentElement(const UActorComponent* Component);
 	static FTypedElementHandle AcquireEditorComponentElementHandle(const UActorComponent* Component, const bool bAllowCreate = true);
 	UFUNCTION(BlueprintCallable, BlueprintPure=false, Category="TypedElementFramework|Component", meta=(DisplayName="Acquire Editor Component Element Handle", ScriptMethod = "AcquireEditorElementHandle"))
-	static FTypedElementHandle K2_AcquireEditorComponentElementHandle(const UActorComponent* Component, const bool bAllowCreate = true);
+	static FScriptTypedElementHandle K2_AcquireEditorComponentElementHandle(const UActorComponent* Component, const bool bAllowCreate = true);
 	static void ReplaceEditorComponentElementHandles(const TMap<const UActorComponent*, const UActorComponent*>& ReplacementComponents);
 #endif
 
@@ -69,7 +69,7 @@ public:
 	static void DestroyEditorSMInstanceElement(const FSMInstanceElementId& SMInstanceElementId);
 	static FTypedElementHandle AcquireEditorSMInstanceElementHandle(const UInstancedStaticMeshComponent* ISMComponent, const int32 InstanceIndex, const bool bAllowCreate = true);
 	UFUNCTION(BlueprintCallable, BlueprintPure=false, Category="TypedElementFramework|StaticMeshInstance", meta=(DisplayName="Acquire Editor SMInstance Element Handle", ScriptMethod = "AcquireEditorElementHandle"))
-	static FTypedElementHandle K2_AcquireEditorSMInstanceElementHandle(const UInstancedStaticMeshComponent* ISMComponent, const int32 InstanceIndex, const bool bAllowCreate = true);
+	static FScriptTypedElementHandle K2_AcquireEditorSMInstanceElementHandle(const UInstancedStaticMeshComponent* ISMComponent, const int32 InstanceIndex, const bool bAllowCreate = true);
 	static FTypedElementHandle AcquireEditorSMInstanceElementHandle(const FSMInstanceId& SMInstanceId, const bool bAllowCreate = true);
 	static void ReplaceEditorSMInstanceElementHandles(const TMap<FSMInstanceId, FSMInstanceId>& ReplacementSMInstanceIds);
 #endif
diff --git a/Engine/Source/Runtime/Engine/Public/Elements/Framework/TypedElementCommonActions.h b/Engine/Source/Runtime/Engine/Public/Elements/Framework/TypedElementCommonActions.h
--- a/Engine/Source/Runtime/Engine/Public/Elements/Framework/TypedElementCommonActions.h
+++ b/Engine/Source/Runtime/Engine/Public/Elements/Framework/TypedElementCommonActions.h
@@ -81,22 +81,45 @@ public:
 	 * Delete any elements from the given list that can be deleted.
 	 * @note This list should have been pre-normalized via UTypedElementSelectionSet::GetNormalizedSelection or UTypedElementSelectionSet::GetNormalizedElementList.
 	 */
-	UFUNCTION(BlueprintCallable, Category="TypedElementFramework|Common")
-	bool DeleteNormalizedElements(const FTypedElementListProxy ElementList, UWorld* World, UTypedElementSelectionSet* InSelectionSet, const FTypedElementDeletionOptions& DeletionOptions);
+	bool DeleteNormalizedElements(const FTypedElementListConstPtr& ElementListPtr, UWorld* World, UTypedElementSelectionSet* InSelectionSet, const FTypedElementDeletionOptions& DeletionOptions);
 
 	/**
 	 * Duplicate any elements from the given selection set that can be duplicated.
 	 * @note Internally this just calls DuplicateNormalizedElements on the result of UTypedElementSelectionSet::GetNormalizedSelection.
 	 */
-	UFUNCTION(BlueprintCallable, Category="TypedElementFramework|Common")
 	TArray<FTypedElementHandle> DuplicateSelectedElements(const UTypedElementSelectionSet* SelectionSet, UWorld* World, const FVector& LocationOffset);
 	
 	/**
 	 * Duplicate any elements from the given list that can be duplicated.
 	 * @note This list should have been pre-normalized via UTypedElementSelectionSet::GetNormalizedSelection or UTypedElementSelectionSet::GetNormalizedElementList.
 	 */
+	TArray<FTypedElementHandle> DuplicateNormalizedElements(const FTypedElementListConstPtr& ElementListPtr, UWorld* World, const FVector& LocationOffset);
+
+
+	/**
+	 * Script Api
+	 */
+
+	/**
+	 * Delete any elements from the given list that can be deleted.
+	 * @note This list should have been pre-normalized via UTypedElementSelectionSet::GetNormalizedSelection or UTypedElementSelectionSet::GetNormalizedElementList.
+	 */
 	UFUNCTION(BlueprintCallable, Category="TypedElementFramework|Common")
-	TArray<FTypedElementHandle> DuplicateNormalizedElements(const FTypedElementListProxy ElementList, UWorld* World, const FVector& LocationOffset);
+	bool DeleteNormalizedElements(const FManagedTypedElementListProxy ElementList, UWorld* World, UTypedElementSelectionSet* InSelectionSet, const FTypedElementDeletionOptions& DeletionOptions);
+
+	/**
+	 * Duplicate any elements from the given selection set that can be duplicated.
+	 * @note Internally this just calls DuplicateNormalizedElements on the result of UTypedElementSelectionSet::GetNormalizedSelection.
+	 */
+	UFUNCTION(BlueprintCallable, DisplayName="Duplicate Selected Elements", Category = "TypedElementFramework|Common", meta=(ScriptName="DuplicateSelectedElements"))
+	TArray<FScriptTypedElementHandle> K2_DuplicateSelectedElements(const UTypedElementSelectionSet* SelectionSet, UWorld* World, const FVector& LocationOffset);
+	
+	/**
+	 * Duplicate any elements from the given list that can be duplicated.
+	 * @note This list should have been pre-normalized via UTypedElementSelectionSet::GetNormalizedSelection or UTypedElementSelectionSet::GetNormalizedElementList.
+	 */
+	UFUNCTION(BlueprintCallable, Category = "TypedElementFramework|Common")
+	TArray<FScriptTypedElementHandle> DuplicateNormalizedElements(const FManagedTypedElementListProxy ElementList, UWorld* World, const FVector& LocationOffset);
 
 private:
 	/**
diff --git a/Engine/Source/Runtime/Engine/Public/Elements/Interfaces/TypedElementWorldInterface.h b/Engine/Source/Runtime/Engine/Public/Elements/Interfaces/TypedElementWorldInterface.h
--- a/Engine/Source/Runtime/Engine/Public/Elements/Interfaces/TypedElementWorldInterface.h
+++ b/Engine/Source/Runtime/Engine/Public/Elements/Interfaces/TypedElementWorldInterface.h
@@ -59,91 +59,76 @@ public:
 	/**
 	 * Is this element considered a template within its world (eg, a CDO or archetype).
 	 */
-	UFUNCTION(BlueprintCallable, Category="TypedElementInterfaces|World")
 	virtual bool IsTemplateElement(const FTypedElementHandle& InElementHandle) { return false; }
 
 	/**
 	 * Can this element actually be edited in the world?
 	 */
-	UFUNCTION(BlueprintCallable, Category="TypedElementInterfaces|World")
 	virtual bool CanEditElement(const FTypedElementHandle& InElementHandle) { return true; }
 
 	/**
 	 * Get the owner level associated with this element, if any.
 	 */
-	UFUNCTION(BlueprintCallable, Category="TypedElementInterfaces|World")
 	virtual ULevel* GetOwnerLevel(const FTypedElementHandle& InElementHandle) { return nullptr; }
 
 	/**
 	 * Get the owner world associated with this element, if any.
 	 */
-	UFUNCTION(BlueprintCallable, Category="TypedElementInterfaces|World")
 	virtual UWorld* GetOwnerWorld(const FTypedElementHandle& InElementHandle) { return nullptr; }
 
 	/**
 	 * Get the bounds of this element, if any.
 	 */
-	UFUNCTION(BlueprintCallable, Category="TypedElementInterfaces|World")
 	virtual bool GetBounds(const FTypedElementHandle& InElementHandle, FBoxSphereBounds& OutBounds) { return false; }
 
 	/**
 	 * Can the given element be moved within the world?
 	 */
-	UFUNCTION(BlueprintCallable, Category="TypedElementInterfaces|World")
 	virtual bool CanMoveElement(const FTypedElementHandle& InElementHandle, const ETypedElementWorldType InWorldType) { return false; }
 
 	/**
 	 * Get the transform of this element within its owner world, if any.
 	 */
-	UFUNCTION(BlueprintCallable, Category="TypedElementInterfaces|World")
 	virtual bool GetWorldTransform(const FTypedElementHandle& InElementHandle, FTransform& OutTransform) { return false; }
 	
 	/**
 	 * Attempt to set the transform of this element within its owner world.
 	 */
-	UFUNCTION(BlueprintCallable, Category="TypedElementInterfaces|World")
 	virtual bool SetWorldTransform(const FTypedElementHandle& InElementHandle, const FTransform& InTransform) { return false; }
 
 	/**
 	 * Get the transform of this element relative to its parent, if any.
 	 */
-	UFUNCTION(BlueprintCallable, Category="TypedElementInterfaces|World")
 	virtual bool GetRelativeTransform(const FTypedElementHandle& InElementHandle, FTransform& OutTransform) { return GetWorldTransform(InElementHandle, OutTransform); }
 	
 	/**
 	 * Attempt to set the transform of this element relative to its parent.
 	 */
-	UFUNCTION(BlueprintCallable, Category="TypedElementInterfaces|World")
 	virtual bool SetRelativeTransform(const FTypedElementHandle& InElementHandle, const FTransform& InTransform) { return SetWorldTransform(InElementHandle, InTransform); }
 
 	/**
 	 * Get the local space offset of this element that should be added to its pivot location, if any.
 	 */
-	UFUNCTION(BlueprintCallable, Category="TypedElementInterfaces|World")
 	virtual bool GetPivotOffset(const FTypedElementHandle& InElementHandle, FVector& OutPivotOffset) { return false; }
 
 	/**
 	 * Attempt to set the local space offset of this element that should be added to its pivot location.
 	 */
-	UFUNCTION(BlueprintCallable, Category="TypedElementInterfaces|World")
 	virtual bool SetPivotOffset(const FTypedElementHandle& InElementHandle, const FVector& InPivotOffset) { return false; }
 
 	/**
 	 * Notify that this element is about to be moved.
 	 */
-	UFUNCTION(BlueprintCallable, Category="TypedElementInterfaces|World")
 	virtual void NotifyMovementStarted(const FTypedElementHandle& InElementHandle) {}
 
 	/**
 	 * Notify that this element is currently being moved.
 	 */
-	UFUNCTION(BlueprintCallable, Category="TypedElementInterfaces|World")
 	virtual void NotifyMovementOngoing(const FTypedElementHandle& InElementHandle) {}
 
 	/**
 	 * Notify that this element is done being moved.
 	 */
-	UFUNCTION(BlueprintCallable, Category="TypedElementInterfaces|World")
 	virtual void NotifyMovementEnded(const FTypedElementHandle& InElementHandle) {}
 
 	/**
@@ -166,14 +151,12 @@ public:
 	/**
 	 * Can the given element be deleted?
 	 */
-	UFUNCTION(BlueprintCallable, Category="TypedElementInterfaces|World")
 	virtual bool CanDeleteElement(const FTypedElementHandle& InElementHandle) { return false; }
 
 	/**
 	 * Delete the given element.
 	 * @note Default version calls DeleteElements with a single element.
 	 */
-	UFUNCTION(BlueprintCallable, Category="TypedElementInterfaces|World")
 	virtual bool DeleteElement(const FTypedElementHandle& InElementHandle, UWorld* InWorld, UTypedElementSelectionSet* InSelectionSet, const FTypedElementDeletionOptions& InDeletionOptions)
 	{
 		return DeleteElements(MakeArrayView(&InElementHandle, 1), InWorld, InSelectionSet, InDeletionOptions);
@@ -188,14 +171,12 @@ public:
 	/**
 	 * Can the given element be duplicated?
 	 */
-	UFUNCTION(BlueprintCallable, Category="TypedElementInterfaces|World")
 	virtual bool CanDuplicateElement(const FTypedElementHandle& InElementHandle) { return false; }
 
 	/**
 	 * Duplicate the given element.
 	 * @note Default version calls DuplicateElements with a single element.
 	 */
-	UFUNCTION(BlueprintCallable, Category="TypedElementInterfaces|World")
 	virtual FTypedElementHandle DuplicateElement(const FTypedElementHandle& InElementHandle, UWorld* InWorld, const FVector& InLocationOffset)
 	{
 		TArray<FTypedElementHandle> NewElements;
@@ -210,6 +191,133 @@ public:
 	virtual void DuplicateElements(TArrayView<const FTypedElementHandle> InElementHandles, UWorld* InWorld, const FVector& InLocationOffset, TArray<FTypedElementHandle>& OutNewElements)
 	{
 	}
+
+
+	/**
+	 * Script Api
+	 */
+
+	/**
+	 * Is this element considered a template within its world (eg, a CDO or archetype).
+	 */
+	UFUNCTION(BlueprintCallable, Category="TypedElementInterfaces|World")
+	virtual bool IsTemplateElement(const FScriptTypedElementHandle& InElementHandle);
+
+	/**
+	 * Can this element actually be edited in the world?
+	 */
+	UFUNCTION(BlueprintCallable, Category="TypedElementInterfaces|World")
+	virtual bool CanEditElement(const FScriptTypedElementHandle& InElementHandle);
+
+	/**
+	 * Get the owner level associated with this element, if any.
+	 */
+	UFUNCTION(BlueprintCallable, Category="TypedElementInterfaces|World")
+	virtual ULevel* GetOwnerLevel(const FScriptTypedElementHandle& InElementHandle);
+
+	/**
+	 * Get the owner world associated with this element, if any.
+	 */
+	UFUNCTION(BlueprintCallable, Category="TypedElementInterfaces|World")
+	virtual UWorld* GetOwnerWorld(const FScriptTypedElementHandle& InElementHandle);
+
+	/**
+	 * Get the bounds of this element, if any.
+	 */
+	UFUNCTION(BlueprintCallable, Category="TypedElementInterfaces|World")
+	virtual bool GetBounds(const FScriptTypedElementHandle& InElementHandle, FBoxSphereBounds& OutBounds);
+
+	/**
+	 * Can the given element be moved within the world?
+	 */
+	UFUNCTION(BlueprintCallable, Category="TypedElementInterfaces|World")
+	virtual bool CanMoveElement(const FScriptTypedElementHandle& InElementHandle, const ETypedElementWorldType InWorldType);
+
+	/**
+	 * Get the transform of this element within its owner world, if any.
+	 */
+	UFUNCTION(BlueprintCallable, Category="TypedElementInterfaces|World")
+	virtual bool GetWorldTransform(const FScriptTypedElementHandle& InElementHandle, FTransform& OutTransform);
+	
+	/**
+	 * Attempt to set the transform of this element within its owner world.
+	 */
+	UFUNCTION(BlueprintCallable, Category="TypedElementInterfaces|World")
+	virtual bool SetWorldTransform(const FScriptTypedElementHandle& InElementHandle, const FTransform& InTransform);
+
+	/**
+	 * Get the transform of this element relative to its parent, if any.
+	 */
+	UFUNCTION(BlueprintCallable, Category="TypedElementInterfaces|World")
+	virtual bool GetRelativeTransform(const FScriptTypedElementHandle& InElementHandle, FTransform& OutTransform);
+	
+	/**
+	 * Attempt to set the transform of this element relative to its parent.
+	 */
+	UFUNCTION(BlueprintCallable, Category="TypedElementInterfaces|World")
+	virtual bool SetRelativeTransform(const FScriptTypedElementHandle& InElementHandle, const FTransform& InTransform);
+
+	/**
+	 * Get the local space offset of this element that should be added to its pivot location, if any.
+	 */
+	UFUNCTION(BlueprintCallable, Category="TypedElementInterfaces|World")
+	virtual bool GetPivotOffset(const FScriptTypedElementHandle& InElementHandle, FVector& OutPivotOffset);
+
+	/**
+	 * Attempt to set the local space offset of this element that should be added to its pivot location.
+	 */
+	UFUNCTION(BlueprintCallable, Category="TypedElementInterfaces|World")
+	virtual bool SetPivotOffset(const FScriptTypedElementHandle& InElementHandle, const FVector& InPivotOffset);
+
+	/**
+	 * Notify that this element is about to be moved.
+	 */
+	UFUNCTION(BlueprintCallable, Category="TypedElementInterfaces|World")
+	virtual void NotifyMovementStarted(const FScriptTypedElementHandle& InElementHandle);
+
+	/**
+	 * Notify that this element is currently being moved.
+	 */
+	UFUNCTION(BlueprintCallable, Category="TypedElementInterfaces|World")
+	virtual void NotifyMovementOngoing(const FScriptTypedElementHandle& InElementHandle);
+
+	/**
+	 * Notify that this element is done being moved.
+	 */
+	UFUNCTION(BlueprintCallable, Category="TypedElementInterfaces|World")
+	virtual void NotifyMovementEnded(const FScriptTypedElementHandle& InElementHandle);
+
+	/**
+	 * Can the given element be deleted?
+	 */
+	UFUNCTION(BlueprintCallable, Category="TypedElementInterfaces|World")
+	virtual bool CanDeleteElement(const FScriptTypedElementHandle& InElementHandle);
+
+	/**
+	 * Delete the given element.
+	 * @note Default version calls DeleteElements with a single element.
+	 */
+	UFUNCTION(BlueprintCallable, Category="TypedElementInterfaces|World")
+	virtual bool DeleteElement(const FScriptTypedElementHandle& InElementHandle, UWorld* InWorld, UTypedElementSelectionSet* InSelectionSet, const FTypedElementDeletionOptions& InDeletionOptions);
+
+	/**
+	 * Can the given element be duplicated?
+	 */
+	UFUNCTION(BlueprintCallable, Category="TypedElementInterfaces|World")
+	virtual bool CanDuplicateElement(const FScriptTypedElementHandle& InElementHandle);
+
+	/**
+	 * Duplicate the given element.
+	 * @note Default version calls DuplicateElements with a single element.
+	 */
+	UFUNCTION(BlueprintCallable, Category="TypedElementInterfaces|World")
+	virtual FScriptTypedElementHandle DuplicateElement(const FScriptTypedElementHandle& InElementHandle, UWorld* InWorld, const FVector& InLocationOffset);
+
+private:
+	/**
+	 * Return the registry associated with this interface implementation
+	 */
+	virtual class UTypedElementRegistry& GetRegistry() const;
 };
 
 template <>
diff --git a/Engine/Source/Runtime/Engine/Public/Elements/SMInstance/SMInstanceElementSelectionInterface.h b/Engine/Source/Runtime/Engine/Public/Elements/SMInstance/SMInstanceElementSelectionInterface.h
--- a/Engine/Source/Runtime/Engine/Public/Elements/SMInstance/SMInstanceElementSelectionInterface.h
+++ b/Engine/Source/Runtime/Engine/Public/Elements/SMInstance/SMInstanceElementSelectionInterface.h
@@ -11,6 +11,6 @@ class ENGINE_API USMInstanceElementSelectionInterface : public UObject, public I
 	GENERATED_BODY()
 
 public:
-	virtual bool SelectElement(const FTypedElementHandle& InElementHandle, FTypedElementListProxy InSelectionSet, const FTypedElementSelectionOptions& InSelectionOptions) override;
-	virtual bool DeselectElement(const FTypedElementHandle& InElementHandle, FTypedElementListProxy InSelectionSet, const FTypedElementSelectionOptions& InSelectionOptions) override;
+	virtual bool SelectElement(const FTypedElementHandle& InElementHandle, const FTypedElementListPtr& InSelectionSet, const FTypedElementSelectionOptions& InSelectionOptions) override;
+	virtual bool DeselectElement(const FTypedElementHandle& InElementHandle, const FTypedElementListPtr& InSelectionSet, const FTypedElementSelectionOptions& InSelectionOptions) override;
 };
diff --git a/Engine/Source/Runtime/TypedElementFramework/Private/Elements/Framework/TypedElementList.cpp b/Engine/Source/Runtime/TypedElementFramework/Private/Elements/Framework/TypedElementList.cpp
--- a/Engine/Source/Runtime/TypedElementFramework/Private/Elements/Framework/TypedElementList.cpp
+++ b/Engine/Source/Runtime/TypedElementFramework/Private/Elements/Framework/TypedElementList.cpp
@@ -18,37 +18,55 @@ void GetElementImpl(const UTypedElementRegistry* InRegistry, const FTypedElement
 	}
 }
 
+void GetElementImpl(const UTypedElementRegistry* InRegistry, const FScriptTypedElementHandle& InElementHandle, const TSubclassOf<UInterface>& InBaseInterfaceType, FTypedElement& OutElement)
+{
+	if (InRegistry && InElementHandle)
+	{
+		InRegistry->Private_GetElementImpl(InElementHandle.GetTypedElementHandle(), InBaseInterfaceType, OutElement);
+	}
+	else
+	{
+		OutElement.Release();
+	}
+}
+
 } // namespace TypedElementList_Private
 
 
-FTypedElementListLegacySync::FTypedElementListLegacySync(const FTypedElementList& InElementList)
+template<class HandleType>
+TTypedElementList<HandleType>::FLegacySync::FLegacySync(const TTypedElementList& InElementList)
 	: ElementList(InElementList)
 {
 }
 
-FTypedElementListLegacySync::FOnSyncEvent& FTypedElementListLegacySync::OnSyncEvent()
+template<class HandleType>
+typename TTypedElementList<HandleType>::FLegacySync::FOnSyncEvent&  TTypedElementList<HandleType>::FLegacySync::OnSyncEvent()
 {
 	return OnSyncEventDelegate;
 }
 
-void FTypedElementListLegacySync::Private_EmitSyncEvent(const ESyncType InSyncType, const FTypedElementHandle& InElementHandle)
+template<class HandleType>
+void TTypedElementList<HandleType>::FLegacySync::Private_EmitSyncEvent(const ESyncType InSyncType, const HandleType& InElementHandle)
 {
 	const bool bIsWithinBatchOperation = IsRunningBatchOperation();
 	bBatchOperationIsDirty |= bIsWithinBatchOperation;
 	OnSyncEventDelegate.Broadcast(ElementList, InSyncType, InElementHandle, bIsWithinBatchOperation);
 }
 
-bool FTypedElementListLegacySync::IsRunningBatchOperation() const
+template<class HandleType>
+bool TTypedElementList<HandleType>::FLegacySync::IsRunningBatchOperation() const
 {
 	return NumOpenBatchOperations > 0;
 }
 
-void FTypedElementListLegacySync::BeginBatchOperation()
+template<class HandleType>
+void TTypedElementList<HandleType>::FLegacySync::BeginBatchOperation()
 {
 	++NumOpenBatchOperations;
 }
 
-void FTypedElementListLegacySync::EndBatchOperation(const bool InNotify)
+template<class HandleType>
+void TTypedElementList<HandleType>::FLegacySync::EndBatchOperation(const bool InNotify)
 {
 	checkf(NumOpenBatchOperations > 0, TEXT("Batch operation underflow!"));
 
@@ -65,12 +83,14 @@ void FTypedElementListLegacySync::EndBatchOperation(const bool InNotify)
 	}
 }
 
-bool FTypedElementListLegacySync::IsBatchOperationDirty() const
+template<class HandleType>
+bool TTypedElementList<HandleType>::FLegacySync::IsBatchOperationDirty() const
 {
 	return bBatchOperationIsDirty;
 }
 
-void FTypedElementListLegacySync::ForceBatchOperationDirty()
+template<class HandleType>
+void TTypedElementList<HandleType>::FLegacySync::ForceBatchOperationDirty()
 {
 	if (NumOpenBatchOperations > 0)
 	{
@@ -79,7 +99,8 @@ void FTypedElementListLegacySync::ForceBatchOperationDirty()
 }
 
 
-FTypedElementListLegacySyncScopedBatch::FTypedElementListLegacySyncScopedBatch(const FTypedElementList& InElementList, const bool InNotify)
+template<class HandleType>
+TTypedElementList<HandleType>::FLegacySyncScopedBatch::FLegacySyncScopedBatch(const TTypedElementList& InElementList, const bool InNotify)
 	: ElementListLegacySync(InElementList.Legacy_GetSyncPtr())
 	, bNotify(InNotify)
 {
@@ -89,7 +110,8 @@ FTypedElementListLegacySyncScopedBatch::FTypedElementListLegacySyncScopedBatch(c
 	}
 }
 
-FTypedElementListLegacySyncScopedBatch::~FTypedElementListLegacySyncScopedBatch()
+template<class HandleType>
+TTypedElementList<HandleType>::FLegacySyncScopedBatch::~FLegacySyncScopedBatch()
 {
 	if (ElementListLegacySync)
 	{
@@ -97,13 +119,15 @@ FTypedElementListLegacySyncScopedBatch::~FTypedElementListLegacySyncScopedBatch(
 	}
 }
 
-bool FTypedElementListLegacySyncScopedBatch::IsDirty() const
+template<class HandleType>
+bool TTypedElementList<HandleType>::FLegacySyncScopedBatch::IsDirty() const
 {
 	return ElementListLegacySync
 		&& ElementListLegacySync->IsBatchOperationDirty();
 }
 
-void FTypedElementListLegacySyncScopedBatch::ForceDirty()
+template<class HandleType>
+void TTypedElementList<HandleType>::FLegacySyncScopedBatch::ForceDirty()
 {
 	if (ElementListLegacySync)
 	{
@@ -111,13 +135,14 @@ void FTypedElementListLegacySyncScopedBatch::ForceDirty()
 	}
 }
 
-
-FTypedElementListRef FTypedElementList::Private_CreateElementList(UTypedElementRegistry* InRegistry)
+template<class HandleType>
+typename TTypedElementList<HandleType>::TTypedElementListRef TTypedElementList<HandleType>::Private_CreateElementList(UTypedElementRegistry* InRegistry)
 {
-	return MakeShareable(new FTypedElementList(InRegistry));
+	return MakeShareable(new TTypedElementList(InRegistry));
 }
 
-FTypedElementList::FTypedElementList(UTypedElementRegistry* InRegistry)
+template<class HandleType>
+TTypedElementList<HandleType>::TTypedElementList(UTypedElementRegistry* InRegistry)
 	: Registry(InRegistry)
 {
 	checkf(InRegistry, TEXT("Registry is null!"));
@@ -127,7 +152,8 @@ FTypedElementList::FTypedElementList(UTypedElementRegistry* InRegistry)
 	InRegistry->Private_OnElementListCreated(this);
 }
 
-FTypedElementList::~FTypedElementList()
+template<class HandleType>
+TTypedElementList<HandleType>::~TTypedElementList()
 {
 	Empty();
 	LegacySync.Reset();
@@ -138,17 +164,19 @@ FTypedElementList::~FTypedElementList()
 	}
 }
 
-FTypedElementListRef FTypedElementList::Clone() const
+template<class HandleType>
+typename TTypedElementList<HandleType>::TTypedElementListRef TTypedElementList<HandleType>::Clone() const
 {
-	FTypedElementListRef ClonedElementList = Private_CreateElementList(Registry.Get());
-	for (const FTypedElementHandle& ElementHandle : ElementHandles)
+	TTypedElementListRef ClonedElementList = Private_CreateElementList(Registry.Get());
+	for (const HandleType& ElementHandle : ElementHandles)
 	{
 		ClonedElementList->Add(ElementHandle);
 	}
 	return ClonedElementList;
 }
 
-UObject* FTypedElementList::GetElementInterface(const FTypedElementHandle& InElementHandle, const TSubclassOf<UInterface>& InBaseInterfaceType) const
+template<class HandleType>
+UObject* TTypedElementList<HandleType>::GetElementInterface(const HandleType& InElementHandle, const TSubclassOf<UInterface>& InBaseInterfaceType) const
 {
 	UTypedElementRegistry* RegistryPtr = Registry.Get();
 	return RegistryPtr
@@ -156,13 +184,14 @@ UObject* FTypedElementList::GetElementInterface(const FTypedElementHandle& InEle
 		: nullptr;
 }
 
-bool FTypedElementList::HasElements(const TSubclassOf<UInterface>& InBaseInterfaceType) const
+template<class HandleType>
+bool TTypedElementList<HandleType>::HasElements(const TSubclassOf<UInterface>& InBaseInterfaceType) const
 {
 	bool bHasFilteredElements = false;
 
 	if (InBaseInterfaceType)
 	{
-		ForEachElementHandle([&bHasFilteredElements](const FTypedElementHandle&)
+		ForEachElementHandle([&bHasFilteredElements](const HandleType&)
 		{
 			bHasFilteredElements = true;
 			return false;
@@ -176,13 +205,14 @@ bool FTypedElementList::HasElements(const TSubclassOf<UInterface>& InBaseInterfa
 	return bHasFilteredElements;
 }
 
-int32 FTypedElementList::CountElements(const TSubclassOf<UInterface>& InBaseInterfaceType) const
+template<class HandleType>
+int32 TTypedElementList<HandleType>::CountElements(const TSubclassOf<UInterface>& InBaseInterfaceType) const
 {
 	int32 NumFilteredElements = 0;
 
 	if (InBaseInterfaceType)
 	{
-		ForEachElementHandle([&NumFilteredElements](const FTypedElementHandle&)
+		ForEachElementHandle([&NumFilteredElements](const HandleType&)
 		{
 			++NumFilteredElements;
 			return true;
@@ -196,17 +226,20 @@ int32 FTypedElementList::CountElements(const TSubclassOf<UInterface>& InBaseInte
 	return NumFilteredElements;
 }
 
-bool FTypedElementList::HasElementsOfType(const FName InElementTypeName) const
+template<class HandleType>
+bool TTypedElementList<HandleType>::HasElementsOfType(const FName InElementTypeName) const
 {
 	return CountElementsOfType(InElementTypeName) > 0;
 }
 
-bool FTypedElementList::HasElementsOfType(const FTypedHandleTypeId InElementTypeId) const
+template<class HandleType>
+bool TTypedElementList<HandleType>::HasElementsOfType(const FTypedHandleTypeId InElementTypeId) const
 {
 	return CountElementsOfType(InElementTypeId) > 0;
 }
 
-int32 FTypedElementList::CountElementsOfType(const FName InElementTypeName) const
+template<class HandleType>
+int32 TTypedElementList<HandleType>::CountElementsOfType(const FName InElementTypeName) const
 {
 	if (UTypedElementRegistry* RegistryPtr = Registry.Get())
 	{
@@ -219,17 +252,19 @@ int32 FTypedElementList::CountElementsOfType(const FName InElementTypeName) cons
 	return 0;
 }
 
-int32 FTypedElementList::CountElementsOfType(const FTypedHandleTypeId InElementTypeId) const
+template<class HandleType>
+int32 TTypedElementList<HandleType>::CountElementsOfType(const FTypedHandleTypeId InElementTypeId) const
 {
 	return ElementCounts.GetCounterValue(FTypedElementCounter::GetElementTypeCategoryName(), InElementTypeId);
 }
 
-TArray<FTypedElementHandle> FTypedElementList::GetElementHandles(const TSubclassOf<UInterface>& InBaseInterfaceType) const
+template<class HandleType>
+TArray<HandleType> TTypedElementList<HandleType>::GetElementHandles(const TSubclassOf<UInterface>& InBaseInterfaceType) const
 {
-	TArray<FTypedElementHandle> FilteredElementHandles;
+	TArray<HandleType> FilteredElementHandles;
 	FilteredElementHandles.Reserve(ElementHandles.Num());
 
-	ForEachElementHandle([&FilteredElementHandles](const FTypedElementHandle& InElementHandle)
+	ForEachElementHandle([&FilteredElementHandles](const HandleType& InElementHandle)
 	{
 		FilteredElementHandles.Add(InElementHandle);
 		return true;
@@ -238,9 +273,10 @@ TArray<FTypedElementHandle> FTypedElementList::GetElementHandles(const TSubclass
 	return FilteredElementHandles;
 }
 
-void FTypedElementList::ForEachElementHandle(TFunctionRef<bool(const FTypedElementHandle&)> InCallback, const TSubclassOf<UInterface>& InBaseInterfaceType) const
+template<class HandleType>
+void TTypedElementList<HandleType>::ForEachElementHandle(TFunctionRef<bool(const HandleType&)> InCallback, const TSubclassOf<UInterface>& InBaseInterfaceType) const
 {
-	for (const FTypedElementHandle& ElementHandle : ElementHandles)
+	for (const HandleType& ElementHandle : ElementHandles)
 	{
 		if (ElementHandle && (!InBaseInterfaceType || GetElementInterface(ElementHandle, InBaseInterfaceType)))
 		{
@@ -252,7 +288,8 @@ void FTypedElementList::ForEachElementHandle(TFunctionRef<bool(const FTypedEleme
 	}
 }
 
-bool FTypedElementList::AddElementImpl(FTypedElementHandle&& InElementHandle)
+template<class HandleType>
+bool TTypedElementList<HandleType>::AddElementImpl(HandleType&& InElementHandle)
 {
 	if (!InElementHandle)
 	{
@@ -266,15 +303,24 @@ bool FTypedElementList::AddElementImpl(FTypedElementHandle&& InElementHandle)
 
 	if (!bAlreadyAdded)
 	{
-		const FTypedElementHandle& AddedElementHandle = ElementHandles.Add_GetRef(MoveTemp(InElementHandle));
-		ElementCounts.AddElement(AddedElementHandle);
+		const HandleType& AddedElementHandle = ElementHandles.Add_GetRef(MoveTemp(InElementHandle));
+
+		if constexpr (std::is_same<HandleType, FTypedElementHandle>::value)
+		{
+			ElementCounts.AddElement(AddedElementHandle);
+		}
+		else
+		{
+			ElementCounts.AddElement(AddedElementHandle.GetTypedElementHandle());
+		}
 		NoteListChanged(EChangeType::Added, AddedElementHandle);
 	}
 
 	return !bAlreadyAdded;
 }
 
-bool FTypedElementList::RemoveElementImpl(const FTypedElementId& InElementId)
+template<class HandleType>
+bool TTypedElementList<HandleType>::RemoveElementImpl(const FTypedElementId& InElementId)
 {
 	if (!InElementId)
 	{
@@ -287,29 +333,38 @@ bool FTypedElementList::RemoveElementImpl(const FTypedElementId& InElementId)
 
 	if (bRemoved)
 	{
-		const int32 ElementHandleIndexToRemove = ElementHandles.IndexOfByPredicate([&InElementId](const FTypedElementHandle& InElementHandle)
+		const int32 ElementHandleIndexToRemove = ElementHandles.IndexOfByPredicate([&InElementId](const HandleType& InElementHandle)
 		{
 			return InElementHandle.GetId() == InElementId;
 		});
 		checkSlow(ElementHandleIndexToRemove != INDEX_NONE);
 
-		FTypedElementHandle RemovedElementHandle = MoveTemp(ElementHandles[ElementHandleIndexToRemove]);
+		HandleType RemovedElementHandle = MoveTemp(ElementHandles[ElementHandleIndexToRemove]);
 		ElementHandles.RemoveAt(ElementHandleIndexToRemove, 1, /*bAllowShrinking*/false);
-		ElementCounts.RemoveElement(RemovedElementHandle);
 
+		if constexpr (std::is_same<HandleType, FTypedElementHandle>::value)
+		{
+			ElementCounts.RemoveElement(RemovedElementHandle);
+		}
+		else
+		{
+			ElementCounts.RemoveElement(RemovedElementHandle.GetTypedElementHandle());
+		}
+	
 		NoteListChanged(EChangeType::Removed, RemovedElementHandle);
 	}
 
 	return bRemoved;
 }
 
-int32 FTypedElementList::RemoveAllElementsImpl(TFunctionRef<bool(const FTypedElementHandle&)> InPredicate)
+template<class HandleType>
+int32 TTypedElementList<HandleType>::RemoveAllElementsImpl(TFunctionRef<bool(const HandleType&)> InPredicate)
 {
 	int32 RemovedCount = 0;
 
 	if (ElementHandles.Num() > 0)
 	{
-		FTypedElementListLegacySyncScopedBatch LegacySyncBatch(*this);
+		FLegacySyncScopedBatch LegacySyncBatch(*this);
 
 		NoteListMayChange();
 
@@ -317,10 +372,18 @@ int32 FTypedElementList::RemoveAllElementsImpl(TFunctionRef<bool(const FTypedEle
 		{
 			if (InPredicate(ElementHandles[Index]))
 			{
-				FTypedElementHandle RemovedElementHandle = MoveTemp(ElementHandles[Index]);
+				HandleType RemovedElementHandle = MoveTemp(ElementHandles[Index]);
 				ElementCombinedIds.Remove(RemovedElementHandle.GetId().GetCombinedId());
 				ElementHandles.RemoveAt(Index, 1, /*bAllowShrinking*/false);
-				ElementCounts.RemoveElement(RemovedElementHandle);
+
+				if constexpr (std::is_same<HandleType, FTypedElementHandle>::value)
+				{
+					ElementCounts.RemoveElement(RemovedElementHandle);
+				}
+				else
+				{
+					ElementCounts.RemoveElement(RemovedElementHandle.GetTypedElementHandle());
+				}
 
 				NoteListChanged(EChangeType::Removed, RemovedElementHandle);
 
@@ -332,27 +395,31 @@ int32 FTypedElementList::RemoveAllElementsImpl(TFunctionRef<bool(const FTypedEle
 	return RemovedCount;
 }
 
-bool FTypedElementList::ContainsElementImpl(const FTypedElementId& InElementId) const
+template<class HandleType>
+bool TTypedElementList<HandleType>::ContainsElementImpl(const FTypedElementId& InElementId) const
 {
 	return InElementId 
 		&& ElementCombinedIds.Contains(InElementId.GetCombinedId());
 }
 
-FTypedElementListLegacySync& FTypedElementList::Legacy_GetSync()
+template<class HandleType>
+typename TTypedElementList<HandleType>::FLegacySync& TTypedElementList<HandleType>::Legacy_GetSync()
 {
 	if (!LegacySync)
 	{
-		LegacySync = MakeUnique<FTypedElementListLegacySync>(*this);
+		LegacySync = MakeUnique<FLegacySync>(*this);
 	}
 	return *LegacySync;
 }
 
-FTypedElementListLegacySync* FTypedElementList::Legacy_GetSyncPtr() const
+template<class HandleType>
+typename TTypedElementList<HandleType>::FLegacySync* TTypedElementList<HandleType>::Legacy_GetSyncPtr() const
 {
 	return LegacySync.Get();
 }
 
-bool FTypedElementList::NotifyPendingChanges()
+template<class HandleType>
+bool TTypedElementList<HandleType>::NotifyPendingChanges()
 {
 	if (bHasPendingNotify)
 	{
@@ -364,7 +431,8 @@ bool FTypedElementList::NotifyPendingChanges()
 	return false;
 }
 
-FTypedElementList::FScopedClearNewPendingChange::FScopedClearNewPendingChange(FTypedElementList& InTypeElementList)
+template<class HandleType>
+TTypedElementList<HandleType>::FScopedClearNewPendingChange::FScopedClearNewPendingChange(TTypedElementList& InTypeElementList)
 {
 	bool bCanClearNewPendingChange = !InTypeElementList.bHasPendingNotify;
 
@@ -379,14 +447,15 @@ FTypedElementList::FScopedClearNewPendingChange::FScopedClearNewPendingChange(FT
 	}
 }
 
-
-FTypedElementList::FScopedClearNewPendingChange::FScopedClearNewPendingChange(FScopedClearNewPendingChange&& Other)
+template<class HandleType>
+TTypedElementList<HandleType>::FScopedClearNewPendingChange::FScopedClearNewPendingChange(FScopedClearNewPendingChange&& Other)
 	: TypedElementList(Other.TypedElementList)
 {
 	Other.TypedElementList = nullptr;
 }
 
-FTypedElementList::FScopedClearNewPendingChange& FTypedElementList::FScopedClearNewPendingChange::operator=(FScopedClearNewPendingChange&& Other)
+template<class HandleType>
+typename TTypedElementList<HandleType>::FScopedClearNewPendingChange& TTypedElementList<HandleType>::FScopedClearNewPendingChange::operator=(FScopedClearNewPendingChange&& Other)
 {
 	TypedElementList = Other.TypedElementList;
 	Other.TypedElementList = nullptr;
@@ -394,7 +463,8 @@ FTypedElementList::FScopedClearNewPendingChange& FTypedElementList::FScopedClear
 	return *this;
 }
 
-FTypedElementList::FScopedClearNewPendingChange::~FScopedClearNewPendingChange()
+template<class HandleType>
+TTypedElementList<HandleType>::FScopedClearNewPendingChange::~FScopedClearNewPendingChange()
 {
 	if (TypedElementList)
 	{
@@ -402,35 +472,38 @@ FTypedElementList::FScopedClearNewPendingChange::~FScopedClearNewPendingChange()
 	}
 }
 
-FTypedElementList::FScopedClearNewPendingChange FTypedElementList::GetScopedClearNewPendingChange()
+template<class HandleType>
+typename TTypedElementList<HandleType>::FScopedClearNewPendingChange TTypedElementList<HandleType>::GetScopedClearNewPendingChange()
 {
 	return FScopedClearNewPendingChange(*this);
 }
 
-void FTypedElementList::NoteListMayChange()
+template<class HandleType>
+void TTypedElementList<HandleType>::NoteListMayChange()
 {
 	OnPreChangeDelegate.Broadcast(*this);
 }
 
-void FTypedElementList::NoteListChanged(const EChangeType InChangeType, const FTypedElementHandle& InElementHandle)
+template<class HandleType>
+void TTypedElementList<HandleType>::NoteListChanged(const EChangeType InChangeType, const HandleType& InElementHandle)
 {
 	bHasPendingNotify = true;
 
 	if (LegacySync)
 	{
-		FTypedElementListLegacySync::ESyncType SyncType = FTypedElementListLegacySync::ESyncType::Modified;
+		typename FLegacySync::ESyncType SyncType = FLegacySync::ESyncType::Modified;
 		switch (InChangeType)
 		{
 		case EChangeType::Added:
-			SyncType = FTypedElementListLegacySync::ESyncType::Added;
+			SyncType = FLegacySync::ESyncType::Added;
 			break;
 
 		case EChangeType::Removed:
-			SyncType = FTypedElementListLegacySync::ESyncType::Removed;
+			SyncType = FLegacySync::ESyncType::Removed;
 			break;
 
 		case EChangeType::Cleared:
-			SyncType = FTypedElementListLegacySync::ESyncType::Cleared;
+			SyncType = FLegacySync::ESyncType::Cleared;
 			break;
 
 		default:
@@ -440,3 +513,63 @@ void FTypedElementList::NoteListChanged(const EChangeType InChangeType, const FT
 		LegacySync->Private_EmitSyncEvent(SyncType, InElementHandle);
 	}
 }
+
+template class TYPEDELEMENTFRAMEWORK_API TTypedElementList<FTypedElementHandle>;
+template class TYPEDELEMENTFRAMEWORK_API TTypedElementList<FScriptTypedElementHandle>;
+
+
+namespace UE::TypedElementFramework
+{
+	/**
+	 * Functions to convert a script list to a native one
+	 * @note Only copies elements; does not copy any bindings!
+	 */
+	FTypedElementListPtr ConvertToNativeTypedElementList(const FScriptTypedElementListConstPtr& ScriptList)
+	{
+		if (ScriptList)
+		{
+			if (UTypedElementRegistry* RawRegistry = ScriptList->GetRegistry())
+			{
+				FTypedElementListPtr NativeList = FTypedElementList::Private_CreateElementList(RawRegistry);
+				NativeList->Reserve(ScriptList->CountElements());
+				ScriptList->ForEachElementHandle([&NativeList](const FScriptTypedElementHandle& ScriptHandle)
+					{
+						if (FTypedElementHandle NativeHandle = ScriptHandle.GetTypedElementHandle())
+						{
+							NativeList->Add(NativeHandle);
+						}
+						return true;
+					});
+
+				return NativeList;
+			}
+		}
+
+		return {};
+	}
+
+	/**
+	 * Functions to convert a script list to a native one
+	 * @note Only copies elements; does not copy any bindings!
+	 */
+	FScriptTypedElementListPtr ConvertToScriptTypedElementList(const FTypedElementListConstPtr& NativeList)
+	{
+		if (NativeList)
+		{
+			if (UTypedElementRegistry* RawRegistry = NativeList->GetRegistry())
+			{
+				FScriptTypedElementListPtr ScriptList = FScriptTypedElementList::Private_CreateElementList(RawRegistry);
+				ScriptList->Reserve(NativeList->CountElements());
+				NativeList->ForEachElementHandle([&ScriptList, &RawRegistry](const FTypedElementHandle& NativeHandle)
+					{
+						ScriptList->Add(RawRegistry->CreateScriptHandle(NativeHandle.GetId()));
+						return true;
+					});
+
+				return ScriptList;
+			}
+		}
+
+		return {};
+	}
+}
diff --git a/Engine/Source/Runtime/TypedElementFramework/Private/Elements/Framework/TypedElementListLibrary.h b/Engine/Source/Runtime/TypedElementFramework/Private/Elements/Framework/TypedElementListLibrary.h
--- a/Engine/Source/Runtime/TypedElementFramework/Private/Elements/Framework/TypedElementListLibrary.h
+++ b/Engine/Source/Runtime/TypedElementFramework/Private/Elements/Framework/TypedElementListLibrary.h
@@ -17,10 +17,10 @@ public:
 	/**
 	 * Create an empty list of elements associated with the given registry.
 	 */
-	UFUNCTION(BlueprintCallable, Category="TypedElementFramework|Registry", meta=(ScriptMethod))
-	static FTypedElementListProxy CreateElementList(UTypedElementRegistry* Registry)
+	UFUNCTION(BlueprintCallable, DisplayName="Create Element List", Category = "TypedElementFramework|Registry", meta=(ScriptMethod, ScriptName="CreateElementList"))
+	static FManagedTypedElementListProxy CreateScriptElementList(UTypedElementRegistry* Registry)
 	{
-		return Registry->CreateElementList();
+		return Registry->CreateScriptElementList();
 	}
 
 	/**
@@ -28,10 +28,10 @@ public:
 	 * @note Only copies elements; does not copy any bindings!
 	 */
 	UFUNCTION(BlueprintCallable, Category="TypedElementFramework|List", meta=(ScriptMethod))
-	static FTypedElementListProxy Clone(const FTypedElementListProxy ElementList)
+	static FManagedTypedElementListProxy Clone(const FManagedTypedElementListProxy ElementList)
 	{
-		FTypedElementListConstPtr ElementListPtr = ElementList.GetElementList();
-		return ElementListPtr ? ElementListPtr->Clone() : FTypedElementListProxy();
+		FScriptTypedElementListConstPtr ElementListPtr = ElementList.GetElementList();
+		return ElementListPtr ? ElementListPtr->Clone() : FManagedTypedElementListProxy();
 	}
 
 	/**
@@ -39,19 +39,19 @@ public:
 	 * @note Use IsValidIndex to test for validity.
 	 */
 	UFUNCTION(BlueprintPure, Category="TypedElementFramework|List", meta=(ScriptMethod))
-	static FTypedElementHandle GetElementHandleAt(const FTypedElementListProxy ElementList, const int32 Index)
+	static FScriptTypedElementHandle GetElementHandleAt(const FManagedTypedElementListProxy ElementList, const int32 Index)
 	{
-		FTypedElementListConstPtr ElementListPtr = ElementList.GetElementList();
-		return ElementListPtr ? ElementListPtr->GetElementHandleAt(Index) : FTypedElementHandle();
+		FScriptTypedElementListConstPtr ElementListPtr = ElementList.GetElementList();
+		return ElementListPtr ? ElementListPtr->GetElementHandleAt(Index) : FScriptTypedElementHandle();
 	}
 
 	/**
 	 * Get the element interface from the given handle.
 	 */
 	UFUNCTION(BlueprintPure, Category="TypedElementFramework|List", meta=(ScriptMethod))
-	static UObject* GetElementInterface(const FTypedElementListProxy ElementList, const FTypedElementHandle& ElementHandle, const TSubclassOf<UInterface> BaseInterfaceType)
+	static UObject* GetElementInterface(const FManagedTypedElementListProxy ElementList, const FScriptTypedElementHandle& ElementHandle, const TSubclassOf<UInterface> BaseInterfaceType)
 	{
-		FTypedElementListConstPtr ElementListPtr = ElementList.GetElementList();
+		FScriptTypedElementListConstPtr ElementListPtr = ElementList.GetElementList();
 		return ElementListPtr ? ElementListPtr->GetElementInterface(ElementHandle, BaseInterfaceType) : nullptr;
 	}
 
@@ -59,9 +59,9 @@ public:
 	 * Test whether there are elements in this list, optionally filtering to elements that implement the given interface.
 	 */
 	UFUNCTION(BlueprintCallable, Category="TypedElementFramework|List", meta=(ScriptMethod))
-	static bool HasElements(const FTypedElementListProxy ElementList, const TSubclassOf<UInterface> BaseInterfaceType = nullptr)
+	static bool HasElements(const FManagedTypedElementListProxy ElementList, const TSubclassOf<UInterface> BaseInterfaceType = nullptr)
 	{
-		FTypedElementListConstPtr ElementListPtr = ElementList.GetElementList();
+		FScriptTypedElementListConstPtr ElementListPtr = ElementList.GetElementList();
 		return ElementListPtr && ElementListPtr->HasElements(BaseInterfaceType);
 	}
 
@@ -69,9 +69,9 @@ public:
 	 * Count the number of elements in this list, optionally filtering to elements that implement the given interface.
 	 */
 	UFUNCTION(BlueprintCallable, Category="TypedElementFramework|List", meta=(ScriptMethod))
-	static int32 CountElements(const FTypedElementListProxy ElementList, const TSubclassOf<UInterface> BaseInterfaceType = nullptr)
+	static int32 CountElements(const FManagedTypedElementListProxy ElementList, const TSubclassOf<UInterface> BaseInterfaceType = nullptr)
 	{
-		FTypedElementListConstPtr ElementListPtr = ElementList.GetElementList();
+		FScriptTypedElementListConstPtr ElementListPtr = ElementList.GetElementList();
 		return ElementListPtr ? ElementListPtr->CountElements(BaseInterfaceType) : 0;
 	}
 
@@ -79,9 +79,9 @@ public:
 	 * Test whether there are elements in this list of the given type.
 	 */
 	UFUNCTION(BlueprintCallable, Category="TypedElementFramework|List", meta=(ScriptMethod))
-	static bool HasElementsOfType(const FTypedElementListProxy ElementList, const FName ElementTypeName)
+	static bool HasElementsOfType(const FManagedTypedElementListProxy ElementList, const FName ElementTypeName)
 	{
-		FTypedElementListConstPtr ElementListPtr = ElementList.GetElementList();
+		FScriptTypedElementListConstPtr ElementListPtr = ElementList.GetElementList();
 		return ElementListPtr && ElementListPtr->HasElementsOfType(ElementTypeName);
 	}
 
@@ -89,9 +89,9 @@ public:
 	 * Count the number of elements in this list of the given type.
 	 */
 	UFUNCTION(BlueprintCallable, Category="TypedElementFramework|List", meta=(ScriptMethod))
-	static int32 CountElementsOfType(const FTypedElementListProxy ElementList, const FName ElementTypeName)
+	static int32 CountElementsOfType(const FManagedTypedElementListProxy ElementList, const FName ElementTypeName)
 	{
-		FTypedElementListConstPtr ElementListPtr = ElementList.GetElementList();
+		FScriptTypedElementListConstPtr ElementListPtr = ElementList.GetElementList();
 		return ElementListPtr ? ElementListPtr->CountElementsOfType(ElementTypeName) : 0;
 	}
 
@@ -99,19 +99,19 @@ public:
 	 * Get the handle of every element in this list, optionally filtering to elements that implement the given interface.
 	 */
 	UFUNCTION(BlueprintCallable, Category="TypedElementFramework|List", meta=(ScriptMethod))
-	static TArray<FTypedElementHandle> GetElementHandles(const FTypedElementListProxy ElementList, const TSubclassOf<UInterface> BaseInterfaceType)
+	static TArray<FScriptTypedElementHandle> GetElementHandles(const FManagedTypedElementListProxy ElementList, const TSubclassOf<UInterface> BaseInterfaceType)
 	{
-		FTypedElementListConstPtr ElementListPtr = ElementList.GetElementList();
-		return ElementListPtr ? ElementListPtr->GetElementHandles(BaseInterfaceType) : TArray<FTypedElementHandle>();
+		FScriptTypedElementListConstPtr ElementListPtr = ElementList.GetElementList();
+		return ElementListPtr ? ElementListPtr->GetElementHandles(BaseInterfaceType) : TArray<FScriptTypedElementHandle>();
 	}
 
 	/**
 	 * Is the given index a valid entry within this element list?
 	 */
 	UFUNCTION(BlueprintPure, Category="TypedElementFramework|List", meta=(ScriptMethod))
-	static bool IsValidIndex(const FTypedElementListProxy ElementList, const int32 Index)
+	static bool IsValidIndex(const FManagedTypedElementListProxy ElementList, const int32 Index)
 	{
-		FTypedElementListConstPtr ElementListPtr = ElementList.GetElementList();
+		FScriptTypedElementListConstPtr ElementListPtr = ElementList.GetElementList();
 		return ElementListPtr && ElementListPtr->IsValidIndex(Index);
 	}
 
@@ -119,9 +119,9 @@ public:
 	 * Get the number of entries within this element list.
 	 */
 	UFUNCTION(BlueprintPure, Category="TypedElementFramework|List", meta=(ScriptMethod))
-	static int32 Num(const FTypedElementListProxy ElementList)
+	static int32 Num(const FManagedTypedElementListProxy ElementList)
 	{
-		FTypedElementListConstPtr ElementListPtr = ElementList.GetElementList();
+		FScriptTypedElementListConstPtr ElementListPtr = ElementList.GetElementList();
 		return ElementListPtr ? ElementListPtr->Num() : 0;
 	}
 
@@ -129,9 +129,9 @@ public:
 	 * Shrink this element list storage to avoid slack.
 	 */
 	UFUNCTION(BlueprintCallable, Category="TypedElementFramework|List", meta=(ScriptMethod))
-	static void Shrink(FTypedElementListProxy ElementList)
+	static void Shrink(FManagedTypedElementListProxy ElementList)
 	{
-		if (FTypedElementListPtr ElementListPtr = ElementList.GetElementList())
+		if (FScriptTypedElementListPtr ElementListPtr = ElementList.GetElementList())
 		{
 			ElementListPtr->Shrink();
 		}
@@ -141,9 +141,9 @@ public:
 	 * Pre-allocate enough memory in this element list to store the given number of entries.
 	 */
 	UFUNCTION(BlueprintCallable, Category="TypedElementFramework|List", meta=(ScriptMethod))
-	static void Reserve(FTypedElementListProxy ElementList, const int32 Size)
+	static void Reserve(FManagedTypedElementListProxy ElementList, const int32 Size)
 	{
-		if (FTypedElementListPtr ElementListPtr = ElementList.GetElementList())
+		if (FScriptTypedElementListPtr ElementListPtr = ElementList.GetElementList())
 		{
 			ElementListPtr->Reserve(Size);
 		}
@@ -153,9 +153,9 @@ public:
 	 * Remove all entries from this element list, potentially leaving space allocated for the given number of entries.
 	 */
 	UFUNCTION(BlueprintCallable, Category="TypedElementFramework|List", meta=(ScriptMethod))
-	static void Empty(FTypedElementListProxy ElementList, const int32 Slack = 0)
+	static void Empty(FManagedTypedElementListProxy ElementList, const int32 Slack = 0)
 	{
-		if (FTypedElementListPtr ElementListPtr = ElementList.GetElementList())
+		if (FScriptTypedElementListPtr ElementListPtr = ElementList.GetElementList())
 		{
 			ElementListPtr->Empty(Slack);
 		}
@@ -165,9 +165,9 @@ public:
 	 * Remove all entries from this element list, preserving existing allocations.
 	 */
 	UFUNCTION(BlueprintCallable, Category="TypedElementFramework|List", meta=(ScriptMethod))
-	static void Reset(FTypedElementListProxy ElementList)
+	static void Reset(FManagedTypedElementListProxy ElementList)
 	{
-		if (FTypedElementListPtr ElementListPtr = ElementList.GetElementList())
+		if (FScriptTypedElementListPtr ElementListPtr = ElementList.GetElementList())
 		{
 			ElementListPtr->Reset();
 		}
@@ -177,9 +177,9 @@ public:
 	 * Does this element list contain an entry for the given element handle?
 	 */
 	UFUNCTION(BlueprintPure, Category="TypedElementFramework|List", meta=(ScriptMethod))
-	static bool Contains(const FTypedElementListProxy ElementList, const FTypedElementHandle& ElementHandle)
+	static bool Contains(const FManagedTypedElementListProxy ElementList, const FScriptTypedElementHandle& ElementHandle)
 	{
-		FTypedElementListConstPtr ElementListPtr = ElementList.GetElementList();
+		FScriptTypedElementListConstPtr ElementListPtr = ElementList.GetElementList();
 		return ElementListPtr && ElementListPtr->Contains(ElementHandle);
 	}
 
@@ -188,9 +188,9 @@ public:
 	 * @return True if the element handle was added, false if it is already in the list.
 	 */
 	UFUNCTION(BlueprintCallable, Category="TypedElementFramework|List", meta=(ScriptMethod))
-	static bool Add(FTypedElementListProxy ElementList, const FTypedElementHandle& ElementHandle)
+	static bool Add(FManagedTypedElementListProxy ElementList, const FScriptTypedElementHandle& ElementHandle)
 	{
-		FTypedElementListPtr ElementListPtr = ElementList.GetElementList();
+		FScriptTypedElementListPtr ElementListPtr = ElementList.GetElementList();
 		return ElementListPtr && ElementListPtr->Add(ElementHandle);
 	}
 
@@ -198,9 +198,9 @@ public:
 	 * Append the given element handles to this element list.
 	 */
 	UFUNCTION(BlueprintCallable, Category="TypedElementFramework|List", meta=(ScriptMethod))
-	static void Append(FTypedElementListProxy ElementList, const TArray<FTypedElementHandle>& ElementHandles)
+	static void Append(FManagedTypedElementListProxy ElementList, const TArray<FScriptTypedElementHandle>& ElementHandles)
 	{
-		if (FTypedElementListPtr ElementListPtr = ElementList.GetElementList())
+		if (FScriptTypedElementListPtr ElementListPtr = ElementList.GetElementList())
 		{
 			ElementListPtr->Append(ElementHandles);
 		}
@@ -210,10 +210,10 @@ public:
 	 * Append the another element list to this element list.
 	 */
 	UFUNCTION(BlueprintCallable, Category="TypedElementFramework|List", meta=(ScriptMethod))
-	static void AppendList(FTypedElementListProxy ElementList, const FTypedElementListProxy OtherElementList)
+	static void AppendList(FManagedTypedElementListProxy ElementList, const FManagedTypedElementListProxy OtherElementList)
 	{
-		FTypedElementListPtr ElementListPtr = ElementList.GetElementList();
-		FTypedElementListConstPtr OtherElementListPtr = OtherElementList.GetElementList();
+		FScriptTypedElementListPtr ElementListPtr = ElementList.GetElementList();
+		FScriptTypedElementListConstPtr OtherElementListPtr = OtherElementList.GetElementList();
 		if (ElementListPtr && OtherElementListPtr)
 		{
 			ElementListPtr->Append(OtherElementListPtr.ToSharedRef());
@@ -225,9 +225,9 @@ public:
 	 * @return True if the element handle was removed, false if it isn't in the list.
 	 */
 	UFUNCTION(BlueprintCallable, Category="TypedElementFramework|List", meta=(ScriptMethod))
-	static bool Remove(FTypedElementListProxy ElementList, const FTypedElementHandle& ElementHandle)
+	static bool Remove(FManagedTypedElementListProxy ElementList, const FScriptTypedElementHandle& ElementHandle)
 	{
-		FTypedElementListPtr ElementListPtr = ElementList.GetElementList();
+		FScriptTypedElementListPtr ElementListPtr = ElementList.GetElementList();
 		return ElementListPtr && ElementListPtr->Remove(ElementHandle);
 	}
 };
diff --git a/Engine/Source/Runtime/TypedElementFramework/Private/Elements/Framework/TypedElementRegistry.cpp b/Engine/Source/Runtime/TypedElementFramework/Private/Elements/Framework/TypedElementRegistry.cpp
--- a/Engine/Source/Runtime/TypedElementFramework/Private/Elements/Framework/TypedElementRegistry.cpp
+++ b/Engine/Source/Runtime/TypedElementFramework/Private/Elements/Framework/TypedElementRegistry.cpp
@@ -267,23 +267,47 @@ FString UTypedElementRegistry::RegistredElementTypesAndInterfacesToString() cons
 
 void UTypedElementRegistry::NotifyElementListPendingChanges()
 {
-	/** 
-	 * We use a critical section here since we need the called function to be able to create or delete TypedElementLists.
-	 * Critical sections are recursive. They can be lock multiple times by the same thread without blocking.
-	 */
-	FScopeLock ActiveElementListsLock(&ActiveElementListsCS);
-	TArray<FTypedElementList*> ElementListToNotify = ActiveElementLists.Array();
+	{
+		/** 
+		 * We use a critical section here since we need the called function to be able to create or delete TypedElementLists.
+		 * Critical sections are recursive. They can be lock multiple times by the same thread without blocking.
+		 */
+		FScopeLock ActiveElementListsLock(&ActiveElementListsCS);
+		TArray<FTypedElementList*> ElementListToNotify = ActiveElementLists.Array();
+
+		bool bHasListPotentialyChanged = false;
+		for (FTypedElementList* ActiveElementList : ElementListToNotify)
+		{
+			if (bHasListPotentialyChanged)
+			{
+				/**
+				 * One of the callbacks could have modified the ActiveElementLists by deleting a ElementList.
+				 * So we need to validate that the ActiveElementList still exist.
+				 */
+				if (ActiveElementLists.Contains(ActiveElementList))
+				{
+					ActiveElementList->NotifyPendingChanges();
+				}
+			}
+			else
+			{
+				bHasListPotentialyChanged = ActiveElementList->NotifyPendingChanges();
+			}
+		}
+	}
 
+	// Script List (no need for a critical section)
 	bool bHasListPotentialyChanged = false;
-	for (FTypedElementList* ActiveElementList : ElementListToNotify)
+	TArray<FScriptTypedElementList*> ElementListToNotify = ActiveScriptElementLists.Array();
+	for (FScriptTypedElementList* ActiveElementList : ElementListToNotify)
 	{
 		if (bHasListPotentialyChanged)
 		{
 			/**
-			 * One of the callbacks could have modified the ActiveElementLists by deleting a ElementList.
-			 * So we need to validate that the ActiveElementList still exist.
-			 */
-			if (ActiveElementLists.Contains(ActiveElementList))
+				* One of the callbacks could have modified the ActiveElementLists by deleting a ElementList.
+				* So we need to validate that the ActiveElementList still exist.
+				*/
+			if (ActiveScriptElementLists.Contains(ActiveElementList))
 			{
 				ActiveElementList->NotifyPendingChanges();
 			}
diff --git a/Engine/Source/Runtime/TypedElementFramework/Private/Elements/Framework/TypedElementUtil.cpp b/Engine/Source/Runtime/TypedElementFramework/Private/Elements/Framework/TypedElementUtil.cpp
--- a/Engine/Source/Runtime/TypedElementFramework/Private/Elements/Framework/TypedElementUtil.cpp
+++ b/Engine/Source/Runtime/TypedElementFramework/Private/Elements/Framework/TypedElementUtil.cpp
@@ -1,7 +1,9 @@
 // Copyright Epic Games, Inc. All Rights Reserved.
 
 #include "Elements/Framework/TypedElementUtil.h"
+
 #include "Elements/Framework/TypedElementList.h"
+#include "Elements/Framework/TypedElementRegistry.h"
 
 namespace TypedElementUtil
 {
@@ -27,4 +29,32 @@ void BatchElementsByType(TArrayView<const FTypedElementHandle> InElementsToBatch
 	}
 }
 
+TArray<FScriptTypedElementHandle> ConvertToScriptElementArray(const TArray<FTypedElementHandle>& InNativeHandles, UTypedElementRegistry* Registry)
+{
+	TArray<FScriptTypedElementHandle> ScriptHandles;
+	ScriptHandles.Reserve(InNativeHandles.Num());
+	for (const FTypedElementHandle& NativeHandle : InNativeHandles)
+	{
+		ScriptHandles.Add(Registry->CreateScriptHandle(NativeHandle.GetId()));
+	}
+	
+	return ScriptHandles;
+}
+
+TArray<FTypedElementHandle> ConvertToNativeElementArray(const TArray<FScriptTypedElementHandle>& InScriptHandles)
+{
+	TArray<FTypedElementHandle> NativeHandles;
+	NativeHandles.Reserve(InScriptHandles.Num());
+
+	for (const FScriptTypedElementHandle& ScriptHandle : InScriptHandles)
+	{
+		if (FTypedElementHandle NativeHandle = ScriptHandle.GetTypedElementHandle())
+		{
+			NativeHandles.Add(MoveTemp(NativeHandle));
+		}
+	}
+
+	return NativeHandles;
+}
+
 } // namespace TypedElementUtil
diff --git a/Engine/Source/Runtime/TypedElementFramework/Private/TypedElementFrameworkTests.cpp b/Engine/Source/Runtime/TypedElementFramework/Private/TypedElementFrameworkTests.cpp
--- a/Engine/Source/Runtime/TypedElementFramework/Private/TypedElementFrameworkTests.cpp
+++ b/Engine/Source/Runtime/TypedElementFramework/Private/TypedElementFrameworkTests.cpp
@@ -142,7 +142,7 @@ bool FTypedElementRegistrySmokeTest::RunTest(const FString& Parameters)
 	};
 
 	const FName DummyElementType_Typed = "DummyElementType_Typed";
-	Registry->RegisterElementType<FTestTypedElementData>(DummyElementType_Typed);
+	Registry->RegisterElementType<FTestTypedElementData, true>(DummyElementType_Typed);
 	Registry->RegisterElementInterface<ITestTypedElementInterfaceA>(DummyElementType_Typed, NewObject<UTestTypedElementInterfaceA_ImplTyped>());
 
 	UObject* TypedElementInterfaceBAndC =  NewObject<UTestTypedElementInterfaceBAndC_Typed>();
@@ -151,7 +151,7 @@ bool FTypedElementRegistrySmokeTest::RunTest(const FString& Parameters)
 
 
 	const FName DummyElementType_Untyped = "DummyElementType_Untyped";
-	Registry->RegisterElementType(DummyElementType_Untyped);
+	Registry->RegisterElementType(DummyElementType_Untyped, true);
 	Registry->RegisterElementInterface<ITestTypedElementInterfaceA>(DummyElementType_Untyped, NewObject<UTestTypedElementInterfaceA_ImplUntyped>());
 
 	TTypedElementOwner<FTestTypedElementData> TypedElement1 = Registry->CreateElement<FTestTypedElementData>(DummyElementType_Typed);
@@ -169,6 +169,7 @@ bool FTypedElementRegistrySmokeTest::RunTest(const FString& Parameters)
 	TestInterfaceAccess(UntypedElement1.AcquireHandle());
 
 	FTypedElementListPtr ElementList = Registry->CreateElementList();
+	ElementList->Reserve(6);
 	ElementList->Add(TypedElement1);
 	ElementList->Add(TypedElement2);
 	ElementList->Add(TypedElement3);
@@ -197,6 +198,14 @@ bool FTypedElementRegistrySmokeTest::RunTest(const FString& Parameters)
 
 	ElementList.Reset();
 
+
+	// Test the weak handles 
+	FScriptTypedElementHandle ScriptHandleTyped = Registry->CreateScriptHandle(TypedElement1.GetId());
+	check(ScriptHandleTyped.IsSet());
+
+	FScriptTypedElementHandle ScriptHandleUntyped = Registry->CreateScriptHandle(UntypedElement1.GetId());
+	check(ScriptHandleUntyped.IsSet());
+
 	Registry->DestroyElement(TypedElement1);
 	Registry->DestroyElement(TypedElement2);
 	Registry->DestroyElement(TypedElement3);
@@ -205,6 +214,9 @@ bool FTypedElementRegistrySmokeTest::RunTest(const FString& Parameters)
 	Registry->DestroyElement(UntypedElement2);
 	Registry->DestroyElement(UntypedElement3);
 
+	check(!ScriptHandleTyped.IsSet());
+	check(!ScriptHandleUntyped.IsSet());
+
 	// Verify that there were no leaks
 	Registry->ProcessDeferredElementsToDestroy();
 
@@ -219,11 +231,11 @@ bool FTypedElementRegistryPerfTest::RunTest(const FString& Parameters)
 	UTypedElementRegistry* Registry = NewObject<UTypedElementRegistry>();
 
 	const FName DummyElementType_Typed = "DummyElementType_Typed";
-	Registry->RegisterElementType<FTestTypedElementData>(DummyElementType_Typed);
+	Registry->RegisterElementType<FTestTypedElementData, true>(DummyElementType_Typed);
 	Registry->RegisterElementInterface<ITestTypedElementInterfaceA>(DummyElementType_Typed, NewObject<UTestTypedElementInterfaceA_ImplTyped>());
 
 	const FName DummyElementType_Untyped = "DummyElementType_Untyped";
-	Registry->RegisterElementType(DummyElementType_Untyped);
+	Registry->RegisterElementType(DummyElementType_Untyped, true);
 	Registry->RegisterElementInterface<ITestTypedElementInterfaceA>(DummyElementType_Untyped, NewObject<UTestTypedElementInterfaceA_ImplUntyped>());
 
 	TArray<TTypedElementOwner<FTestTypedElementData>> TypedOwnerHandles;
diff --git a/Engine/Source/Runtime/TypedElementFramework/Private/TypedElementFrameworkTests.h b/Engine/Source/Runtime/TypedElementFramework/Private/TypedElementFrameworkTests.h
--- a/Engine/Source/Runtime/TypedElementFramework/Private/TypedElementFrameworkTests.h
+++ b/Engine/Source/Runtime/TypedElementFramework/Private/TypedElementFrameworkTests.h
@@ -21,11 +21,15 @@ class ITestTypedElementInterfaceA
 	GENERATED_BODY()
 
 public:
-	UFUNCTION(BlueprintCallable, Category="TypedElementFramework|Testing")
 	virtual FText GetDisplayName(const FTypedElementHandle& InElementHandle) { return FText(); }
-	
-	UFUNCTION(BlueprintCallable, Category="TypedElementFramework|Testing")
+
 	virtual bool SetDisplayName(const FTypedElementHandle& InElementHandle, FText InNewName, bool bNotify = true) { return false; }
+
+	UFUNCTION(BlueprintCallable, Category = "TypedElementFramework|Testing")
+	virtual FText GetDisplayName(const FScriptTypedElementHandle& InElementHandle) { return FText(); }
+
+	UFUNCTION(BlueprintCallable, Category = "TypedElementFramework|Testing")
+	virtual bool SetDisplayName(const FScriptTypedElementHandle& InElementHandle, FText InNewName, bool bNotify = true) { return false; }
 };
 
 UINTERFACE(MinimalAPI, BlueprintType, meta = (CannotImplementInterfaceInBlueprint))
@@ -39,8 +43,10 @@ class ITestTypedElementInterfaceB
 	GENERATED_BODY()
 
 public:
-	UFUNCTION(BlueprintCallable, Category = "TypedElementFramework|Testing")
 	virtual bool MarkAsTested(const FTypedElementHandle& InElementHandle) { return false; }
+
+	UFUNCTION(BlueprintCallable, Category = "TypedElementFramework|Testing")
+	virtual bool MarkAsTested(const FScriptTypedElementHandle& InElementHandle) { return false; }
 };
 
 UINTERFACE(MinimalAPI, BlueprintType, meta = (CannotImplementInterfaceInBlueprint))
@@ -54,8 +60,10 @@ class ITestTypedElementInterfaceC
 	GENERATED_BODY()
 
 public:
-	UFUNCTION(BlueprintCallable, Category = "TypedElementFramework|Testing")
 	virtual bool GetIsTested(const FTypedElementHandle& InElementHandle) const { return false; }
+
+	UFUNCTION(BlueprintCallable, Category = "TypedElementFramework|Testing")
+	virtual bool GetIsTested(const FScriptTypedElementHandle& InElementHandle) const { return false; }
 };
 
 template <>
diff --git a/Engine/Source/Runtime/TypedElementFramework/Public/Elements/Framework/TypedElementData.h b/Engine/Source/Runtime/TypedElementFramework/Public/Elements/Framework/TypedElementData.h
--- a/Engine/Source/Runtime/TypedElementFramework/Public/Elements/Framework/TypedElementData.h
+++ b/Engine/Source/Runtime/TypedElementFramework/Public/Elements/Framework/TypedElementData.h
@@ -346,6 +346,167 @@ class TTypedElementInternalData<void> : public FTypedElementInternalData
 {
 };
 
+/**
+ * Internal data that act as a ptr to a control block for the data represented by a ScriptTypedElementHandle
+ */
+class FScriptTypedElementInternalDataPtr
+{
+public:
+
+	FScriptTypedElementInternalDataPtr() = default;
+
+	~FScriptTypedElementInternalDataPtr()
+	{
+		DecrementCount();
+	}
+
+	FScriptTypedElementInternalDataPtr(const FScriptTypedElementInternalDataPtr& Other)
+		: ControlBlock(Other.ControlBlock)
+	{
+		IncrementCount();
+	}
+
+	FScriptTypedElementInternalDataPtr(FScriptTypedElementInternalDataPtr&& Other)
+		: ControlBlock(Other.ControlBlock)
+	{
+		Other.ControlBlock = nullptr;
+	}
+
+	FScriptTypedElementInternalDataPtr& operator=(const FScriptTypedElementInternalDataPtr& Other)
+	{
+		ControlBlock = Other.ControlBlock;
+		IncrementCount();
+		return *this;
+	}
+
+	FScriptTypedElementInternalDataPtr& operator=(FScriptTypedElementInternalDataPtr&& Other)
+	{
+		ControlBlock = Other.ControlBlock;
+		Other.ControlBlock = nullptr;
+		return *this;
+	}
+
+	bool operator==(const FScriptTypedElementInternalDataPtr& Other) const
+	{
+		return ControlBlock == Other.ControlBlock;
+	}
+
+	bool operator!=(const FScriptTypedElementInternalDataPtr& Other) const
+	{
+		return !(*this == Other);
+	}
+
+	bool IsSet() const
+	{
+		return ControlBlock && ControlBlock->Data;
+	}
+
+	void Release()
+	{
+		DecrementCount();
+		ControlBlock = nullptr;
+	}
+
+	FORCEINLINE const FTypedElementId& GetId() const
+	{
+		if (ControlBlock)
+		{
+			if (ControlBlock->Data)
+			{
+				return ControlBlock->Data->GetId();
+			}
+		}
+
+		return FTypedElementId::Unset;
+	}
+
+	FORCEINLINE FTypedElementInternalData* GetInternalData() const
+	{
+		if (ControlBlock)
+		{
+			if (ControlBlock->Data)
+			{
+				return ControlBlock->Data;
+			}
+		}
+
+		return nullptr;
+	}
+
+protected:
+	void IncrementCount()
+	{
+		if (ControlBlock)
+		{
+			++ControlBlock->WeakRefCount;
+		}
+	}
+
+	void DecrementCount()
+	{
+		if (ControlBlock)
+		{
+			--ControlBlock->WeakRefCount;
+			if (ControlBlock->WeakRefCount == 0)
+			{
+				delete ControlBlock;
+				ControlBlock = nullptr;
+			}
+		}
+	}
+
+	FScriptTypedElementInternalDataPtr(FTypedElementInternalData& InternalData)
+		: ControlBlock(new FScriptTypedElementInternalDataControlBlock(InternalData))
+	{
+		IncrementCount();
+	}
+
+	/**
+	 * Internal data that act as a weak reference count and a control block for the ScriptTypedElementHandle
+	 */
+	struct FScriptTypedElementInternalDataControlBlock
+	{
+		FScriptTypedElementInternalDataControlBlock(FTypedElementInternalData& InData)
+			: Data(&InData)
+		{
+		}
+
+		FScriptTypedElementInternalDataControlBlock() = delete;
+
+		// The control block should never be copied/Moved
+		FScriptTypedElementInternalDataControlBlock(const FScriptTypedElementInternalDataControlBlock&) = delete;
+		FScriptTypedElementInternalDataControlBlock(FScriptTypedElementInternalDataControlBlock&&) = delete;
+		FScriptTypedElementInternalDataControlBlock& operator=(const FScriptTypedElementInternalDataControlBlock&) = delete;
+		FScriptTypedElementInternalDataControlBlock& operator=(FScriptTypedElementInternalDataControlBlock&&) = delete;
+
+
+		FTypedElementInternalData* Data = nullptr;
+		FTypedElementRefCount WeakRefCount = 0;
+	};
+
+	FScriptTypedElementInternalDataControlBlock* ControlBlock = nullptr;
+};
+
+class FScriptTypedElementInternalDataOwner : public FScriptTypedElementInternalDataPtr
+{
+public:
+
+	FScriptTypedElementInternalDataOwner(FTypedElementInternalData& InInternalData)
+		: FScriptTypedElementInternalDataPtr(InInternalData)
+	{
+		ControlBlock = new FScriptTypedElementInternalDataControlBlock(InInternalData);
+		IncrementCount();
+	}
+
+	// Tell all the weak reference that the data they point to is no longer valid
+	~FScriptTypedElementInternalDataOwner()
+	{
+		ControlBlock->Data = nullptr;
+		DecrementCount();
+		ControlBlock = nullptr;
+	}
+};
+
 /**
  * Data store implementation used by the element registry to manage internal data. 
  * @note This is the generic implementation that uses an array and manages the IDs itself.
@@ -371,6 +532,26 @@ public:
 		return InternalData;
 	}
 
+	FScriptTypedElementInternalDataPtr GetInternalDataForScriptHandle(const FTypedHandleElementId InElementId)
+	{
+		// Script handles are not thread safe
+		uint32 Hash = GetTypeHash(InElementId);
+		if (FScriptTypedElementInternalDataOwner* ScriptElementInternalData  = ScriptInternalDataMap.FindByHash(Hash, InElementId))
+		{
+			return *ScriptElementInternalData;
+		}
+
+		FReadScopeLock InternalDataLock(InternalDataRW);
+		checkSlow(InternalDataArray.IsValidIndex(InElementId));
+		return ScriptInternalDataMap.EmplaceByHash(Hash,InElementId, InternalDataArray[InElementId]);
+	}
+
+	void DisableScriptHandlesForElement(const FTypedHandleElementId InElementId)
+	{
+		// Script handles are not thread safe
+		ScriptInternalDataMap.Remove(InElementId);
+	}
+
 	void RemoveDataForElement(const FTypedHandleElementId InElementId, const FTypedElementInternalData* InExpectedDataPtr)
 	{
 		FWriteScopeLock InternalDataLock(InternalDataRW);
@@ -412,6 +593,7 @@ private:
 	mutable FRWLock InternalDataRW;
 	TChunkedArray<TTypedElementInternalData<ElementDataType>> InternalDataArray;
 	TArray<int32> InternalDataFreeIndices;
+	TMap<FTypedHandleElementId, FScriptTypedElementInternalDataOwner> ScriptInternalDataMap;
 };
 
 /**
@@ -442,6 +624,26 @@ public:
 		return InternalData;
 	}
 
+	FScriptTypedElementInternalDataPtr GetInternalDataForScriptHandle(const FTypedHandleElementId InElementId)
+	{
+		// Script handles are not thread safe
+		uint32 Hash = GetTypeHash(InElementId);
+		if (FScriptTypedElementInternalDataOwner* ScriptElementInternalData = ScriptInternalDataMap.FindByHash(Hash, InElementId))
+		{
+			return *ScriptElementInternalData;
+		}
+
+		FReadScopeLock InternalDataLock(InternalDataRW);
+		const int32 Index = ElementIdToArrayIndex.FindChecked(InElementId);
+		return ScriptInternalDataMap.EmplaceByHash(Hash, InElementId, InternalDataArray[Index]);
+	}
+
+	void DisableScriptHandlesForElement(const FTypedHandleElementId InElementId)
+	{
+		// Script handles are not thread safe
+		ScriptInternalDataMap.Remove(InElementId);
+	}
+
 	void RemoveDataForElement(const FTypedHandleElementId InElementId, const FTypedElementInternalData* InExpectedDataPtr)
 	{
 		FWriteScopeLock InternalDataLock(InternalDataRW);
@@ -486,4 +688,5 @@ private:
 	TChunkedArray<TTypedElementInternalData<void>> InternalDataArray;
 	TArray<int32> InternalDataFreeIndices;
 	TMap<FTypedHandleElementId, int32> ElementIdToArrayIndex;
+	TMap<FTypedHandleElementId, FScriptTypedElementInternalDataOwner> ScriptInternalDataMap;
 };
diff --git a/Engine/Source/Runtime/TypedElementFramework/Public/Elements/Framework/TypedElementHandle.h b/Engine/Source/Runtime/TypedElementFramework/Public/Elements/Framework/TypedElementHandle.h
--- a/Engine/Source/Runtime/TypedElementFramework/Public/Elements/Framework/TypedElementHandle.h
+++ b/Engine/Source/Runtime/TypedElementFramework/Public/Elements/Framework/TypedElementHandle.h
@@ -14,11 +14,8 @@
  * C++ code may choose to use TTypedElement instead, which is a combination of an element handle and its associated element interface.
  * @note Handles auto-release on destruction.
  */
-USTRUCT(BlueprintType)
 struct TYPEDELEMENTFRAMEWORK_API FTypedElementHandle
 {
-	GENERATED_BODY()
-
 public:
 	FTypedElementHandle() = default;
 
@@ -614,7 +611,100 @@ private:
 };
 using FTypedElementOwner = TTypedElementOwner<void>;
 
-/** Script exposure for FTypedElementHandle. */
+
+/**
+ * Script exposure for the typed element handle struct type
+ * Act as a weak handle to simplify the scripting use of the typed element framework and making it safer to use by avoiding crash in case of a bad usage.
+ * This type is the standard way that an element is passed through to interfaces for a script (Blueprint or Python), and also the type that is stored in the script element lists.
+ * C++ code may choose to use TTypedElement instead, which is a combination of an element handle and its associated element interface.
+ *
+ * Note: This type shouldn't be used in the engine code as it come with a performance and memory overhead that we want to avoid when compare to the native handles (FTypedElementHandle).
+ */
+USTRUCT(BlueprintType)
+struct TYPEDELEMENTFRAMEWORK_API FScriptTypedElementHandle
+{
+	GENERATED_BODY()
+public:
+
+	FScriptTypedElementHandle() = default;
+
+	FScriptTypedElementHandle(const FScriptTypedElementHandle& InOther)
+		: InternalData(InOther.InternalData)
+	{
+	}
+
+	FScriptTypedElementHandle(FScriptTypedElementHandle&& InOther)
+		: InternalData(MoveTemp(InOther.InternalData))
+	{
+	}
+
+	FScriptTypedElementHandle& operator=(const FScriptTypedElementHandle& InOther)
+	{
+		InternalData = InOther.InternalData;
+		return *this;
+	}
+
+	FScriptTypedElementHandle& operator=(FScriptTypedElementHandle&& InOther)
+	{
+		InternalData = MoveTemp(InOther.InternalData);
+		return *this;
+	}
+
+	FORCEINLINE bool operator==(const FScriptTypedElementHandle& InOther) const
+	{
+		return InternalData == InOther.InternalData;
+	}
+
+	FORCEINLINE bool operator!=(const FScriptTypedElementHandle& InOther) const
+	{
+		return !(*this == InOther);
+	}
+
+	FORCEINLINE explicit operator bool() const
+	{
+		return IsSet();
+	}
+
+	FORCEINLINE bool IsSet() const
+	{
+		return InternalData.IsSet();
+	}
+
+	FORCEINLINE void Release()
+	{
+		InternalData.Release();
+	}
+
+	FORCEINLINE void Private_Initialize(FScriptTypedElementInternalDataPtr&& InInternalData)
+	{
+		InternalData = InInternalData;
+	}
+
+	FORCEINLINE const FTypedElementId& GetId() const
+	{
+		return InternalData.GetId();
+	}
+
+	/**
+	 * Return typed element handle from the script typed element handle
+	 * If this script handle is invalid it will return a invalid TypedElementHandle
+	 */
+	FTypedElementHandle GetTypedElementHandle() const
+	{
+		FTypedElementHandle Handle;
+		if (FTypedElementInternalData* TypedElementInternalData = InternalData.GetInternalData())
+		{
+			Handle.Private_InitializeAddRef(*TypedElementInternalData);
+		}
+
+		return Handle;
+	}
+
+private:
+	FScriptTypedElementInternalDataPtr InternalData;
+};
+
+/** Script exposure for FScriptTypedElementHandle. */
 UCLASS()
 class UTypedElementHandleLibrary : public UObject
 {
@@ -625,7 +715,7 @@ public:
 	 * Has this handle been initialized to a valid element?
 	 */
 	UFUNCTION(BlueprintPure, Category="TypedElementFramework|Handle", meta=(ScriptMethod, ScriptOperator="bool"))
-	static bool IsSet(const FTypedElementHandle& ElementHandle)
+	static bool IsSet(const FScriptTypedElementHandle& ElementHandle)
 	{
 		return ElementHandle.IsSet();
 	}
@@ -634,7 +724,7 @@ public:
 	 * Release this handle and set it back to an empty state.
 	 */
 	UFUNCTION(BlueprintCallable, Category="TypedElementFramework|Handle", meta=(ScriptMethod))
-	static void Release(UPARAM(ref) FTypedElementHandle& ElementHandle)
+	static void Release(UPARAM(ref) FScriptTypedElementHandle& ElementHandle)
 	{
 		ElementHandle.Release();
 	}
@@ -643,7 +733,7 @@ public:
 	 * Are these two handles equal?
 	 */
 	UFUNCTION(BlueprintPure, Category="TypedElementFramework|Handle", meta=(DisplayName="Equal (TypedElementHandle)", CompactNodeTitle="==", Keywords="== equal", ScriptMethod, ScriptOperator="=="))
-	static bool Equal(const FTypedElementHandle& LHS, const FTypedElementHandle& RHS)
+	static bool Equal(const FScriptTypedElementHandle& LHS, const FScriptTypedElementHandle& RHS)
 	{
 		return LHS == RHS;
 	}
@@ -652,7 +742,7 @@ public:
 	 * Are these two handles not equal?
 	 */
 	UFUNCTION(BlueprintPure, Category="TypedElementFramework|Handle", meta=(DisplayName="NotEqual (TypedElementHandle)", CompactNodeTitle="!=", Keywords="!= not equal", ScriptMethod, ScriptOperator="!="))
-	static bool NotEqual(const FTypedElementHandle& LHS, const FTypedElementHandle& RHS)
+	static bool NotEqual(const FScriptTypedElementHandle& LHS, const FScriptTypedElementHandle& RHS)
 	{
 		return LHS != RHS;
 	}
diff --git a/Engine/Source/Runtime/TypedElementFramework/Public/Elements/Framework/TypedElementList.h b/Engine/Source/Runtime/TypedElementFramework/Public/Elements/Framework/TypedElementList.h
--- a/Engine/Source/Runtime/TypedElementFramework/Public/Elements/Framework/TypedElementList.h
+++ b/Engine/Source/Runtime/TypedElementFramework/Public/Elements/Framework/TypedElementList.h
@@ -19,15 +19,17 @@ namespace TypedElementList_Private
 
 TYPEDELEMENTFRAMEWORK_API void GetElementImpl(const UTypedElementRegistry* InRegistry, const FTypedElementHandle& InElementHandle, const TSubclassOf<UInterface>& InBaseInterfaceType, FTypedElement& OutElement);
 
-template <typename BaseInterfaceType>
-FORCEINLINE void GetElement(const UTypedElementRegistry* InRegistry, const FTypedElementHandle& InElementHandle, TTypedElement<BaseInterfaceType>& OutElement)
+TYPEDELEMENTFRAMEWORK_API void GetElementImpl(const UTypedElementRegistry* InRegistry, const FScriptTypedElementHandle& InElementHandle, const TSubclassOf<UInterface>& InBaseInterfaceType, FTypedElement& OutElement);
+
+template <typename BaseInterfaceType, class HandleType>
+FORCEINLINE void GetElement(const UTypedElementRegistry* InRegistry, const HandleType& InElementHandle, TTypedElement<BaseInterfaceType>& OutElement)
 {
 	static_assert(sizeof(TTypedElement<BaseInterfaceType>) == sizeof(FTypedElement), "All TTypedElement instances must be the same size for this cast implementation to work!");
 	GetElementImpl(InRegistry, InElementHandle, BaseInterfaceType::UClassType::StaticClass(), reinterpret_cast<FTypedElement&>(OutElement));
 }
 
-template <typename BaseInterfaceType>
-FORCEINLINE TTypedElement<BaseInterfaceType> GetElement(const UTypedElementRegistry* InRegistry, const FTypedElementHandle& InElementHandle)
+template <typename BaseInterfaceType, class HandleType>
+FORCEINLINE TTypedElement<BaseInterfaceType> GetElement(const UTypedElementRegistry* InRegistry, const HandleType& InElementHandle)
 {
 	TTypedElement<BaseInterfaceType> Element;
 	GetElement(InRegistry, InElementHandle, Element);
@@ -36,119 +38,34 @@ FORCEINLINE TTypedElement<BaseInterfaceType> GetElement(const UTypedElementRegis
 
 } // namespace TypedElementList_Private
 
-/**
- * Interface to allow external systems (such as USelection) to receive immediate sync notifications as an element list is changed.
- * This exists purely as a bridging mechanism and shouldn't be relied on for new code. It is lazily created as needed.
- */
-class TYPEDELEMENTFRAMEWORK_API FTypedElementListLegacySync
-{
-public:
-	enum class ESyncType : uint8
-	{
-		/**
-		 * An element was added to the element list.
-		 * The ElementHandle argument will be set to the element that was added.
-		 */
-		Added,
-
-		/**
-		 * An element was removed from the element list.
-		 * The ElementHandle argument will be set to the element that was removed.
-		 */
-		Removed,
-
-		/**
-		 * The element list was modified in an unknown way.
-		 * The ElementHandle argument will be unset.
-		 */
-		Modified,
-
-		/**
-		 * The element list was cleared.
-		 * The ElementHandle argument will be unset.
-		 */
-		Cleared,
-
-		/**
-		 * The element list was modified as part of a batch or bulk operation.
-		 * The ElementHandle argument will be unset.
-		 * @note A batch operation will emit internal (bIsWithinBatchOperation=true) Added, Removed, Modified and Cleared updates during the batch, 
-		 *       so if you respond to those internal updates you may choose to ignore this one. Otherwise you should treat it the same as Modified.
-		 */
-		BatchComplete,
-	};
-	
-	FTypedElementListLegacySync(const FTypedElementList& InElementList);
-
-	void Private_EmitSyncEvent(const ESyncType InSyncType, const FTypedElementHandle& InElementHandle = FTypedElementHandle());
-
-	DECLARE_EVENT_FourParams(FTypedElementListLegacySync, FOnSyncEvent, const FTypedElementList& /*InElementList*/, ESyncType /*InSyncType*/, const FTypedElementHandle& /*InElementHandle*/, bool /*bIsWithinBatchOperation*/);
-	FOnSyncEvent& OnSyncEvent();
-
-	bool IsRunningBatchOperation() const;
-	void BeginBatchOperation();
-	void EndBatchOperation(const bool InNotify = true);
-	bool IsBatchOperationDirty() const;
-	void ForceBatchOperationDirty();
-
-private:
-	const FTypedElementList& ElementList;
-
-	FOnSyncEvent OnSyncEventDelegate;
-
-	int32 NumOpenBatchOperations = 0;
-	bool bBatchOperationIsDirty = false;
-};
-
-/**
- * Helper to batch immediate sync notifications for legacy code.
- * Does nothing if no legacy sync has been created for the given instance.
- */
-class TYPEDELEMENTFRAMEWORK_API FTypedElementListLegacySyncScopedBatch
-{
-public:
-	explicit FTypedElementListLegacySyncScopedBatch(const FTypedElementList& InElementList, const bool InNotify = true);
-	~FTypedElementListLegacySyncScopedBatch();
-
-	FTypedElementListLegacySyncScopedBatch(const FTypedElementListLegacySyncScopedBatch&) = delete;
-	FTypedElementListLegacySyncScopedBatch& operator=(const FTypedElementListLegacySyncScopedBatch&) = delete;
-
-	FTypedElementListLegacySyncScopedBatch(FTypedElementListLegacySyncScopedBatch&&) = delete;
-	FTypedElementListLegacySyncScopedBatch& operator=(FTypedElementListLegacySyncScopedBatch&&) = delete;
-
-	bool IsDirty() const;
-	void ForceDirty();
-
-private:
-	FTypedElementListLegacySync* ElementListLegacySync = nullptr;
-	bool bNotify = true;
-};
 
 /**
  * A list of element handles.
  * Provides high-level access to groups of elements, including accessing elements that implement specific interfaces.
  */
-class TYPEDELEMENTFRAMEWORK_API FTypedElementList final : public TSharedFromThis<FTypedElementList>
+template<class HandleType>
+class TTypedElementList final : public TSharedFromThis<TTypedElementList<HandleType>>
 {
 public:
-	~FTypedElementList();
 
-	/**
-	 * Internal function used by the element registry to create an element list instance.
-	 */
-	static FTypedElementListRef Private_CreateElementList(UTypedElementRegistry* InRegistry);
+	using FHandleType = HandleType;
+	using TTypedElementListRef = TSharedRef<TTypedElementList>;
+	using TTypedElementListConstRef = TSharedRef<const TTypedElementList>;
+
+
+	~TTypedElementList();
 
 	/**
 	 * Clone this list instance.
 	 * @note Only copies elements; does not copy any bindings!
 	 */
-	FTypedElementListRef Clone() const;
+	TTypedElementListRef Clone() const;
 
 	/**
 	 * Get the element handle at the given index.
 	 * @note Use IsValidIndex to test for validity.
 	 */
-	FORCEINLINE FTypedElementHandle operator[](const int32 InIndex) const
+	FORCEINLINE HandleType operator[](const int32 InIndex) const
 	{
 		return GetElementHandleAt(InIndex);
 	}
@@ -157,7 +74,7 @@ public:
 	 * Get the element handle at the given index.
 	 * @note Use IsValidIndex to test for validity.
 	 */
-	FORCEINLINE FTypedElementHandle GetElementHandleAt(const int32 InIndex) const
+	FORCEINLINE HandleType GetElementHandleAt(const int32 InIndex) const
 	{
 		return ElementHandles[InIndex];
 	}
@@ -186,7 +103,7 @@ public:
 	 * Get the element from the given handle.
 	 */
 	template <typename BaseInterfaceType>
-	FORCEINLINE TTypedElement<BaseInterfaceType> GetElement(const FTypedElementHandle& InElementHandle) const
+	FORCEINLINE TTypedElement<BaseInterfaceType> GetElement(const HandleType& InElementHandle) const
 	{
 		return TypedElementList_Private::GetElement<BaseInterfaceType>(Registry.Get(), InElementHandle);
 	}
@@ -195,7 +112,7 @@ public:
 	 * Get the element from the given handle.
 	 */
 	template <typename BaseInterfaceType>
-	FORCEINLINE void GetElement(const FTypedElementHandle& InElementHandle, TTypedElement<BaseInterfaceType>& OutElement) const
+	FORCEINLINE void GetElement(const HandleType& InElementHandle, TTypedElement<BaseInterfaceType>& OutElement) const
 	{
 		TypedElementList_Private::GetElement(Registry.Get(), InElementHandle, OutElement);
 	}
@@ -293,7 +210,7 @@ public:
 	/**
 	 * Get the element interface from the given handle.
 	 */
-	UObject* GetElementInterface(const FTypedElementHandle& InElementHandle, const TSubclassOf<UInterface>& InBaseInterfaceType) const;
+	UObject* GetElementInterface(const HandleType& InElementHandle, const TSubclassOf<UInterface>& InBaseInterfaceType) const;
 
 	/**
 	 * Test whether there are elements in this list, optionally filtering to elements that implement the given interface.
@@ -320,17 +237,17 @@ public:
 	/**
 	 * Get the handle of every element in this list, optionally filtering to elements that implement the given interface.
 	 */
-	TArray<FTypedElementHandle> GetElementHandles(const TSubclassOf<UInterface>& InBaseInterfaceType = nullptr) const;
+	TArray<HandleType> GetElementHandles(const TSubclassOf<UInterface>& InBaseInterfaceType = nullptr) const;
 
 	/**
 	 * Get the handle of every element in this list, optionally filtering to elements that implement the given interface.
 	 */
 	template <typename ArrayAllocator>
-	void GetElementHandles(TArray<FTypedElementHandle, ArrayAllocator>& OutArray, const TSubclassOf<UInterface>& InBaseInterfaceType = nullptr) const
+	void GetElementHandles(TArray<HandleType, ArrayAllocator>& OutArray, const TSubclassOf<UInterface>& InBaseInterfaceType = nullptr) const
 	{
 		OutArray.Reset();
 		OutArray.Reserve(ElementHandles.Num());
-		ForEachElementHandle([&OutArray](const FTypedElementHandle& InElementHandle)
+		ForEachElementHandle([&OutArray](const HandleType& InElementHandle)
 		{
 			OutArray.Add(InElementHandle);
 			return true;
@@ -341,7 +258,7 @@ public:
 	 * Enumerate the handle of every element in this list, optionally filtering to elements that implement the given interface.
 	 * @note Return true from the callback to continue enumeration.
 	 */
-	void ForEachElementHandle(TFunctionRef<bool(const FTypedElementHandle&)> InCallback, const TSubclassOf<UInterface>& InBaseInterfaceType = nullptr) const;
+	void ForEachElementHandle(TFunctionRef<bool(const HandleType&)> InCallback, const TSubclassOf<UInterface>& InBaseInterfaceType = nullptr) const;
 
 	/**
 	 * Enumerate the elements in this list that implement the given interface.
@@ -351,7 +268,7 @@ public:
 	void ForEachElement(TFunctionRef<bool(const TTypedElement<BaseInterfaceType>&)> InCallback) const
 	{
 		TTypedElement<BaseInterfaceType> TempElement;
-		for (const FTypedElementHandle& ElementHandle : ElementHandles)
+		for (const HandleType& ElementHandle : ElementHandles)
 		{
 			GetElement(ElementHandle, TempElement);
 			if (TempElement && !InCallback(TempElement))
@@ -445,7 +362,7 @@ public:
 	/**
 	 * Does this element list contain an entry for the given element handle?
 	 */
-	FORCEINLINE bool Contains(const FTypedElementHandle& InElementHandle) const
+	FORCEINLINE bool Contains(const HandleType& InElementHandle) const
 	{
 		return ContainsElementImpl(InElementHandle.GetId());
 	}
@@ -463,7 +380,7 @@ public:
 	 * Add the given element handle to this element list, if it isn't already in the list.
 	 * @return True if the element handle was added, false if it is already in the list.
 	 */
-	FORCEINLINE bool Add(const FTypedElementHandle& InElementHandle)
+	FORCEINLINE bool Add(const HandleType& InElementHandle)
 	{
 		return AddElementImpl(CopyTemp(InElementHandle));
 	}
@@ -472,7 +389,7 @@ public:
 	 * Add the given element handle to this element list, if it isn't already in the list.
 	 * @return True if the element handle was added, false if it is already in the list.
 	 */
-	FORCEINLINE bool Add(FTypedElementHandle&& InElementHandle)
+	FORCEINLINE bool Add(HandleType&& InElementHandle)
 	{
 		return AddElementImpl(MoveTemp(InElementHandle));
 	}
@@ -490,14 +407,14 @@ public:
 	/**
 	 * Append another element list to this element list.
 	 */
-	void Append(FTypedElementListConstRef InElementList)
+	void Append(const TTypedElementListConstRef& InElementList)
 	{
 		if (this != &InElementList.Get())
 		{
-			FTypedElementListLegacySyncScopedBatch LegacySyncBatch(*this);
+			FLegacySyncScopedBatch LegacySyncBatch(*this);
 
 			Reserve(Num() + InElementList->Num());
-			InElementList->ForEachElementHandle([this](const FTypedElementHandle& ElementHandle)
+			InElementList->ForEachElementHandle([this](const HandleType& ElementHandle)
 			{
 				AddElementImpl(CopyTemp(ElementHandle));
 				return true;
@@ -508,12 +425,12 @@ public:
 	/**
 	 * Append the given element handles to this element list.
 	 */
-	void Append(TArrayView<const FTypedElementHandle> InElementHandles)
+	void Append(TArrayView<const HandleType> InElementHandles)
 	{
-		FTypedElementListLegacySyncScopedBatch LegacySyncBatch(*this);
+		FLegacySyncScopedBatch LegacySyncBatch(*this);
 
 		Reserve(Num() + InElementHandles.Num());
-		for (const FTypedElementHandle& ElementHandle : InElementHandles)
+		for (const HandleType& ElementHandle : InElementHandles)
 		{
 			AddElementImpl(CopyTemp(ElementHandle));
 		}
@@ -534,7 +451,7 @@ public:
 	template <typename ElementDataType>
 	void Append(TArrayView<const TTypedElementOwner<ElementDataType>> InElementOwners)
 	{
-		FTypedElementListLegacySyncScopedBatch LegacySyncBatch(*this);
+		FLegacySyncScopedBatch LegacySyncBatch(*this);
 
 		Reserve(Num() + InElementOwners.Num());
 		for (const TTypedElementOwner<ElementDataType>& ElementOwner : InElementOwners)
@@ -556,7 +473,7 @@ public:
 	 * Remove the given element handle from this element list, if it is in the list.
 	 * @return True if the element handle was removed, false if it isn't in the list.
 	 */
-	FORCEINLINE bool Remove(const FTypedElementHandle& InElementHandle)
+	FORCEINLINE bool Remove(const HandleType& InElementHandle)
 	{
 		return RemoveElementImpl(InElementHandle.GetId());
 	}
@@ -575,7 +492,7 @@ public:
 	 * Remove any element handles that match the given predicate from this element list.
 	 * @return The number of element handles removed.
 	 */
-	FORCEINLINE int32 RemoveAll(TFunctionRef<bool(const FTypedElementHandle&)> InPredicate)
+	FORCEINLINE int32 RemoveAll(TFunctionRef<bool(const HandleType&)> InPredicate)
 	{
 		return RemoveAllElementsImpl(InPredicate);
 	}
@@ -588,7 +505,7 @@ public:
 	int32 RemoveAll(TFunctionRef<bool(const TTypedElement<BaseInterfaceType>&)> InPredicate)
 	{
 		TTypedElement<BaseInterfaceType> TempElement;
-		return RemoveAllElementsImpl([this, &TempElement, &InPredicate](const FTypedElementHandle& InElementHandle)
+		return RemoveAllElementsImpl([this, &TempElement, &InPredicate](const HandleType& InElementHandle)
 		{
 			GetElement(InElementHandle, TempElement);
 			return TempElement && InPredicate(TempElement);
@@ -607,7 +524,7 @@ public:
 	 * Access the delegate that is invoked whenever this element list is potentially about to change.
 	 * @note This may be called even if no actual change happens, so may be called multiple times without a corresponding OnChanged notification.
 	 */
-	DECLARE_EVENT_OneParam(FTypedElementList, FOnPreChange, const FTypedElementList& /*InElementList*/);
+	DECLARE_EVENT_OneParam(TTypedElementList, FOnPreChange, const TTypedElementList& /*InElementList*/);
 	FOnPreChange& OnPreChange()
 	{
 		return OnPreChangeDelegate;
@@ -617,7 +534,7 @@ public:
 	 * Access the delegate that is invoked whenever this element list has been changed.
 	 * @note This is called automatically at the end of each frame, but can also be manually invoked by NotifyPendingChanges.
 	 */
-	DECLARE_EVENT_OneParam(FTypedElementList, FOnChanged, const FTypedElementList& /*InElementList*/);
+	DECLARE_EVENT_OneParam(TTypedElementList, FOnChanged, const TTypedElementList& /*InElementList*/);
 	FOnChanged& OnChanged()
 	{
 		return OnChangedDelegate;
@@ -636,7 +553,7 @@ public:
 	struct TYPEDELEMENTFRAMEWORK_API FScopedClearNewPendingChange
 	{
 		FScopedClearNewPendingChange() = default;
-		FScopedClearNewPendingChange(FTypedElementList& InTypeElementList);
+		FScopedClearNewPendingChange(TTypedElementList& InTypeElementList);
 
 		FScopedClearNewPendingChange(const FScopedClearNewPendingChange&) = delete;
 		FScopedClearNewPendingChange& operator=(const FScopedClearNewPendingChange&) = delete;
@@ -647,7 +564,95 @@ public:
 		~FScopedClearNewPendingChange();
 
 	private:
-		FTypedElementList* TypedElementList = nullptr;
+		TTypedElementList* TypedElementList = nullptr;
+	};
+
+	/**
+	 * Interface to allow external systems (such as USelection) to receive immediate sync notifications as an element list is changed.
+	 * This exists purely as a bridging mechanism and shouldn't be relied on for new code. It is lazily created as needed.
+	 */
+	class TYPEDELEMENTFRAMEWORK_API FLegacySync
+	{
+	public:
+		enum class ESyncType : uint8
+		{
+			/**
+			 * An element was added to the element list.
+			 * The ElementHandle argument will be set to the element that was added.
+			 */
+			Added,
+
+			/**
+			 * An element was removed from the element list.
+			 * The ElementHandle argument will be set to the element that was removed.
+			 */
+			Removed,
+
+			/**
+			 * The element list was modified in an unknown way.
+			 * The ElementHandle argument will be unset.
+			 */
+			Modified,
+
+			/**
+			 * The element list was cleared.
+			 * The ElementHandle argument will be unset.
+			 */
+			Cleared,
+
+			/**
+			 * The element list was modified as part of a batch or bulk operation.
+			 * The ElementHandle argument will be unset.
+			 * @note A batch operation will emit internal (bIsWithinBatchOperation=true) Added, Removed, Modified and Cleared updates during the batch, 
+			 *       so if you respond to those internal updates you may choose to ignore this one. Otherwise you should treat it the same as Modified.
+			 */
+			BatchComplete,
+		};
+	
+		FLegacySync(const TTypedElementList& InElementList);
+
+		void Private_EmitSyncEvent(const ESyncType InSyncType, const HandleType& InElementHandle = HandleType());
+
+		DECLARE_EVENT_FourParams(FLegacySync, FOnSyncEvent, const TTypedElementList& /*InElementList*/, ESyncType /*InSyncType*/, const HandleType& /*InElementHandle*/, bool /*bIsWithinBatchOperation*/);
+		FOnSyncEvent& OnSyncEvent();
+
+		bool IsRunningBatchOperation() const;
+		void BeginBatchOperation();
+		void EndBatchOperation(const bool InNotify = true);
+		bool IsBatchOperationDirty() const;
+		void ForceBatchOperationDirty();
+
+	private:
+		const TTypedElementList& ElementList;
+
+		FOnSyncEvent OnSyncEventDelegate;
+
+		int32 NumOpenBatchOperations = 0;
+		bool bBatchOperationIsDirty = false;
+	};
+
+	/**
+	 * Helper to batch immediate sync notifications for legacy code.
+	 * Does nothing if no legacy sync has been created for the given instance.
+	 */
+	class TYPEDELEMENTFRAMEWORK_API FLegacySyncScopedBatch
+	{
+	public:
+		explicit FLegacySyncScopedBatch(const TTypedElementList& InElementList, const bool InNotify = true);
+		~FLegacySyncScopedBatch();
+
+		FLegacySyncScopedBatch(const FLegacySyncScopedBatch&) = delete;
+		FLegacySyncScopedBatch& operator=(const FLegacySyncScopedBatch&) = delete;
+
+		FLegacySyncScopedBatch(FLegacySyncScopedBatch&&) = delete;
+		FLegacySyncScopedBatch& operator=(FLegacySyncScopedBatch&&) = delete;
+
+		bool IsDirty() const;
+		void ForceDirty();
+
+	private:
+		FLegacySync* ElementListLegacySync = nullptr;
+		bool bNotify = true;
 	};
 
 	/**
@@ -659,13 +664,23 @@ public:
 	 * Access the interface to allow external systems (such as USelection) to receive immediate sync notifications as an element list is changed.
 	 * This exists purely as a bridging mechanism and shouldn't be relied on for new code. It is lazily created as needed.
 	 */
-	FTypedElementListLegacySync& Legacy_GetSync();
+	FLegacySync& Legacy_GetSync();
 
 	/**
 	 * Access the interface to allow external systems (such as USelection) to receive immediate sync notifications as an element list is changed.
 	 * This exists purely as a bridging mechanism and shouldn't be relied on for new code. This will return null if no legacy sync has been created for this instance.
 	 */
-	FTypedElementListLegacySync* Legacy_GetSyncPtr() const;
+	FLegacySync* Legacy_GetSyncPtr() const;
+
+	/**
+	 * Internal function used by the element registry to create an element list instance.
+	 */
+	static TTypedElementListRef Private_CreateElementList(UTypedElementRegistry* InRegistry);
+
+	UTypedElementRegistry* GetRegistry() const
+	{
+		return Registry.Get();
+	}
 
 private:
 	enum class EChangeType : uint8
@@ -689,15 +704,15 @@ private:
 		Cleared,
 	};
 
-	explicit FTypedElementList(UTypedElementRegistry* InRegistry);
+	explicit TTypedElementList(UTypedElementRegistry* InRegistry);
 
-	bool AddElementImpl(FTypedElementHandle&& InElementHandle);
+	bool AddElementImpl(HandleType&& InElementHandle);
 	bool RemoveElementImpl(const FTypedElementId& InElementId);
-	int32 RemoveAllElementsImpl(TFunctionRef<bool(const FTypedElementHandle&)> InPredicate);
+	int32 RemoveAllElementsImpl(TFunctionRef<bool(const HandleType&)> InPredicate);
 	bool ContainsElementImpl(const FTypedElementId& InElementId) const;
 
 	void NoteListMayChange();
-	void NoteListChanged(const EChangeType InChangeType, const FTypedElementHandle& InElementHandle = FTypedElementHandle());
+	void NoteListChanged(const EChangeType InChangeType, const HandleType& InElementHandle = HandleType());
 
 	/**
 	 * Element registry this element list is associated with.
@@ -714,7 +729,7 @@ private:
 	 * Array of element handles present in this element list.
 	 * These are stored in the same order that they are added, and the set above can be used to optimize certain queries.
 	 */
-	TArray<FTypedElementHandle> ElementHandles;
+	TArray<HandleType> ElementHandles;
 
 	/**
 	 * Tracks various categories of counters for the elements within the list (eg, the number of elements of a given type).
@@ -740,5 +755,22 @@ private:
 	 * Interface to allow external systems (such as USelection) to receive immediate sync notifications as an element list is changed.
 	 * This exists purely as a bridging mechanism and shouldn't be relied on for new code. It is lazily created as needed.
 	 */
-	TUniquePtr<FTypedElementListLegacySync> LegacySync;
+	TUniquePtr<FLegacySync> LegacySync;
 };
+
+
+namespace UE::TypedElementFramework
+{
+	/**
+	 * Functions to convert a script list to a native one
+	 * @note Only copies elements; does not copy any bindings!
+	 */
+	TYPEDELEMENTFRAMEWORK_API FTypedElementListPtr ConvertToNativeTypedElementList(const FScriptTypedElementListConstPtr& ScriptList);
+
+	/**
+	 * Functions to convert a script list to a native one
+	 * @note Only copies elements; does not copy any bindings!
+	 */
+	TYPEDELEMENTFRAMEWORK_API FScriptTypedElementListPtr ConvertToScriptTypedElementList(const FTypedElementListConstPtr& NativeList);
+
+}
\ No newline at end of file
diff --git a/Engine/Source/Runtime/TypedElementFramework/Public/Elements/Framework/TypedElementListFwd.h b/Engine/Source/Runtime/TypedElementFramework/Public/Elements/Framework/TypedElementListFwd.h
--- a/Engine/Source/Runtime/TypedElementFramework/Public/Elements/Framework/TypedElementListFwd.h
+++ b/Engine/Source/Runtime/TypedElementFramework/Public/Elements/Framework/TypedElementListFwd.h
@@ -4,10 +4,30 @@
 
 #include "Templates/SharedPointer.h"
 
-class FTypedElementList;
+template<class HandleType>
+class TTypedElementList;
+
+
+struct FTypedElementHandle;
+
+// Native typed element list.
+using FTypedElementList = TTypedElementList<FTypedElementHandle>;
 
 using FTypedElementListPtr = TSharedPtr<FTypedElementList>;
 using FTypedElementListRef = TSharedRef<FTypedElementList>;
 
 using FTypedElementListConstPtr = TSharedPtr<const FTypedElementList>;
 using FTypedElementListConstRef = TSharedRef<const FTypedElementList>;
+
+
+struct FScriptTypedElementHandle;
+
+// Script typed element list. It should only be use for the script exposure apis since the script handles does have a performance overhead over the normal handles.
+using FScriptTypedElementList = TTypedElementList<FScriptTypedElementHandle>;
+
+using FScriptTypedElementListPtr = TSharedPtr<FScriptTypedElementList>;
+using FScriptTypedElementListRef = TSharedRef<FScriptTypedElementList>;
+
+using FScriptTypedElementListConstPtr = TSharedPtr<const FScriptTypedElementList>;
+using FScriptTypedElementListConstRef = TSharedRef<const FScriptTypedElementList>;
+
diff --git a/Engine/Source/Runtime/TypedElementFramework/Public/Elements/Framework/TypedElementListProxy.h b/Engine/Source/Runtime/TypedElementFramework/Public/Elements/Framework/TypedElementListProxy.h
--- a/Engine/Source/Runtime/TypedElementFramework/Public/Elements/Framework/TypedElementListProxy.h
+++ b/Engine/Source/Runtime/TypedElementFramework/Public/Elements/Framework/TypedElementListProxy.h
@@ -7,37 +7,41 @@
 #include "TypedElementListProxy.generated.h"
 
 /**
- * A list of element handles (proxy to a FTypedElementList instance).
+ * A list of managed element handles (proxy to a FManagedTypedElementList instance).
  * Provides high-level access to groups of elements, including accessing elements that implement specific interfaces.
+ * 
+ * Note: the managed list proxy use should be avoided when not using it for the script exposure apis.
+ * The weak model for the handles come with an additional cost to the runtime performance and the memory usage.
  */
-USTRUCT(BlueprintType, DisplayName="TypedElementList", meta=(ScriptName="TypedElementList"))
-struct TYPEDELEMENTFRAMEWORK_API FTypedElementListProxy
+USTRUCT(BlueprintType, DisplayName="Typed Element List", meta=(ScriptName="TypedElementList"))
+struct TYPEDELEMENTFRAMEWORK_API FManagedTypedElementListProxy
 {
 	GENERATED_BODY()
-	
+
 public:
-	FTypedElementListProxy() = default;
+	FManagedTypedElementListProxy() = default;
 
-	FTypedElementListProxy(FTypedElementListRef InElementList)
+	FManagedTypedElementListProxy(FScriptTypedElementListRef InElementList)
 		: ElementList(MoveTemp(InElementList))
 	{
 	}
 
-	FTypedElementListProxy(FTypedElementListConstRef InElementList)
-		: ElementList(ConstCastSharedRef<FTypedElementList>(InElementList))
+	FManagedTypedElementListProxy(FScriptTypedElementListConstRef InElementList)
+		: ElementList(ConstCastSharedRef<FScriptTypedElementList>(InElementList))
 	{
 	}
 
-	FTypedElementListPtr GetElementList()
+	FScriptTypedElementListPtr GetElementList()
 	{
 		return ElementList;
 	}
 
-	FTypedElementListConstPtr GetElementList() const
+	FScriptTypedElementListConstPtr GetElementList() const
 	{
 		return ElementList;
 	}
-	
+
 private:
-	FTypedElementListPtr ElementList;
+	FScriptTypedElementListPtr ElementList;
 };
+
diff --git a/Engine/Source/Runtime/TypedElementFramework/Public/Elements/Framework/TypedElementRegistry.h b/Engine/Source/Runtime/TypedElementFramework/Public/Elements/Framework/TypedElementRegistry.h
--- a/Engine/Source/Runtime/TypedElementFramework/Public/Elements/Framework/TypedElementRegistry.h
+++ b/Engine/Source/Runtime/TypedElementFramework/Public/Elements/Framework/TypedElementRegistry.h
@@ -43,7 +43,7 @@ public:
 	/**
 	 * Get the singleton instance of the registry used in most cases.
 	 */
-	UFUNCTION(BlueprintPure, Category="TypedElementFramework|Registry")
+	UFUNCTION(BlueprintPure, DisplayName="Get Default Typed Element Registry", Category = "TypedElementFramework|Registry", meta=(ScriptName="GetDefaultTypedElementRegistry"))
 	static UTypedElementRegistry* GetInstance();
 
 	/**
@@ -96,19 +96,28 @@ public:
 
 	/**
 	 * Register an element type that doesn't require any additional payload data.
+	 * @param bSupportScriptHandles Does this element type support script handles. If true the elements of this type can only be destroyed from the game thread
 	 */
-	FORCEINLINE void RegisterElementType(const FName InElementTypeName)
+	FORCEINLINE void RegisterElementType(const FName InElementTypeName, bool bSupportScriptHandles)
 	{
-		RegisterElementTypeImpl(InElementTypeName, MakeUnique<TRegisteredElementType<void>>());
+		if (bSupportScriptHandles)
+		{
+			RegisterElementTypeImpl(InElementTypeName, MakeUnique<TRegisteredElementType<void, true>>());
+		}
+		else
+		{
+			RegisterElementTypeImpl(InElementTypeName, MakeUnique<TRegisteredElementType<void, false>>());
+		}
 	}
 
 	/**
 	 * Register an element type that has additional payload data.
+	 * Note if bSupportScriptHandles is true, elements of this type should only be destroyed from the game thread
 	 */
-	template <typename ElementDataType>
+	template <typename ElementDataType, bool bSupportScriptHandles>
 	FORCEINLINE void RegisterElementType(const FName InElementTypeName)
 	{
-		RegisterElementTypeImpl(InElementTypeName, MakeUnique<TRegisteredElementType<ElementDataType>>());
+		RegisterElementTypeImpl(InElementTypeName, MakeUnique<TRegisteredElementType<ElementDataType, bSupportScriptHandles>>());
 	}
 
 	/**
@@ -142,12 +151,20 @@ public:
 	/**
 	 * Get the element interface supported by the given handle, or null if there is no support for this interface.
 	 */
-	UFUNCTION(BlueprintPure, Category="TypedElementFramework|Registry")
 	FORCEINLINE UObject* GetElementInterface(const FTypedElementHandle& InElementHandle, const TSubclassOf<UInterface> InBaseInterfaceType) const
 	{
 		return GetElementInterfaceImpl(InElementHandle.GetId().GetTypeId(), InBaseInterfaceType);
 	}
 
+	/**
+	 * Get the element interface supported by the given handle, or null if there is no support for this interface or if the handle is invalid.
+	 */
+	UFUNCTION(BlueprintPure, Category="TypedElementFramework|Registry")
+	FORCEINLINE UObject* GetElementInterface(const FScriptTypedElementHandle& InElementHandle, const TSubclassOf<UInterface> InBaseInterfaceType) const
+	{
+		return GetElementInterfaceImpl(InElementHandle.GetId().GetTypeId(), InBaseInterfaceType);
+	}
+
 	/**
 	 * Create an element that doesn't require any additional payload data.
 	 * @note The associated handle ID should be something that can externally be used to uniquely identify this element, until DestroyElementHandle is called on this handle.
@@ -253,6 +270,14 @@ public:
 		return FTypedElementList::Private_CreateElementList(this);
 	}
 
+	/**
+	 * Create an empty list of script elements associated with this registry.
+	 */
+	FORCEINLINE FScriptTypedElementListRef CreateScriptElementList()
+	{
+		return FScriptTypedElementList::Private_CreateElementList(this);
+	}
+
 	/**
 	 * Create an empty list of elements associated with this registry, populated from the given minimal IDs that are valid.
 	 */
@@ -289,6 +314,28 @@ public:
 	 */
 	FString RegistredElementTypesAndInterfacesToString() const;
 
+	/**
+	 * Create a script handle. This should only be use for the script exposure apis.
+	 */
+	FORCEINLINE FScriptTypedElementHandle CreateScriptHandle(const FTypedElementId& InElementId)
+	{
+		if (InElementId == FTypedElementId::Unset)
+		{
+			return FScriptTypedElementHandle();
+		}
+
+		FRegisteredElementType* RegisteredElementType = GetRegisteredElementTypeFromId(InElementId.GetTypeId());
+		if (!RegisteredElementType)
+		{
+			FFrame::KismetExecutionMessage(TEXT("Element type ID has not been registered to this registry!"), ELogVerbosity::Error);
+			return FScriptTypedElementHandle();
+		}
+	
+		FScriptTypedElementHandle ScriptTypedElementHandle;
+		ScriptTypedElementHandle.Private_Initialize(RegisteredElementType->GetDataForScriptElement(InElementId.GetElementId()));
+		return ScriptTypedElementHandle;
+	}
+
 	void Private_OnElementListCreated(FTypedElementList* InElementList)
 	{
 		FScopeLock ActiveElementListsLock(&ActiveElementListsCS);
@@ -301,6 +348,19 @@ public:
 		ActiveElementLists.Remove(InElementList);
 	}
 
+	
+	void Private_OnElementListCreated(FScriptTypedElementList* InElementList)
+	{
+		// Script type element are not thread safe so no need for a critical section
+		ActiveScriptElementLists.Add(InElementList);
+	}
+
+	void Private_OnElementListDestroyed(FScriptTypedElementList* InElementList)
+	{
+		// Script type element are not thread safe so no need for a critical section
+		ActiveScriptElementLists.Remove(InElementList);
+	}
+
 	// Note: Access for FTypedElementList
 	FORCEINLINE void Private_GetElementImpl(const FTypedElementHandle& InElementHandle, const TSubclassOf<UInterface>& InBaseInterfaceType, FTypedElement& OutElement) const
 	{
@@ -341,6 +401,7 @@ private:
 		virtual FTypedElementInternalData& AddDataForElement(FTypedHandleElementId& InOutElementId) = 0;
 		virtual void RemoveDataForElement(const FTypedHandleElementId InElementId, const FTypedElementInternalData* InExpectedDataPtr, const bool bDefer) = 0;
 		virtual const FTypedElementInternalData& GetDataForElement(const FTypedHandleElementId InElementId) const = 0;
+		virtual FScriptTypedElementInternalDataPtr GetDataForScriptElement(const FTypedHandleElementId InElementId) = 0;
 		virtual void ProcessDeferredElementsToRemove() = 0;
 		virtual void SetDataTypeId(const FTypedHandleTypeId InTypeId) = 0;
 		virtual FTypedHandleTypeId GetDataTypeId() const = 0;
@@ -352,7 +413,7 @@ private:
 		TSortedMap<FName, UObject*, FDefaultAllocator, FNameFastLess> Interfaces;
 	};
 
-	template <typename ElementDataType>
+	template <typename ElementDataType, bool bSupportScriptHandles>
 	struct TRegisteredElementType : public FRegisteredElementType
 	{
 		virtual ~TRegisteredElementType() = default;
@@ -364,6 +425,12 @@ private:
 
 		virtual void RemoveDataForElement(const FTypedHandleElementId InElementId, const FTypedElementInternalData* InExpectedDataPtr, const bool bDefer) override
 		{
+			if constexpr (bSupportScriptHandles)
+			{
+				checkSlow(IsInGameThread());
+				HandleDataStore.DisableScriptHandlesForElement(InElementId);
+			}
+
 			if (bDefer)
 			{
 				InExpectedDataPtr->StoreDestructionRequestCallstack();
@@ -377,6 +444,18 @@ private:
 			}
 		}
 
+		virtual FScriptTypedElementInternalDataPtr GetDataForScriptElement(const FTypedHandleElementId InElementId) override
+		{
+			if constexpr (bSupportScriptHandles)
+			{
+				return HandleDataStore.GetInternalDataForScriptHandle(InElementId);
+			}
+			else
+			{
+				return {};
+			}
+		}
+
 		virtual const FTypedElementInternalData& GetDataForElement(const FTypedHandleElementId InElementId) const override
 		{
 			return HandleDataStore.GetDataForElement(InElementId);
@@ -555,6 +634,8 @@ private:
 	mutable FCriticalSection ActiveElementListsCS;
 	TSet<FTypedElementList*> ActiveElementLists;
 
+	TSet<FScriptTypedElementList*> ActiveScriptElementLists; 
+
 	uint8 DisableElementDestructionOnGCCount = 0;
 	bool bIsWithinFrame = false;
 
diff --git a/Engine/Source/Runtime/TypedElementFramework/Public/Elements/Framework/TypedElementUtil.h b/Engine/Source/Runtime/TypedElementFramework/Public/Elements/Framework/TypedElementUtil.h
--- a/Engine/Source/Runtime/TypedElementFramework/Public/Elements/Framework/TypedElementUtil.h
+++ b/Engine/Source/Runtime/TypedElementFramework/Public/Elements/Framework/TypedElementUtil.h
@@ -5,6 +5,8 @@
 #include "Elements/Framework/TypedElementHandle.h"
 #include "Elements/Framework/TypedElementListFwd.h"
 
+class UTypedElementRegistry;
+
 namespace TypedElementUtil
 {
 
@@ -14,4 +16,7 @@ namespace TypedElementUtil
 TYPEDELEMENTFRAMEWORK_API void BatchElementsByType(FTypedElementListConstRef InElementsToBatch, TMap<FTypedHandleTypeId, TArray<FTypedElementHandle>>& OutElementsByType);
 TYPEDELEMENTFRAMEWORK_API void BatchElementsByType(TArrayView<const FTypedElementHandle> InElementsToBatch, TMap<FTypedHandleTypeId, TArray<FTypedElementHandle>>& OutElementsByType);
 
+TYPEDELEMENTFRAMEWORK_API TArray<FScriptTypedElementHandle> ConvertToScriptElementArray(const TArray<FTypedElementHandle>& InNativeHandles, UTypedElementRegistry* Registry);
+TYPEDELEMENTFRAMEWORK_API TArray<FTypedElementHandle> ConvertToNativeElementArray(const TArray<FScriptTypedElementHandle>& InScriptHandles);
+
 } // namespace TypedElementUtil
diff --git a/Engine/Source/Runtime/TypedElementRuntime/Private/Elements/Framework/TypedElementSelectionSet.cpp b/Engine/Source/Runtime/TypedElementRuntime/Private/Elements/Framework/TypedElementSelectionSet.cpp
--- a/Engine/Source/Runtime/TypedElementRuntime/Private/Elements/Framework/TypedElementSelectionSet.cpp
+++ b/Engine/Source/Runtime/TypedElementRuntime/Private/Elements/Framework/TypedElementSelectionSet.cpp
@@ -1,7 +1,9 @@
 // Copyright Epic Games, Inc. All Rights Reserved.
 
 #include "Elements/Framework/TypedElementSelectionSet.h"
+
 #include "Elements/Framework/TypedElementRegistry.h"
+#include "Elements/Framework/TypedElementUtil.h"
 #include "Elements/Interfaces/TypedElementHierarchyInterface.h"
 
 #include "UObject/GCObjectScopeGuard.h"
@@ -211,7 +213,7 @@ bool UTypedElementSelectionSet::SelectElements(const TArray<FTypedElementHandle>
 
 bool UTypedElementSelectionSet::SelectElements(TArrayView<const FTypedElementHandle> InElementHandles, const FTypedElementSelectionOptions InSelectionOptions)
 {
-	FTypedElementListLegacySyncScopedBatch LegacySyncBatch(*ElementList, InSelectionOptions.AllowLegacyNotifications());
+	FTypedElementList::FLegacySyncScopedBatch LegacySyncBatch(*ElementList, InSelectionOptions.AllowLegacyNotifications());
 
 	bool bSelectionChanged = false;
 
@@ -225,7 +227,7 @@ bool UTypedElementSelectionSet::SelectElements(TArrayView<const FTypedElementHan
 
 bool UTypedElementSelectionSet::SelectElements(FTypedElementListConstRef InElementList, const FTypedElementSelectionOptions InSelectionOptions)
 {
-	FTypedElementListLegacySyncScopedBatch LegacySyncBatch(*ElementList, InSelectionOptions.AllowLegacyNotifications());
+	FTypedElementList::FLegacySyncScopedBatch LegacySyncBatch(*ElementList, InSelectionOptions.AllowLegacyNotifications());
 
 	bool bSelectionChanged = false;
 
@@ -276,7 +278,7 @@ bool UTypedElementSelectionSet::DeselectElements(const TArray<FTypedElementHandl
 
 bool UTypedElementSelectionSet::DeselectElements(TArrayView<const FTypedElementHandle> InElementHandles, const FTypedElementSelectionOptions InSelectionOptions)
 {
-	FTypedElementListLegacySyncScopedBatch LegacySyncBatch(*ElementList, InSelectionOptions.AllowLegacyNotifications());
+	FTypedElementList::FLegacySyncScopedBatch LegacySyncBatch(*ElementList, InSelectionOptions.AllowLegacyNotifications());
 
 	bool bSelectionChanged = false;
 
@@ -290,7 +292,7 @@ bool UTypedElementSelectionSet::DeselectElements(TArrayView<const FTypedElementH
 
 bool UTypedElementSelectionSet::DeselectElements(FTypedElementListConstRef InElementList, const FTypedElementSelectionOptions InSelectionOptions)
 {
-	FTypedElementListLegacySyncScopedBatch LegacySyncBatch(*ElementList, InSelectionOptions.AllowLegacyNotifications());
+	FTypedElementList::FLegacySyncScopedBatch LegacySyncBatch(*ElementList, InSelectionOptions.AllowLegacyNotifications());
 
 	bool bSelectionChanged = false;
 
@@ -305,7 +307,7 @@ bool UTypedElementSelectionSet::DeselectElements(FTypedElementListConstRef InEle
 
 bool UTypedElementSelectionSet::ClearSelection(const FTypedElementSelectionOptions InSelectionOptions)
 {
-	FTypedElementListLegacySyncScopedBatch LegacySyncBatch(*ElementList, InSelectionOptions.AllowLegacyNotifications());
+	FTypedElementList::FLegacySyncScopedBatch LegacySyncBatch(*ElementList, InSelectionOptions.AllowLegacyNotifications());
 
 	bool bSelectionChanged = false;
 
@@ -340,7 +342,7 @@ bool UTypedElementSelectionSet::SetSelection(const TArray<FTypedElementHandle>&
 
 bool UTypedElementSelectionSet::SetSelection(TArrayView<const FTypedElementHandle> InElementHandles, const FTypedElementSelectionOptions InSelectionOptions)
 {
-	FTypedElementListLegacySyncScopedBatch LegacySyncBatch(*ElementList, InSelectionOptions.AllowLegacyNotifications());
+	FTypedElementList::FLegacySyncScopedBatch LegacySyncBatch(*ElementList, InSelectionOptions.AllowLegacyNotifications());
 
 	bool bSelectionChanged = false;
 
@@ -352,7 +354,7 @@ bool UTypedElementSelectionSet::SetSelection(TArrayView<const FTypedElementHandl
 
 bool UTypedElementSelectionSet::SetSelection(FTypedElementListConstRef InElementList, const FTypedElementSelectionOptions InSelectionOptions)
 {
-	FTypedElementListLegacySyncScopedBatch LegacySyncBatch(*ElementList, InSelectionOptions.AllowLegacyNotifications());
+	FTypedElementList::FLegacySyncScopedBatch LegacySyncBatch(*ElementList, InSelectionOptions.AllowLegacyNotifications());
 
 	bool bSelectionChanged = false;
 
@@ -549,3 +551,113 @@ void UTypedElementSelectionSet::OnElementListChanged(const FTypedElementList& In
 	OnChangedDelegate.Broadcast(this);
 	OnSelectionChange.Broadcast(this);
 }
+
+bool UTypedElementSelectionSet::IsElementSelected(const FScriptTypedElementHandle& InElementHandle, const FTypedElementIsSelectedOptions InSelectionOptions) const
+{
+	FTypedElementHandle NativeHandle = InElementHandle.GetTypedElementHandle();
+	if (!NativeHandle)
+	{
+		FFrame::KismetExecutionMessage(TEXT("InElementHandle is not a valid handle."), ELogVerbosity::Error);
+		return false;
+	}
+
+	return IsElementSelected(NativeHandle, InSelectionOptions);
+}
+
+bool UTypedElementSelectionSet::CanSelectElement(const FScriptTypedElementHandle& InElementHandle, const FTypedElementSelectionOptions InSelectionOptions) const
+{
+	FTypedElementHandle NativeHandle = InElementHandle.GetTypedElementHandle();
+	if (!NativeHandle)
+	{
+		FFrame::KismetExecutionMessage(TEXT("InElementHandle is not a valid handle."), ELogVerbosity::Error);
+		return false;
+	}
+
+	return CanSelectElement(NativeHandle, InSelectionOptions);
+}
+
+bool UTypedElementSelectionSet::CanDeselectElement(const FScriptTypedElementHandle& InElementHandle, const FTypedElementSelectionOptions InSelectionOptions) const
+{
+	FTypedElementHandle NativeHandle = InElementHandle.GetTypedElementHandle();
+	if (!NativeHandle)
+	{
+		FFrame::KismetExecutionMessage(TEXT("InElementHandle is not a valid handle."), ELogVerbosity::Error);
+		return false;
+	}
+
+	return CanDeselectElement(NativeHandle, InSelectionOptions);
+}
+
+bool UTypedElementSelectionSet::SelectElement(const FScriptTypedElementHandle& InElementHandle, const FTypedElementSelectionOptions InSelectionOptions)
+{
+	FTypedElementHandle NativeHandle = InElementHandle.GetTypedElementHandle();
+	if (!NativeHandle)
+	{
+		FFrame::KismetExecutionMessage(TEXT("InElementHandle is not a valid handle."), ELogVerbosity::Error);
+		return false;
+	}
+
+	return SelectElement(NativeHandle, InSelectionOptions);
+}
+
+bool UTypedElementSelectionSet::SelectElements(const TArray<FScriptTypedElementHandle>& InElementHandles, const FTypedElementSelectionOptions InSelectionOptions)
+{
+	TArray<FTypedElementHandle> NativeHandles = TypedElementUtil::ConvertToNativeElementArray(InElementHandles);
+	return SelectElements(NativeHandles, InSelectionOptions);
+}
+
+
+bool UTypedElementSelectionSet::DeselectElement(const FScriptTypedElementHandle& InElementHandle, const FTypedElementSelectionOptions InSelectionOptions)
+{
+	FTypedElementHandle NativeHandle = InElementHandle.GetTypedElementHandle();
+	if (!NativeHandle)
+	{
+		FFrame::KismetExecutionMessage(TEXT("InElementHandle is not a valid handle."), ELogVerbosity::Error);
+		return false;
+	}
+
+	return DeselectElement(NativeHandle, InSelectionOptions);
+}
+
+bool UTypedElementSelectionSet::DeselectElements(const TArray<FScriptTypedElementHandle>& InElementHandles, const FTypedElementSelectionOptions InSelectionOptions)
+{
+	TArray<FTypedElementHandle> NativeHandles = TypedElementUtil::ConvertToNativeElementArray(InElementHandles);
+	return DeselectElements(NativeHandles, InSelectionOptions);
+}
+
+bool UTypedElementSelectionSet::SetSelection(const TArray<FScriptTypedElementHandle>& InElementHandles, const FTypedElementSelectionOptions InSelectionOptions)
+{
+	TArray<FTypedElementHandle> NativeHandles = TypedElementUtil::ConvertToNativeElementArray(InElementHandles);
+	return SetSelection(NativeHandles, InSelectionOptions);
+}
+
+bool UTypedElementSelectionSet::AllowSelectionModifiers(const FScriptTypedElementHandle& InElementHandle) const
+{
+	FTypedElementHandle NativeHandle = InElementHandle.GetTypedElementHandle();
+	if (!NativeHandle)
+	{
+		FFrame::KismetExecutionMessage(TEXT("InElementHandle is not a valid handle."), ELogVerbosity::Error);
+		return false;
+	}
+
+	return AllowSelectionModifiers(NativeHandle);
+}
+
+
+FScriptTypedElementHandle UTypedElementSelectionSet::GetSelectionElement(const FScriptTypedElementHandle& InElementHandle, const ETypedElementSelectionMethod InSelectionMethod) const
+{
+	FTypedElementHandle NativeHandle = InElementHandle.GetTypedElementHandle();
+	if (!NativeHandle)
+	{
+		FFrame::KismetExecutionMessage(TEXT("InElementHandle is not a valid handle."), ELogVerbosity::Error);
+		return FScriptTypedElementHandle();
+	}
+
+	return ElementList->GetRegistry()->CreateScriptHandle(GetSelectionElement(NativeHandle, InSelectionMethod).GetId());
+}
+
+TArray<FScriptTypedElementHandle> UTypedElementSelectionSet::K2_GetSelectedElementHandles(const TSubclassOf<UInterface> InBaseInterfaceType) const
+{
+	TArray<FTypedElementHandle> NativeHandles = GetSelectedElementHandles(InBaseInterfaceType);
+	return TypedElementUtil::ConvertToScriptElementArray(NativeHandles, ElementList->GetRegistry());
+}
diff --git a/Engine/Source/Runtime/TypedElementRuntime/Private/Elements/Framework/TypedElementSelectionSetLibrary.h b/Engine/Source/Runtime/TypedElementRuntime/Private/Elements/Framework/TypedElementSelectionSetLibrary.h
--- a/Engine/Source/Runtime/TypedElementRuntime/Private/Elements/Framework/TypedElementSelectionSetLibrary.h
+++ b/Engine/Source/Runtime/TypedElementRuntime/Private/Elements/Framework/TypedElementSelectionSetLibrary.h
@@ -2,11 +2,33 @@
 
 #pragma once
 
-#include "CoreMinimal.h"
 #include "Elements/Framework/TypedElementListProxy.h"
 #include "Elements/Framework/TypedElementSelectionSet.h"
+
+#include "CoreMinimal.h"
+#include "UObject/Stack.h"
+
 #include "TypedElementSelectionSetLibrary.generated.h"
 
+namespace UE::TypedElementFramework::Private
+{
+	// Return true if the selection set is invalid
+	bool CheckForInvalidSelectionSet(UTypedElementSelectionSet* SelectionSet)
+	{
+		if (!SelectionSet)
+		{
+			FFrame::KismetExecutionMessage(TEXT("SelectionSet is null."), ELogVerbosity::Error);
+		}
+
+		return true;
+	}
+}
+
+/**
+ * Library of functions for the scripting of Typed Elements that use both a selection set and a element list
+ * 
+ * Note: These functions should only be use for scripting proposes only as they come at higher performance cost then their non script implementation
+ */
 UCLASS()
 class UTypedElementSelectionSetLibrary : public UObject
 {
@@ -18,10 +40,21 @@ public:
 	 * @return True if the selection was changed, false otherwise.
 	 */
 	UFUNCTION(BlueprintCallable, Category="TypedElementFramework|Selection", meta=(ScriptMethod))
-	static bool SelectElementsFromList(UTypedElementSelectionSet* SelectionSet, const FTypedElementListProxy ElementList, const FTypedElementSelectionOptions SelectionOptions)
+	static bool SelectElementsFromList(UTypedElementSelectionSet* SelectionSet, const FManagedTypedElementListProxy ElementList, const FTypedElementSelectionOptions SelectionOptions)
 	{
-		FTypedElementListConstPtr ElementListPtr = ElementList.GetElementList();
-		return ElementListPtr && SelectionSet->SelectElements(ElementListPtr.ToSharedRef(), SelectionOptions);
+		if (!UE::TypedElementFramework::Private::CheckForInvalidSelectionSet(SelectionSet))
+		{
+			return false;
+		}
+
+		FTypedElementListPtr NativeList = UE::TypedElementFramework::ConvertToNativeTypedElementList(ElementList.GetElementList());
+		if (!NativeList)
+		{
+			FFrame::KismetExecutionMessage(TEXT("ElementList is in a invalid state."), ELogVerbosity::Error);
+			return false;
+		}
+
+		return SelectionSet->SelectElements(NativeList.ToSharedRef(), SelectionOptions);
 	}
 
 	/**
@@ -29,10 +62,21 @@ public:
 	 * @return True if the selection was changed, false otherwise.
 	 */
 	UFUNCTION(BlueprintCallable, Category="TypedElementFramework|Selection", meta=(ScriptMethod))
-	static bool DeselectElementsFromList(UTypedElementSelectionSet* SelectionSet, const FTypedElementListProxy ElementList, const FTypedElementSelectionOptions SelectionOptions)
+	static bool DeselectElementsFromList(UTypedElementSelectionSet* SelectionSet, const FManagedTypedElementListProxy ElementList, const FTypedElementSelectionOptions SelectionOptions)
 	{
-		FTypedElementListConstPtr ElementListPtr = ElementList.GetElementList();
-		return ElementListPtr && SelectionSet->DeselectElements(ElementListPtr.ToSharedRef(), SelectionOptions);
+		if (!UE::TypedElementFramework::Private::CheckForInvalidSelectionSet(SelectionSet))
+		{
+			return false;
+		}
+
+		FTypedElementListPtr NativeList = UE::TypedElementFramework::ConvertToNativeTypedElementList(ElementList.GetElementList());
+		if (!NativeList)
+		{
+			FFrame::KismetExecutionMessage(TEXT("ElementList is in a invalid state."), ELogVerbosity::Error);
+			return false;
+		}
+
+		return SelectionSet->DeselectElements(NativeList.ToSharedRef(), SelectionOptions);
 	}
 
 	/**
@@ -41,11 +85,16 @@ public:
 	 * @return True if the selection was changed, false otherwise.
 	 */
 	UFUNCTION(BlueprintCallable, Category="TypedElementFramework|Selection", meta=(ScriptMethod))
-	static bool SetSelectionFromList(UTypedElementSelectionSet* SelectionSet, const FTypedElementListProxy ElementList, const FTypedElementSelectionOptions SelectionOptions)
+	static bool SetSelectionFromList(UTypedElementSelectionSet* SelectionSet, const FManagedTypedElementListProxy ElementList, const FTypedElementSelectionOptions SelectionOptions)
 	{
-		FTypedElementListConstPtr ElementListPtr = ElementList.GetElementList();
-		return ElementListPtr
-			? SelectionSet->SetSelection(ElementListPtr.ToSharedRef(), SelectionOptions)
+		if (!UE::TypedElementFramework::Private::CheckForInvalidSelectionSet(SelectionSet))
+		{
+			return false;
+		}
+
+		FTypedElementListPtr NativeList = UE::TypedElementFramework::ConvertToNativeTypedElementList(ElementList.GetElementList());
+		return NativeList
+			? SelectionSet->SetSelection(NativeList.ToSharedRef(), SelectionOptions)
 			: SelectionSet->ClearSelection(SelectionOptions);
 	}
 
@@ -54,9 +103,15 @@ public:
 	 * This will do things like expand out groups, and resolve any parent<->child elements so that duplication operations aren't performed on both the parent and the child.
 	 */
 	UFUNCTION(BlueprintPure, Category="TypedElementFramework|Selection", meta=(ScriptMethod))
-	static FTypedElementListProxy GetNormalizedSelection(UTypedElementSelectionSet* SelectionSet, const FTypedElementSelectionNormalizationOptions NormalizationOptions)
+	static FManagedTypedElementListProxy GetNormalizedSelection(UTypedElementSelectionSet* SelectionSet, const FTypedElementSelectionNormalizationOptions NormalizationOptions)
 	{
-		return SelectionSet->GetNormalizedSelection(NormalizationOptions);
+		if (!UE::TypedElementFramework::Private::CheckForInvalidSelectionSet(SelectionSet))
+		{
+			return FManagedTypedElementListProxy();
+		}
+
+		FScriptTypedElementListPtr NormalizedSelection = UE::TypedElementFramework::ConvertToScriptTypedElementList(SelectionSet->GetNormalizedSelection(NormalizationOptions));
+		return NormalizedSelection ? NormalizedSelection.ToSharedRef() : FManagedTypedElementListProxy();
 	}
 
 	/**
@@ -64,9 +119,14 @@ public:
 	 * This will do things like expand out groups, and resolve any parent<->child elements so that duplication operations aren't performed on both the parent and the child.
 	 */
 	UFUNCTION(BlueprintPure, Category="TypedElementFramework|Selection", meta=(ScriptMethod))
-	static FTypedElementListProxy GetNormalizedElementList(UTypedElementSelectionSet* SelectionSet, const FTypedElementListProxy ElementList, const FTypedElementSelectionNormalizationOptions NormalizationOptions)
+	static FManagedTypedElementListProxy GetNormalizedElementList(UTypedElementSelectionSet* SelectionSet, const FManagedTypedElementListProxy ElementList, const FTypedElementSelectionNormalizationOptions NormalizationOptions)
 	{
-		FTypedElementListConstPtr ElementListPtr = ElementList.GetElementList();
-		return ElementListPtr ? SelectionSet->GetNormalizedElementList(ElementListPtr.ToSharedRef(), NormalizationOptions) : FTypedElementListProxy();
+		if (!UE::TypedElementFramework::Private::CheckForInvalidSelectionSet(SelectionSet))
+		{
+			return FManagedTypedElementListProxy();
+		}
+
+		FTypedElementListPtr NativeList = UE::TypedElementFramework::ConvertToNativeTypedElementList(ElementList.GetElementList());
+		return NativeList ?  UE::TypedElementFramework::ConvertToScriptTypedElementList(SelectionSet->GetNormalizedElementList(NativeList.ToSharedRef(), NormalizationOptions)).ToSharedRef() : FManagedTypedElementListProxy();
 	}
 };
diff --git a/Engine/Source/Runtime/TypedElementRuntime/Private/Elements/Interfaces/TypedElementAssetDataInterface.cpp b/Engine/Source/Runtime/TypedElementRuntime/Private/Elements/Interfaces/TypedElementAssetDataInterface.cpp
--- a/Engine/Source/Runtime/TypedElementRuntime/Private/Elements/Interfaces/TypedElementAssetDataInterface.cpp
+++ b/Engine/Source/Runtime/TypedElementRuntime/Private/Elements/Interfaces/TypedElementAssetDataInterface.cpp
@@ -2,6 +2,8 @@
 
 #include "Elements/Interfaces/TypedElementAssetDataInterface.h"
 
+#include "UObject/Stack.h"
+
 TArray<FAssetData> ITypedElementAssetDataInterface::GetAllReferencedAssetDatas(const FTypedElementHandle& InElementHandle)
 {
 	TArray<FAssetData> AssetDatas;
@@ -19,3 +21,27 @@ FAssetData ITypedElementAssetDataInterface::GetAssetData(const FTypedElementHand
 {
 	return FAssetData();
 }
+
+TArray<FAssetData> ITypedElementAssetDataInterface::GetAllReferencedAssetDatas(const FScriptTypedElementHandle& InElementHandle)
+{
+	FTypedElementHandle NativeHandle = InElementHandle.GetTypedElementHandle();
+	if (!NativeHandle)
+	{
+		FFrame::KismetExecutionMessage(TEXT("InElementHandle is not a valid handle."), ELogVerbosity::Error);
+		return {};
+	}
+
+	return GetAllReferencedAssetDatas(NativeHandle);
+}
+
+FAssetData ITypedElementAssetDataInterface::GetAssetData(const FScriptTypedElementHandle& InElementHandle)
+{
+	FTypedElementHandle NativeHandle = InElementHandle.GetTypedElementHandle();
+	if (!NativeHandle)
+	{
+		FFrame::KismetExecutionMessage(TEXT("InElementHandle is not a valid handle."), ELogVerbosity::Error);
+		return FAssetData();
+	}
+
+	return GetAssetData(NativeHandle);
+}
diff --git a/Engine/Source/Runtime/TypedElementRuntime/Private/Elements/Interfaces/TypedElementHierarchyInterface.cpp b/Engine/Source/Runtime/TypedElementRuntime/Private/Elements/Interfaces/TypedElementHierarchyInterface.cpp
new file mode 100644
--- /dev/null
+++ b/Engine/Source/Runtime/TypedElementRuntime/Private/Elements/Interfaces/TypedElementHierarchyInterface.cpp
@@ -0,0 +1,41 @@
+// Copyright Epic Games, Inc. All Rights Reserved.
+
+#include "Elements/Interfaces/TypedElementHierarchyInterface.h"
+
+#include "Elements/Framework/TypedElementRegistry.h"
+#include "Elements/Framework/TypedElementUtil.h"
+
+#include "UObject/Stack.h"
+
+FScriptTypedElementHandle ITypedElementHierarchyInterface::GetParentElement(const FScriptTypedElementHandle& InElementHandle, const bool bAllowCreate)
+{
+	FTypedElementHandle NativeHandle = InElementHandle.GetTypedElementHandle();
+	if (!NativeHandle)
+	{
+		FFrame::KismetExecutionMessage(TEXT("InElementHandle is not a valid handle."), ELogVerbosity::Error);
+		return FScriptTypedElementHandle();
+	}
+	
+	return GetRegistry().CreateScriptHandle(GetParentElement(NativeHandle, bAllowCreate).GetId());
+}
+
+void ITypedElementHierarchyInterface::GetChildElements(const FScriptTypedElementHandle& InElementHandle, TArray<FScriptTypedElementHandle>& OutElementHandles, const bool bAllowCreate)
+{
+	FTypedElementHandle NativeHandle = InElementHandle.GetTypedElementHandle();
+	if (!NativeHandle)
+	{
+		FFrame::KismetExecutionMessage(TEXT("InElementHandle is not a valid handle."), ELogVerbosity::Error);
+		return;
+	}
+
+	TArray<FTypedElementHandle> NativeChildHandles;
+	GetChildElements(NativeHandle, NativeChildHandles, bAllowCreate);
+
+	OutElementHandles = TypedElementUtil::ConvertToScriptElementArray(NativeChildHandles, &GetRegistry());
+}
+
+class UTypedElementRegistry& ITypedElementHierarchyInterface::GetRegistry() const
+{
+	return *UTypedElementRegistry::GetInstance();
+}
+
diff --git a/Engine/Source/Runtime/TypedElementRuntime/Private/Elements/Interfaces/TypedElementObjectInterface.cpp b/Engine/Source/Runtime/TypedElementRuntime/Private/Elements/Interfaces/TypedElementObjectInterface.cpp
--- a/Engine/Source/Runtime/TypedElementRuntime/Private/Elements/Interfaces/TypedElementObjectInterface.cpp
+++ b/Engine/Source/Runtime/TypedElementRuntime/Private/Elements/Interfaces/TypedElementObjectInterface.cpp
@@ -4,6 +4,8 @@
 
 #include "Elements/Framework/TypedElementRegistry.h"
 
+#include "UObject/Stack.h"
+
 UObject* ITypedElementObjectInterface::GetObject(const FTypedElementHandle& InElementHandle)
 {
 	return nullptr;
@@ -14,3 +16,28 @@ UClass* ITypedElementObjectInterface::GetObjectClass(const FTypedElementHandle&
 	UObject* HandleAsObject = GetObject(InElementHandle);
 	return HandleAsObject ? HandleAsObject->GetClass() : nullptr;
 }
+
+UObject* ITypedElementObjectInterface::GetObject(const FScriptTypedElementHandle& InElementHandle)
+{
+	FTypedElementHandle NativeHandle = InElementHandle.GetTypedElementHandle();
+	if (!NativeHandle)
+	{
+		FFrame::KismetExecutionMessage(TEXT("InElementHandle is not a valid handle."), ELogVerbosity::Error);
+		return nullptr;
+	}
+
+	return GetObject(NativeHandle);
+}
+
+UClass* ITypedElementObjectInterface::GetObjectClass(const FScriptTypedElementHandle& InElementHandle)
+{
+	FTypedElementHandle NativeHandle = InElementHandle.GetTypedElementHandle();
+	if (!NativeHandle)
+	{
+		FFrame::KismetExecutionMessage(TEXT("InElementHandle is not a valid handle."), ELogVerbosity::Error);
+		return nullptr;
+	}
+
+	return GetObjectClass(NativeHandle);
+}
+
diff --git a/Engine/Source/Runtime/TypedElementRuntime/Private/Elements/Interfaces/TypedElementSelectionInterface.cpp b/Engine/Source/Runtime/TypedElementRuntime/Private/Elements/Interfaces/TypedElementSelectionInterface.cpp
--- a/Engine/Source/Runtime/TypedElementRuntime/Private/Elements/Interfaces/TypedElementSelectionInterface.cpp
+++ b/Engine/Source/Runtime/TypedElementRuntime/Private/Elements/Interfaces/TypedElementSelectionInterface.cpp
@@ -1,22 +1,143 @@
 // Copyright Epic Games, Inc. All Rights Reserved.
 
 #include "Elements/Interfaces/TypedElementSelectionInterface.h"
+
 #include "Elements/Framework/TypedElementList.h"
+#include "Elements/Framework/TypedElementRegistry.h"
+#include "Elements/Framework/TypedElementUtil.h"
+
+#include "UObject/Stack.h"
+
+bool ITypedElementSelectionInterface::IsElementSelected(const FTypedElementHandle& InElementHandle, const FTypedElementListConstPtr& InSelectionSet, const FTypedElementIsSelectedOptions& InSelectionOptions)
+{
+	return InSelectionSet && InSelectionSet->Contains(InElementHandle);
+}
+
+bool ITypedElementSelectionInterface::SelectElement(const FTypedElementHandle& InElementHandle, const FTypedElementListPtr& InSelectionSet, const FTypedElementSelectionOptions& InSelectionOptions)
+{
+	return InSelectionSet && InSelectionSet->Add(InElementHandle);
+}
+
+bool ITypedElementSelectionInterface::DeselectElement(const FTypedElementHandle& InElementHandle, const FTypedElementListPtr& InSelectionSet, const FTypedElementSelectionOptions& InSelectionOptions)
+{
+	return InSelectionSet && InSelectionSet->Remove(InElementHandle);
+}
 
-bool ITypedElementSelectionInterface::IsElementSelected(const FTypedElementHandle& InElementHandle, const FTypedElementListProxy InSelectionSet, const FTypedElementIsSelectedOptions& InSelectionOptions)
+bool ITypedElementSelectionInterface::IsElementSelected(const FScriptTypedElementHandle& InElementHandle, const FManagedTypedElementListProxy InSelectionSet, const FTypedElementIsSelectedOptions& InSelectionOptions)
 {
-	FTypedElementListConstPtr SelectionSetPtr = InSelectionSet.GetElementList();
-	return SelectionSetPtr && SelectionSetPtr->Contains(InElementHandle);
+	FTypedElementHandle NativeHandle = InElementHandle.GetTypedElementHandle();
+	if (!NativeHandle)
+	{
+		FFrame::KismetExecutionMessage(TEXT("InElementHandle is not a valid handle."), ELogVerbosity::Error);
+		return false;
+	}
+
+	FTypedElementListPtr NativeList = UE::TypedElementFramework::ConvertToNativeTypedElementList(InSelectionSet.GetElementList());
+	if (!NativeList)
+	{
+		FFrame::KismetExecutionMessage(TEXT("InSelectionSet is in a invalid state."), ELogVerbosity::Error);
+		return false;
+	}
+
+	return IsElementSelected(NativeHandle, NativeList, InSelectionOptions);
 }
 
-bool ITypedElementSelectionInterface::SelectElement(const FTypedElementHandle& InElementHandle, FTypedElementListProxy InSelectionSet, const FTypedElementSelectionOptions& InSelectionOptions)
+bool ITypedElementSelectionInterface::CanSelectElement(const FScriptTypedElementHandle& InElementHandle, const FTypedElementSelectionOptions& InSelectionOptions)
 {
-	FTypedElementListPtr SelectionSetPtr = InSelectionSet.GetElementList();
-	return SelectionSetPtr && SelectionSetPtr->Add(InElementHandle);
+	FTypedElementHandle NativeHandle = InElementHandle.GetTypedElementHandle();
+	if (!NativeHandle)
+	{
+		FFrame::KismetExecutionMessage(TEXT("InElementHandle is not a valid handle."), ELogVerbosity::Error);
+		return false;
+	}
+
+	return CanSelectElement(NativeHandle, InSelectionOptions);
+}
+
+bool ITypedElementSelectionInterface::CanDeselectElement(const FScriptTypedElementHandle& InElementHandle, const FTypedElementSelectionOptions& InSelectionOptions)
+{
+	FTypedElementHandle NativeHandle = InElementHandle.GetTypedElementHandle();
+	if (!NativeHandle)
+	{
+		FFrame::KismetExecutionMessage(TEXT("InElementHandle is not a valid handle."), ELogVerbosity::Error);
+		return false;
+	}
+
+	return CanDeselectElement(NativeHandle, InSelectionOptions);
 }
 
-bool ITypedElementSelectionInterface::DeselectElement(const FTypedElementHandle& InElementHandle, FTypedElementListProxy InSelectionSet, const FTypedElementSelectionOptions& InSelectionOptions)
+bool ITypedElementSelectionInterface::SelectElement(const FScriptTypedElementHandle& InElementHandle, FManagedTypedElementListProxy InSelectionSet, const FTypedElementSelectionOptions& InSelectionOptions)
 {
-	FTypedElementListPtr SelectionSetPtr = InSelectionSet.GetElementList();
-	return SelectionSetPtr && SelectionSetPtr->Remove(InElementHandle);
+	FTypedElementHandle NativeHandle = InElementHandle.GetTypedElementHandle();
+	if (!NativeHandle)
+	{
+		FFrame::KismetExecutionMessage(TEXT("InElementHandle is not a valid handle."), ELogVerbosity::Error);
+		return false;
+	}
+
+	FTypedElementListPtr NativeList = UE::TypedElementFramework::ConvertToNativeTypedElementList(InSelectionSet.GetElementList());
+	if (!NativeList)
+	{
+		FFrame::KismetExecutionMessage(TEXT("InSelectionSet is in a invalid state."), ELogVerbosity::Error);
+		return false;
+	}
+
+	return SelectElement(NativeHandle, NativeList, InSelectionOptions);
+}
+
+bool ITypedElementSelectionInterface::DeselectElement(const FScriptTypedElementHandle& InElementHandle, FManagedTypedElementListProxy InSelectionSet, const FTypedElementSelectionOptions& InSelectionOptions)
+{
+	FTypedElementHandle NativeHandle = InElementHandle.GetTypedElementHandle();
+	if (!NativeHandle)
+	{
+		FFrame::KismetExecutionMessage(TEXT("InElementHandle is not a valid handle."), ELogVerbosity::Error);
+		return false;
+	}
+
+	FTypedElementListPtr NativeList = UE::TypedElementFramework::ConvertToNativeTypedElementList(InSelectionSet.GetElementList());
+	if (!NativeList)
+	{
+		FFrame::KismetExecutionMessage(TEXT("InSelectionSet is in a invalid state."), ELogVerbosity::Error);
+		return false;
+	}
+
+	return DeselectElement(NativeHandle, NativeList, InSelectionOptions);
+}
+
+bool ITypedElementSelectionInterface::AllowSelectionModifiers(const FScriptTypedElementHandle& InElementHandle, const FManagedTypedElementListProxy InSelectionSet)
+{
+	FTypedElementHandle NativeHandle = InElementHandle.GetTypedElementHandle();
+	if (!NativeHandle)
+	{
+		FFrame::KismetExecutionMessage(TEXT("InElementHandle is not a valid handle."), ELogVerbosity::Error);
+		return false;
+	}
+
+	FTypedElementListPtr NativeList = UE::TypedElementFramework::ConvertToNativeTypedElementList(InSelectionSet.GetElementList());
+	if (!NativeList)
+	{
+		FFrame::KismetExecutionMessage(TEXT("InSelectionSet is in a invalid state."), ELogVerbosity::Error);
+		return false;
+	}
+
+	return AllowSelectionModifiers(NativeHandle, NativeList);
+}
+
+FScriptTypedElementHandle ITypedElementSelectionInterface::GetSelectionElement(const FScriptTypedElementHandle& InElementHandle, const FManagedTypedElementListProxy InCurrentSelection, const ETypedElementSelectionMethod InSelectionMethod)
+{
+	FTypedElementHandle NativeHandle = InElementHandle.GetTypedElementHandle();
+	if (!NativeHandle)
+	{
+		FFrame::KismetExecutionMessage(TEXT("InElementHandle is not a valid handle."), ELogVerbosity::Error);
+		return {};
+	}
+
+	FTypedElementListPtr NativeList = UE::TypedElementFramework::ConvertToNativeTypedElementList(InCurrentSelection.GetElementList());
+	if (!NativeList)
+	{
+		FFrame::KismetExecutionMessage(TEXT("InCurrentSelection is in a invalid state."), ELogVerbosity::Error);
+		return {};
+	}
+
+	return NativeList->GetRegistry()->CreateScriptHandle(GetSelectionElement(NativeHandle, NativeList, InSelectionMethod).GetId());
 }
diff --git a/Engine/Source/Runtime/TypedElementRuntime/Public/Elements/Framework/TypedElementSelectionSet.h b/Engine/Source/Runtime/TypedElementRuntime/Public/Elements/Framework/TypedElementSelectionSet.h
--- a/Engine/Source/Runtime/TypedElementRuntime/Public/Elements/Framework/TypedElementSelectionSet.h
+++ b/Engine/Source/Runtime/TypedElementRuntime/Public/Elements/Framework/TypedElementSelectionSet.h
@@ -165,33 +165,28 @@ public:
 	/**
 	 * Test to see whether the given element is currently considered selected.
 	 */
-	UFUNCTION(BlueprintPure, Category="TypedElementFramework|Selection")
 	bool IsElementSelected(const FTypedElementHandle& InElementHandle, const FTypedElementIsSelectedOptions InSelectionOptions) const;
 
 	/**
 	 * Test to see whether the given element can be selected.
 	 */
-	UFUNCTION(BlueprintPure, Category="TypedElementFramework|Selection")
 	bool CanSelectElement(const FTypedElementHandle& InElementHandle, const FTypedElementSelectionOptions InSelectionOptions) const;
 
 	/**
 	 * Test to see whether the given element can be deselected.
 	 */
-	UFUNCTION(BlueprintPure, Category="TypedElementFramework|Selection")
 	bool CanDeselectElement(const FTypedElementHandle& InElementHandle, const FTypedElementSelectionOptions InSelectionOptions) const;
 
 	/**
 	 * Attempt to select the given element.
 	 * @return True if the selection was changed, false otherwise.
 	 */
-	UFUNCTION(BlueprintCallable, Category="TypedElementFramework|Selection")
 	bool SelectElement(const FTypedElementHandle& InElementHandle, const FTypedElementSelectionOptions InSelectionOptions);
 
 	/**
 	 * Attempt to select the given elements.
 	 * @return True if the selection was changed, false otherwise.
 	 */
-	UFUNCTION(BlueprintCallable, Category="TypedElementFramework|Selection")
 	bool SelectElements(const TArray<FTypedElementHandle>& InElementHandles, const FTypedElementSelectionOptions InSelectionOptions);
 	bool SelectElements(TArrayView<const FTypedElementHandle> InElementHandles, const FTypedElementSelectionOptions InSelectionOptions);
 	bool SelectElements(FTypedElementListConstRef InElementList, const FTypedElementSelectionOptions InSelectionOptions);
@@ -200,14 +195,12 @@ public:
 	 * Attempt to deselect the given element.
 	 * @return True if the selection was changed, false otherwise.
 	 */
-	UFUNCTION(BlueprintCallable, Category="TypedElementFramework|Selection")
 	bool DeselectElement(const FTypedElementHandle& InElementHandle, const FTypedElementSelectionOptions InSelectionOptions);
 
 	/**
 	 * Attempt to deselect the given elements.
 	 * @return True if the selection was changed, false otherwise.
 	 */
-	UFUNCTION(BlueprintCallable, Category="TypedElementFramework|Selection")
 	bool DeselectElements(const TArray<FTypedElementHandle>& InElementHandles, const FTypedElementSelectionOptions InSelectionOptions);
 	bool DeselectElements(TArrayView<const FTypedElementHandle> InElementHandles, const FTypedElementSelectionOptions InSelectionOptions);
 	bool DeselectElements(FTypedElementListConstRef InElementList, const FTypedElementSelectionOptions InSelectionOptions);
@@ -224,7 +217,6 @@ public:
 	 * @note Equivalent to calling ClearSelection then SelectElements, but happens in a single batch.
 	 * @return True if the selection was changed, false otherwise.
 	 */
-	UFUNCTION(BlueprintCallable, Category="TypedElementFramework|Selection")
 	bool SetSelection(const TArray<FTypedElementHandle>& InElementHandles, const FTypedElementSelectionOptions InSelectionOptions);
 	bool SetSelection(TArrayView<const FTypedElementHandle> InElementHandles, const FTypedElementSelectionOptions InSelectionOptions);
 	bool SetSelection(FTypedElementListConstRef InElementList, const FTypedElementSelectionOptions InSelectionOptions);
@@ -232,13 +224,11 @@ public:
 	/**
 	 * Test to see whether selection modifiers (Ctrl or Shift) are allowed while selecting this element.
 	 */
-	UFUNCTION(BlueprintPure, Category="TypedElementFramework|Selection")
 	bool AllowSelectionModifiers(const FTypedElementHandle& InElementHandle) const;
 
 	/**
 	 * Given an element, return the element that should actually perform a selection operation.
 	 */
-	UFUNCTION(BlueprintPure, Category="TypedElementFramework|Selection")
 	FTypedElementHandle GetSelectionElement(const FTypedElementHandle& InElementHandle, const ETypedElementSelectionMethod InSelectionMethod) const;
 
 	/**
@@ -285,7 +275,6 @@ public:
 	/**
 	 * Get the handle of every selected element, optionally filtering to elements that implement the given interface.
 	 */
-	UFUNCTION(BlueprintCallable, BlueprintPure=false, Category="TypedElementFramework|Selection")
 	TArray<FTypedElementHandle> GetSelectedElementHandles(const TSubclassOf<UInterface> InBaseInterfaceType = nullptr) const
 	{
 		return ElementList->GetElementHandles(InBaseInterfaceType);
@@ -506,7 +495,7 @@ public:
 	 * Access the interface to allow external systems (such as USelection) to receive immediate sync notifications as the underlying element list is changed.
 	 * This exists purely as a bridging mechanism and shouldn't be relied on for new code. It is lazily created as needed.
 	 */
-	FTypedElementListLegacySync& Legacy_GetElementListSync()
+	FTypedElementList::FLegacySync& Legacy_GetElementListSync()
 	{
 		return ElementList->Legacy_GetSync();
 	}
@@ -515,7 +504,7 @@ public:
 	 * Access the interface to allow external systems (such as USelection) to receive immediate sync notifications as the underlying element list is changed.
 	 * This exists purely as a bridging mechanism and shouldn't be relied on for new code. This will return null if no legacy sync has been created for this instance.
 	 */
-	FTypedElementListLegacySync* Legacy_GetElementListSyncPtr()
+	FTypedElementList::FLegacySync* Legacy_GetElementListSyncPtr()
 	{
 		return ElementList->Legacy_GetSyncPtr();
 	}
@@ -544,6 +533,84 @@ public:
 	UFUNCTION(BlueprintCallable, Category = "TypedElementFramework|Selection")
 	void RestoreSelectionState(const FTypedElementSelectionSetState& InSelectionState);
 
+
+	/*
+	 * Script Api
+	 */
+
+	/**
+	 * Test to see whether the given element is currently considered selected.
+	 */
+	UFUNCTION(BlueprintPure, Category="TypedElementFramework|Selection")
+	bool IsElementSelected(const FScriptTypedElementHandle& InElementHandle, const FTypedElementIsSelectedOptions InSelectionOptions) const;
+
+	/**
+	 * Test to see whether the given element can be selected.
+	 */
+	UFUNCTION(BlueprintPure, Category="TypedElementFramework|Selection")
+	bool CanSelectElement(const FScriptTypedElementHandle& InElementHandle, const FTypedElementSelectionOptions InSelectionOptions) const;
+
+	/**
+	 * Test to see whether the given element can be deselected.
+	 */
+	UFUNCTION(BlueprintPure, Category="TypedElementFramework|Selection")
+	bool CanDeselectElement(const FScriptTypedElementHandle& InElementHandle, const FTypedElementSelectionOptions InSelectionOptions) const;
+
+	/**
+	 * Attempt to select the given element.
+	 * @return True if the selection was changed, false otherwise.
+	 */
+	UFUNCTION(BlueprintCallable, Category="TypedElementFramework|Selection")
+	bool SelectElement(const FScriptTypedElementHandle& InElementHandle, const FTypedElementSelectionOptions InSelectionOptions);
+
+	/**
+	 * Attempt to select the given elements.
+	 * @return True if the selection was changed, false otherwise.
+	 */
+	UFUNCTION(BlueprintCallable, Category="TypedElementFramework|Selection")
+	bool SelectElements(const TArray<FScriptTypedElementHandle>& InElementHandles, const FTypedElementSelectionOptions InSelectionOptions);
+
+	/**
+	 * Attempt to deselect the given element.
+	 * @return True if the selection was changed, false otherwise.
+	 */
+	UFUNCTION(BlueprintCallable, Category="TypedElementFramework|Selection")
+	bool DeselectElement(const FScriptTypedElementHandle& InElementHandle, const FTypedElementSelectionOptions InSelectionOptions);
+
+	/**
+	 * Attempt to deselect the given elements.
+	 * @return True if the selection was changed, false otherwise.
+	 */
+	UFUNCTION(BlueprintCallable, Category="TypedElementFramework|Selection")
+	bool DeselectElements(const TArray<FScriptTypedElementHandle>& InElementHandles, const FTypedElementSelectionOptions InSelectionOptions);
+
+	/**
+	 * Attempt to make the selection the given elements.
+	 * @note Equivalent to calling ClearSelection then SelectElements, but happens in a single batch.
+	 * @return True if the selection was changed, false otherwise.
+	 */
+	UFUNCTION(BlueprintCallable, Category="TypedElementFramework|Selection")
+	bool SetSelection(const TArray<FScriptTypedElementHandle>& InElementHandles, const FTypedElementSelectionOptions InSelectionOptions);
+
+	/**
+	 * Test to see whether selection modifiers (Ctrl or Shift) are allowed while selecting this element.
+	 */
+	UFUNCTION(BlueprintPure, Category="TypedElementFramework|Selection")
+	bool AllowSelectionModifiers(const FScriptTypedElementHandle& InElementHandle) const;
+
+	/**
+	 * Given an element, return the element that should actually perform a selection operation.
+	 */
+	UFUNCTION(BlueprintPure, Category = "TypedElementFramework|Selection")
+	FScriptTypedElementHandle GetSelectionElement(const FScriptTypedElementHandle& InElementHandle, const ETypedElementSelectionMethod InSelectionMethod) const;
+
+
+	/**
+	 * Get the handle of every selected element, optionally filtering to elements that implement the given interface.
+	 */
+	UFUNCTION(BlueprintCallable, BlueprintPure=false, DisplayName="Get Selected Element Handles", Category="TypedElementFramework|Selection", meta=(ScriptName="GetSelectedElementHandles"))
+	TArray<FScriptTypedElementHandle> K2_GetSelectedElementHandles(const TSubclassOf<UInterface> InBaseInterfaceType = nullptr) const;
+
 	DECLARE_DYNAMIC_MULTICAST_DELEGATE_OneParam(FOnPreChangeDynamic, const UTypedElementSelectionSet*, SelectionSet);
 
 	/** Delegate that is invoked whenever the underlying element list is potentially about to change. */
diff --git a/Engine/Source/Runtime/TypedElementRuntime/Public/Elements/Interfaces/TypedElementAssetDataInterface.h b/Engine/Source/Runtime/TypedElementRuntime/Public/Elements/Interfaces/TypedElementAssetDataInterface.h
--- a/Engine/Source/Runtime/TypedElementRuntime/Public/Elements/Interfaces/TypedElementAssetDataInterface.h
+++ b/Engine/Source/Runtime/TypedElementRuntime/Public/Elements/Interfaces/TypedElementAssetDataInterface.h
@@ -25,14 +25,31 @@ public:
 	 *
 	 * @returns An array of valid asset datas.
 	 */
-	UFUNCTION(BlueprintCallable, Category = "TypedElementInterfaces|AssetData")
 	virtual TArray<FAssetData> GetAllReferencedAssetDatas(const FTypedElementHandle& InElementHandle);
 
 	/**
 	 * Returns the asset data for the given handle, if it exists.
 	 */
-	UFUNCTION(BlueprintCallable, Category = "TypedElementInterfaces|AssetData")
 	virtual FAssetData GetAssetData(const FTypedElementHandle& InElementHandle);
+
+	/**
+	 * Script Api
+	 */
+
+	/**
+	 * Returns any asset datas for content objects referenced by handle.
+	 * If the given handle itself has valid asset data, it should be returned as the last element of the array.
+	 *
+	 * @returns An array of valid asset datas.
+	 */
+	UFUNCTION(BlueprintCallable, Category = "TypedElementInterfaces|AssetData")
+	virtual TArray<FAssetData> GetAllReferencedAssetDatas(const FScriptTypedElementHandle& InElementHandle);
+
+	/**
+	 * Returns the asset data for the given handle, if it exists.
+	 */
+	UFUNCTION(BlueprintCallable, Category = "TypedElementInterfaces|AssetData")
+	virtual FAssetData GetAssetData(const FScriptTypedElementHandle& InElementHandle);
 };
 
 template <>
diff --git a/Engine/Source/Runtime/TypedElementRuntime/Public/Elements/Interfaces/TypedElementHierarchyInterface.h b/Engine/Source/Runtime/TypedElementRuntime/Public/Elements/Interfaces/TypedElementHierarchyInterface.h
--- a/Engine/Source/Runtime/TypedElementRuntime/Public/Elements/Interfaces/TypedElementHierarchyInterface.h
+++ b/Engine/Source/Runtime/TypedElementRuntime/Public/Elements/Interfaces/TypedElementHierarchyInterface.h
@@ -22,7 +22,6 @@ public:
 	 * Get the logical parent of this element, if any.
 	 * eg) A component might return its actor, or a static mesh instance might return its ISM component.
 	 */
-	UFUNCTION(BlueprintCallable, Category="TypedElementInterfaces|Hierarchy")
 	virtual FTypedElementHandle GetParentElement(const FTypedElementHandle& InElementHandle, const bool bAllowCreate = true) { return FTypedElementHandle(); }
 
 	/**
@@ -31,8 +30,34 @@ public:
 	 *
 	 * @note Appends to OutElementHandles.
 	 */
-	UFUNCTION(BlueprintCallable, Category="TypedElementInterfaces|Hierarchy")
 	virtual void GetChildElements(const FTypedElementHandle& InElementHandle, TArray<FTypedElementHandle>& OutElementHandles, const bool bAllowCreate = true) {}
+
+	/**
+	 * Script Api
+	 */
+
+	/**
+	 * Get the logical parent of this element, if any.
+	 * eg) A component might return its actor, or a static mesh instance might return its ISM component.
+	 */
+	UFUNCTION(BlueprintCallable, Category="TypedElementInterfaces|Hierarchy")
+	virtual FScriptTypedElementHandle GetParentElement(const FScriptTypedElementHandle& InElementHandle, const bool bAllowCreate = true);
+
+	/**
+	 * Get the logical children of this element, if any.
+	 * eg) An actor might return its component, or an ISM component might return its static mesh instances.
+	 *
+	 * @note Appends to OutElementHandles.
+	 */
+	UFUNCTION(BlueprintCallable, Category="TypedElementInterfaces|Hierarchy")
+	virtual void GetChildElements(const FScriptTypedElementHandle& InElementHandle, TArray<FScriptTypedElementHandle>& OutElementHandles, const bool bAllowCreate = true);
+
+private:
+	
+	/**
+	 * Return the registry used when creating new elements.
+	 */
+	virtual class UTypedElementRegistry& GetRegistry() const;
 };
 
 template <>
diff --git a/Engine/Source/Runtime/TypedElementRuntime/Public/Elements/Interfaces/TypedElementObjectInterface.h b/Engine/Source/Runtime/TypedElementRuntime/Public/Elements/Interfaces/TypedElementObjectInterface.h
--- a/Engine/Source/Runtime/TypedElementRuntime/Public/Elements/Interfaces/TypedElementObjectInterface.h
+++ b/Engine/Source/Runtime/TypedElementRuntime/Public/Elements/Interfaces/TypedElementObjectInterface.h
@@ -21,13 +21,11 @@ public:
 	/**
 	 * Get the object instance that this handle represents, if any.
 	 */
-	UFUNCTION(BlueprintCallable, Category="TypedElementInterfaces|Object")
 	virtual UObject* GetObject(const FTypedElementHandle& InElementHandle);
 
 	/**
 	 * Gets the object instance's class that the handle represents, if any. 
 	 */
-	UFUNCTION(BlueprintCallable, Category = "TypedElementInterfaces|Object")
 	virtual UClass* GetObjectClass(const FTypedElementHandle& InElementHandle);
 
 	/**
@@ -63,6 +61,23 @@ public:
 
 		return nullptr;
 	}
+
+
+	/**
+	 * Script api
+	 */
+
+	/**
+	 * Get the object instance that this handle represents, if any.
+	 */
+	UFUNCTION(BlueprintCallable, Category="TypedElementInterfaces|Object")
+	virtual UObject* GetObject(const FScriptTypedElementHandle& InElementHandle);
+
+	/**
+	 * Gets the object instance's class that the handle represents, if any. 
+	 */
+	UFUNCTION(BlueprintCallable, Category = "TypedElementInterfaces|Object")
+	virtual UClass* GetObjectClass(const FScriptTypedElementHandle& InElementHandle);
 };
 
 template <>
diff --git a/Engine/Source/Runtime/TypedElementRuntime/Public/Elements/Interfaces/TypedElementSelectionInterface.h b/Engine/Source/Runtime/TypedElementRuntime/Public/Elements/Interfaces/TypedElementSelectionInterface.h
--- a/Engine/Source/Runtime/TypedElementRuntime/Public/Elements/Interfaces/TypedElementSelectionInterface.h
+++ b/Engine/Source/Runtime/TypedElementRuntime/Public/Elements/Interfaces/TypedElementSelectionInterface.h
@@ -143,46 +143,39 @@ public:
 	/**
 	 * Test to see whether the given element is currently considered selected.
 	 */
-	UFUNCTION(BlueprintCallable, Category="TypedElementInterfaces|Selection")
-	virtual bool IsElementSelected(const FTypedElementHandle& InElementHandle, const FTypedElementListProxy InSelectionSet, const FTypedElementIsSelectedOptions& InSelectionOptions);
+	virtual bool IsElementSelected(const FTypedElementHandle& InElementHandle, const FTypedElementListConstPtr& SelectionSetPtr, const FTypedElementIsSelectedOptions& InSelectionOptions);
 
 	/**
 	 * Test to see whether the given element can be selected.
 	 */
-	UFUNCTION(BlueprintCallable, Category="TypedElementInterfaces|Selection")
 	virtual bool CanSelectElement(const FTypedElementHandle& InElementHandle, const FTypedElementSelectionOptions& InSelectionOptions) { return true; }
 
 	/**
 	 * Test to see whether the given element can be deselected.
 	 */
-	UFUNCTION(BlueprintCallable, Category="TypedElementInterfaces|Selection")
 	virtual bool CanDeselectElement(const FTypedElementHandle& InElementHandle, const FTypedElementSelectionOptions& InSelectionOptions) { return true; }
 
 	/**
 	 * Attempt to select the given element.
 	 * @return True if the selection was changed, false otherwise.
 	 */
-	UFUNCTION(BlueprintCallable, Category="TypedElementInterfaces|Selection")
-	virtual bool SelectElement(const FTypedElementHandle& InElementHandle, FTypedElementListProxy InSelectionSet, const FTypedElementSelectionOptions& InSelectionOptions);
+	virtual bool SelectElement(const FTypedElementHandle& InElementHandle, const FTypedElementListPtr& InSelectionSet, const FTypedElementSelectionOptions& InSelectionOptions);
 
 	/**
 	 * Attempt to deselect the given element.
 	 * @return True if the selection was changed, false otherwise.
 	 */
-	UFUNCTION(BlueprintCallable, Category="TypedElementInterfaces|Selection")
-	virtual bool DeselectElement(const FTypedElementHandle& InElementHandle, FTypedElementListProxy InSelectionSet, const FTypedElementSelectionOptions& InSelectionOptions);
+	virtual bool DeselectElement(const FTypedElementHandle& InElementHandle, const FTypedElementListPtr& InSelectionSet, const FTypedElementSelectionOptions& InSelectionOptions);
 
 	/**
 	 * Test to see whether selection modifiers (Ctrl or Shift) are allowed while selecting this element.
 	 */
-	UFUNCTION(BlueprintCallable, Category="TypedElementInterfaces|Selection")
-	virtual bool AllowSelectionModifiers(const FTypedElementHandle& InElementHandle, const FTypedElementListProxy InSelectionSet) { return true; }
+	virtual bool AllowSelectionModifiers(const FTypedElementHandle& InElementHandle, const FTypedElementListConstPtr& InSelectionSet) { return true; }
 
 	/**
 	 * Given an element, return the element that should actually perform a selection operation.
 	 */
-	UFUNCTION(BlueprintCallable, Category="TypedElementInterfaces|Selection")
-	virtual FTypedElementHandle GetSelectionElement(const FTypedElementHandle& InElementHandle, const FTypedElementListProxy InCurrentSelection, const ETypedElementSelectionMethod InSelectionMethod) { return InElementHandle; }
+	virtual FTypedElementHandle GetSelectionElement(const FTypedElementHandle& InElementHandle, const FTypedElementListConstPtr& InCurrentSelection, const ETypedElementSelectionMethod InSelectionMethod) { return InElementHandle; }
 
 	/**
 	 * Test to see whether the given element prevents the selection set state from being transacted for undo/redo (eg, if the element belongs to a PIE instance).
@@ -215,6 +208,55 @@ public:
 		return TransactedElement;
 	}
 
+
+	/**
+	 * Script Api
+	 */
+
+	/**
+	 * Test to see whether the given element is currently considered selected.
+	 */
+	UFUNCTION(BlueprintCallable, Category = "TypedElementInterfaces|Selection")
+	virtual bool IsElementSelected(const FScriptTypedElementHandle& InElementHandle, const FManagedTypedElementListProxy InSelectionSet, const FTypedElementIsSelectedOptions& InSelectionOptions);
+
+	/**
+	 * Test to see whether the given element can be selected.
+	 */
+	UFUNCTION(BlueprintCallable, Category="TypedElementInterfaces|Selection")
+	virtual bool CanSelectElement(const FScriptTypedElementHandle& InElementHandle, const FTypedElementSelectionOptions& InSelectionOptions);
+
+	/**
+	 * Test to see whether the given element can be deselected.
+	 */
+	UFUNCTION(BlueprintCallable, Category="TypedElementInterfaces|Selection")
+	virtual bool CanDeselectElement(const FScriptTypedElementHandle& InElementHandle, const FTypedElementSelectionOptions& InSelectionOptions);
+
+	/**
+	 * Attempt to select the given element.
+	 * @return True if the selection was changed, false otherwise.
+	 */
+	UFUNCTION(BlueprintCallable, Category="TypedElementInterfaces|Selection")
+	virtual bool SelectElement(const FScriptTypedElementHandle& InElementHandle, FManagedTypedElementListProxy InSelectionSet, const FTypedElementSelectionOptions& InSelectionOptions);
+
+	/**
+	 * Attempt to deselect the given element.
+	 * @return True if the selection was changed, false otherwise.
+	 */
+	UFUNCTION(BlueprintCallable, Category="TypedElementInterfaces|Selection")
+	virtual bool DeselectElement(const FScriptTypedElementHandle& InElementHandle, FManagedTypedElementListProxy InSelectionSet, const FTypedElementSelectionOptions& InSelectionOptions);
+
+	/**
+	 * Test to see whether selection modifiers (Ctrl or Shift) are allowed while selecting this element.
+	 */
+	UFUNCTION(BlueprintCallable, Category="TypedElementInterfaces|Selection")
+	virtual bool AllowSelectionModifiers(const FScriptTypedElementHandle& InElementHandle, const FManagedTypedElementListProxy InSelectionSet);
+
+	/**
+	 * Given an element, return the element that should actually perform a selection operation.
+	 */
+	UFUNCTION(BlueprintCallable, Category="TypedElementInterfaces|Selection")
+	virtual FScriptTypedElementHandle GetSelectionElement(const FScriptTypedElementHandle& InElementHandle, const FManagedTypedElementListProxy InCurrentSelection, const ETypedElementSelectionMethod InSelectionMethod);
+
 protected:
 	/**
 	 * Create a transacted element instance that can be used to save/load elements of the implementation type for undo/redo.
