Fixed an issue with my last changelist.
Renamed a class that used the wrong terminology in it's name. ( FManagedTypedElementListProxy -> FScriptTypedElementListProxy)

diff --git a/Engine/Source/Runtime/Engine/Private/Elements/Framework/TypedElementCommonActions.cpp b/Engine/Source/Runtime/Engine/Private/Elements/Framework/TypedElementCommonActions.cpp
--- a/Engine/Source/Runtime/Engine/Private/Elements/Framework/TypedElementCommonActions.cpp
+++ b/Engine/Source/Runtime/Engine/Private/Elements/Framework/TypedElementCommonActions.cpp
@@ -87,7 +87,7 @@ TArray<FTypedElementHandle> UTypedElementCommonActions::DuplicateNormalizedEleme
 	return NewElements;
 }
 
-bool UTypedElementCommonActions::DeleteNormalizedElements(const FManagedTypedElementListProxy ElementList, UWorld* World, UTypedElementSelectionSet* InSelectionSet, const FTypedElementDeletionOptions& DeletionOptions)
+bool UTypedElementCommonActions::DeleteNormalizedElements(const FScriptTypedElementListProxy ElementList, UWorld* World, UTypedElementSelectionSet* InSelectionSet, const FTypedElementDeletionOptions& DeletionOptions)
 {
 	FTypedElementListPtr NativeList = UE::TypedElementFramework::ConvertToNativeTypedElementList(ElementList.GetElementList());
 	if (!NativeList)
@@ -109,7 +109,7 @@ TArray<FScriptTypedElementHandle> UTypedElementCommonActions::K2_DuplicateSelect
 	return TypedElementUtil::ConvertToScriptElementArray(DuplicateSelectedElements(SelectionSet, World, LocationOffset), SelectionSet->GetElementList()->GetRegistry());
 }
 
-TArray<FScriptTypedElementHandle> UTypedElementCommonActions::DuplicateNormalizedElements(const FManagedTypedElementListProxy ElementList, UWorld* World, const FVector& LocationOffset)
+TArray<FScriptTypedElementHandle> UTypedElementCommonActions::DuplicateNormalizedElements(const FScriptTypedElementListProxy ElementList, UWorld* World, const FVector& LocationOffset)
 {
 	FTypedElementListPtr NativeList = UE::TypedElementFramework::ConvertToNativeTypedElementList(ElementList.GetElementList());
 	if (!NativeList)
diff --git a/Engine/Source/Runtime/Engine/Public/Elements/Framework/TypedElementCommonActions.h b/Engine/Source/Runtime/Engine/Public/Elements/Framework/TypedElementCommonActions.h
--- a/Engine/Source/Runtime/Engine/Public/Elements/Framework/TypedElementCommonActions.h
+++ b/Engine/Source/Runtime/Engine/Public/Elements/Framework/TypedElementCommonActions.h
@@ -105,7 +105,7 @@ public:
 	 * @note This list should have been pre-normalized via UTypedElementSelectionSet::GetNormalizedSelection or UTypedElementSelectionSet::GetNormalizedElementList.
 	 */
 	UFUNCTION(BlueprintCallable, Category="TypedElementFramework|Common")
-	bool DeleteNormalizedElements(const FManagedTypedElementListProxy ElementList, UWorld* World, UTypedElementSelectionSet* InSelectionSet, const FTypedElementDeletionOptions& DeletionOptions);
+	bool DeleteNormalizedElements(const FScriptTypedElementListProxy ElementList, UWorld* World, UTypedElementSelectionSet* InSelectionSet, const FTypedElementDeletionOptions& DeletionOptions);
 
 	/**
 	 * Duplicate any elements from the given selection set that can be duplicated.
@@ -119,7 +119,7 @@ public:
 	 * @note This list should have been pre-normalized via UTypedElementSelectionSet::GetNormalizedSelection or UTypedElementSelectionSet::GetNormalizedElementList.
 	 */
 	UFUNCTION(BlueprintCallable, Category = "TypedElementFramework|Common")
-	TArray<FScriptTypedElementHandle> DuplicateNormalizedElements(const FManagedTypedElementListProxy ElementList, UWorld* World, const FVector& LocationOffset);
+	TArray<FScriptTypedElementHandle> DuplicateNormalizedElements(const FScriptTypedElementListProxy ElementList, UWorld* World, const FVector& LocationOffset);
 
 private:
 	/**
diff --git a/Engine/Source/Runtime/TypedElementFramework/Private/Elements/Framework/TypedElementListLibrary.h b/Engine/Source/Runtime/TypedElementFramework/Private/Elements/Framework/TypedElementListLibrary.h
--- a/Engine/Source/Runtime/TypedElementFramework/Private/Elements/Framework/TypedElementListLibrary.h
+++ b/Engine/Source/Runtime/TypedElementFramework/Private/Elements/Framework/TypedElementListLibrary.h
@@ -18,7 +18,7 @@ public:
 	 * Create an empty list of elements associated with the given registry.
 	 */
 	UFUNCTION(BlueprintCallable, DisplayName="Create Element List", Category = "TypedElementFramework|Registry", meta=(ScriptMethod, ScriptName="CreateElementList"))
-	static FManagedTypedElementListProxy CreateScriptElementList(UTypedElementRegistry* Registry)
+	static FScriptTypedElementListProxy CreateScriptElementList(UTypedElementRegistry* Registry)
 	{
 		return Registry->CreateScriptElementList();
 	}
@@ -28,10 +28,10 @@ public:
 	 * @note Only copies elements; does not copy any bindings!
 	 */
 	UFUNCTION(BlueprintCallable, Category="TypedElementFramework|List", meta=(ScriptMethod))
-	static FManagedTypedElementListProxy Clone(const FManagedTypedElementListProxy ElementList)
+	static FScriptTypedElementListProxy Clone(const FScriptTypedElementListProxy ElementList)
 	{
 		FScriptTypedElementListConstPtr ElementListPtr = ElementList.GetElementList();
-		return ElementListPtr ? ElementListPtr->Clone() : FManagedTypedElementListProxy();
+		return ElementListPtr ? ElementListPtr->Clone() : FScriptTypedElementListProxy();
 	}
 
 	/**
@@ -39,7 +39,7 @@ public:
 	 * @note Use IsValidIndex to test for validity.
 	 */
 	UFUNCTION(BlueprintPure, Category="TypedElementFramework|List", meta=(ScriptMethod))
-	static FScriptTypedElementHandle GetElementHandleAt(const FManagedTypedElementListProxy ElementList, const int32 Index)
+	static FScriptTypedElementHandle GetElementHandleAt(const FScriptTypedElementListProxy ElementList, const int32 Index)
 	{
 		FScriptTypedElementListConstPtr ElementListPtr = ElementList.GetElementList();
 		return ElementListPtr ? ElementListPtr->GetElementHandleAt(Index) : FScriptTypedElementHandle();
@@ -49,7 +49,7 @@ public:
 	 * Get the element interface from the given handle.
 	 */
 	UFUNCTION(BlueprintPure, Category="TypedElementFramework|List", meta=(ScriptMethod))
-	static UObject* GetElementInterface(const FManagedTypedElementListProxy ElementList, const FScriptTypedElementHandle& ElementHandle, const TSubclassOf<UInterface> BaseInterfaceType)
+	static UObject* GetElementInterface(const FScriptTypedElementListProxy ElementList, const FScriptTypedElementHandle& ElementHandle, const TSubclassOf<UInterface> BaseInterfaceType)
 	{
 		FScriptTypedElementListConstPtr ElementListPtr = ElementList.GetElementList();
 		return ElementListPtr ? ElementListPtr->GetElementInterface(ElementHandle, BaseInterfaceType) : nullptr;
@@ -59,7 +59,7 @@ public:
 	 * Test whether there are elements in this list, optionally filtering to elements that implement the given interface.
 	 */
 	UFUNCTION(BlueprintCallable, Category="TypedElementFramework|List", meta=(ScriptMethod))
-	static bool HasElements(const FManagedTypedElementListProxy ElementList, const TSubclassOf<UInterface> BaseInterfaceType = nullptr)
+	static bool HasElements(const FScriptTypedElementListProxy ElementList, const TSubclassOf<UInterface> BaseInterfaceType = nullptr)
 	{
 		FScriptTypedElementListConstPtr ElementListPtr = ElementList.GetElementList();
 		return ElementListPtr && ElementListPtr->HasElements(BaseInterfaceType);
@@ -69,7 +69,7 @@ public:
 	 * Count the number of elements in this list, optionally filtering to elements that implement the given interface.
 	 */
 	UFUNCTION(BlueprintCallable, Category="TypedElementFramework|List", meta=(ScriptMethod))
-	static int32 CountElements(const FManagedTypedElementListProxy ElementList, const TSubclassOf<UInterface> BaseInterfaceType = nullptr)
+	static int32 CountElements(const FScriptTypedElementListProxy ElementList, const TSubclassOf<UInterface> BaseInterfaceType = nullptr)
 	{
 		FScriptTypedElementListConstPtr ElementListPtr = ElementList.GetElementList();
 		return ElementListPtr ? ElementListPtr->CountElements(BaseInterfaceType) : 0;
@@ -79,7 +79,7 @@ public:
 	 * Test whether there are elements in this list of the given type.
 	 */
 	UFUNCTION(BlueprintCallable, Category="TypedElementFramework|List", meta=(ScriptMethod))
-	static bool HasElementsOfType(const FManagedTypedElementListProxy ElementList, const FName ElementTypeName)
+	static bool HasElementsOfType(const FScriptTypedElementListProxy ElementList, const FName ElementTypeName)
 	{
 		FScriptTypedElementListConstPtr ElementListPtr = ElementList.GetElementList();
 		return ElementListPtr && ElementListPtr->HasElementsOfType(ElementTypeName);
@@ -89,7 +89,7 @@ public:
 	 * Count the number of elements in this list of the given type.
 	 */
 	UFUNCTION(BlueprintCallable, Category="TypedElementFramework|List", meta=(ScriptMethod))
-	static int32 CountElementsOfType(const FManagedTypedElementListProxy ElementList, const FName ElementTypeName)
+	static int32 CountElementsOfType(const FScriptTypedElementListProxy ElementList, const FName ElementTypeName)
 	{
 		FScriptTypedElementListConstPtr ElementListPtr = ElementList.GetElementList();
 		return ElementListPtr ? ElementListPtr->CountElementsOfType(ElementTypeName) : 0;
@@ -99,7 +99,7 @@ public:
 	 * Get the handle of every element in this list, optionally filtering to elements that implement the given interface.
 	 */
 	UFUNCTION(BlueprintCallable, Category="TypedElementFramework|List", meta=(ScriptMethod))
-	static TArray<FScriptTypedElementHandle> GetElementHandles(const FManagedTypedElementListProxy ElementList, const TSubclassOf<UInterface> BaseInterfaceType)
+	static TArray<FScriptTypedElementHandle> GetElementHandles(const FScriptTypedElementListProxy ElementList, const TSubclassOf<UInterface> BaseInterfaceType)
 	{
 		FScriptTypedElementListConstPtr ElementListPtr = ElementList.GetElementList();
 		return ElementListPtr ? ElementListPtr->GetElementHandles(BaseInterfaceType) : TArray<FScriptTypedElementHandle>();
@@ -109,7 +109,7 @@ public:
 	 * Is the given index a valid entry within this element list?
 	 */
 	UFUNCTION(BlueprintPure, Category="TypedElementFramework|List", meta=(ScriptMethod))
-	static bool IsValidIndex(const FManagedTypedElementListProxy ElementList, const int32 Index)
+	static bool IsValidIndex(const FScriptTypedElementListProxy ElementList, const int32 Index)
 	{
 		FScriptTypedElementListConstPtr ElementListPtr = ElementList.GetElementList();
 		return ElementListPtr && ElementListPtr->IsValidIndex(Index);
@@ -119,7 +119,7 @@ public:
 	 * Get the number of entries within this element list.
 	 */
 	UFUNCTION(BlueprintPure, Category="TypedElementFramework|List", meta=(ScriptMethod))
-	static int32 Num(const FManagedTypedElementListProxy ElementList)
+	static int32 Num(const FScriptTypedElementListProxy ElementList)
 	{
 		FScriptTypedElementListConstPtr ElementListPtr = ElementList.GetElementList();
 		return ElementListPtr ? ElementListPtr->Num() : 0;
@@ -129,7 +129,7 @@ public:
 	 * Shrink this element list storage to avoid slack.
 	 */
 	UFUNCTION(BlueprintCallable, Category="TypedElementFramework|List", meta=(ScriptMethod))
-	static void Shrink(FManagedTypedElementListProxy ElementList)
+	static void Shrink(FScriptTypedElementListProxy ElementList)
 	{
 		if (FScriptTypedElementListPtr ElementListPtr = ElementList.GetElementList())
 		{
@@ -141,7 +141,7 @@ public:
 	 * Pre-allocate enough memory in this element list to store the given number of entries.
 	 */
 	UFUNCTION(BlueprintCallable, Category="TypedElementFramework|List", meta=(ScriptMethod))
-	static void Reserve(FManagedTypedElementListProxy ElementList, const int32 Size)
+	static void Reserve(FScriptTypedElementListProxy ElementList, const int32 Size)
 	{
 		if (FScriptTypedElementListPtr ElementListPtr = ElementList.GetElementList())
 		{
@@ -153,7 +153,7 @@ public:
 	 * Remove all entries from this element list, potentially leaving space allocated for the given number of entries.
 	 */
 	UFUNCTION(BlueprintCallable, Category="TypedElementFramework|List", meta=(ScriptMethod))
-	static void Empty(FManagedTypedElementListProxy ElementList, const int32 Slack = 0)
+	static void Empty(FScriptTypedElementListProxy ElementList, const int32 Slack = 0)
 	{
 		if (FScriptTypedElementListPtr ElementListPtr = ElementList.GetElementList())
 		{
@@ -165,7 +165,7 @@ public:
 	 * Remove all entries from this element list, preserving existing allocations.
 	 */
 	UFUNCTION(BlueprintCallable, Category="TypedElementFramework|List", meta=(ScriptMethod))
-	static void Reset(FManagedTypedElementListProxy ElementList)
+	static void Reset(FScriptTypedElementListProxy ElementList)
 	{
 		if (FScriptTypedElementListPtr ElementListPtr = ElementList.GetElementList())
 		{
@@ -177,7 +177,7 @@ public:
 	 * Does this element list contain an entry for the given element handle?
 	 */
 	UFUNCTION(BlueprintPure, Category="TypedElementFramework|List", meta=(ScriptMethod))
-	static bool Contains(const FManagedTypedElementListProxy ElementList, const FScriptTypedElementHandle& ElementHandle)
+	static bool Contains(const FScriptTypedElementListProxy ElementList, const FScriptTypedElementHandle& ElementHandle)
 	{
 		FScriptTypedElementListConstPtr ElementListPtr = ElementList.GetElementList();
 		return ElementListPtr && ElementListPtr->Contains(ElementHandle);
@@ -188,7 +188,7 @@ public:
 	 * @return True if the element handle was added, false if it is already in the list.
 	 */
 	UFUNCTION(BlueprintCallable, Category="TypedElementFramework|List", meta=(ScriptMethod))
-	static bool Add(FManagedTypedElementListProxy ElementList, const FScriptTypedElementHandle& ElementHandle)
+	static bool Add(FScriptTypedElementListProxy ElementList, const FScriptTypedElementHandle& ElementHandle)
 	{
 		FScriptTypedElementListPtr ElementListPtr = ElementList.GetElementList();
 		return ElementListPtr && ElementListPtr->Add(ElementHandle);
@@ -198,7 +198,7 @@ public:
 	 * Append the given element handles to this element list.
 	 */
 	UFUNCTION(BlueprintCallable, Category="TypedElementFramework|List", meta=(ScriptMethod))
-	static void Append(FManagedTypedElementListProxy ElementList, const TArray<FScriptTypedElementHandle>& ElementHandles)
+	static void Append(FScriptTypedElementListProxy ElementList, const TArray<FScriptTypedElementHandle>& ElementHandles)
 	{
 		if (FScriptTypedElementListPtr ElementListPtr = ElementList.GetElementList())
 		{
@@ -210,7 +210,7 @@ public:
 	 * Append the another element list to this element list.
 	 */
 	UFUNCTION(BlueprintCallable, Category="TypedElementFramework|List", meta=(ScriptMethod))
-	static void AppendList(FManagedTypedElementListProxy ElementList, const FManagedTypedElementListProxy OtherElementList)
+	static void AppendList(FScriptTypedElementListProxy ElementList, const FScriptTypedElementListProxy OtherElementList)
 	{
 		FScriptTypedElementListPtr ElementListPtr = ElementList.GetElementList();
 		FScriptTypedElementListConstPtr OtherElementListPtr = OtherElementList.GetElementList();
@@ -225,7 +225,7 @@ public:
 	 * @return True if the element handle was removed, false if it isn't in the list.
 	 */
 	UFUNCTION(BlueprintCallable, Category="TypedElementFramework|List", meta=(ScriptMethod))
-	static bool Remove(FManagedTypedElementListProxy ElementList, const FScriptTypedElementHandle& ElementHandle)
+	static bool Remove(FScriptTypedElementListProxy ElementList, const FScriptTypedElementHandle& ElementHandle)
 	{
 		FScriptTypedElementListPtr ElementListPtr = ElementList.GetElementList();
 		return ElementListPtr && ElementListPtr->Remove(ElementHandle);
diff --git a/Engine/Source/Runtime/TypedElementFramework/Public/Elements/Framework/TypedElementListProxy.h b/Engine/Source/Runtime/TypedElementFramework/Public/Elements/Framework/TypedElementListProxy.h
--- a/Engine/Source/Runtime/TypedElementFramework/Public/Elements/Framework/TypedElementListProxy.h
+++ b/Engine/Source/Runtime/TypedElementFramework/Public/Elements/Framework/TypedElementListProxy.h
@@ -7,26 +7,26 @@
 #include "TypedElementListProxy.generated.h"
 
 /**
- * A list of managed element handles (proxy to a FManagedTypedElementList instance).
+ * A list of script element handles (proxy to a FScriptTypedElementList instance).
  * Provides high-level access to groups of elements, including accessing elements that implement specific interfaces.
  * 
- * Note: the managed list proxy use should be avoided when not using it for the script exposure apis.
+ * Note: the script list proxy use should be avoided when not using it for the script exposure apis.
  * The weak model for the handles come with an additional cost to the runtime performance and the memory usage.
  */
 USTRUCT(BlueprintType, DisplayName="Typed Element List", meta=(ScriptName="TypedElementList"))
-struct TYPEDELEMENTFRAMEWORK_API FManagedTypedElementListProxy
+struct TYPEDELEMENTFRAMEWORK_API FScriptTypedElementListProxy
 {
 	GENERATED_BODY()
 
 public:
-	FManagedTypedElementListProxy() = default;
+	FScriptTypedElementListProxy() = default;
 
-	FManagedTypedElementListProxy(FScriptTypedElementListRef InElementList)
+	FScriptTypedElementListProxy(FScriptTypedElementListRef InElementList)
 		: ElementList(MoveTemp(InElementList))
 	{
 	}
 
-	FManagedTypedElementListProxy(FScriptTypedElementListConstRef InElementList)
+	FScriptTypedElementListProxy(FScriptTypedElementListConstRef InElementList)
 		: ElementList(ConstCastSharedRef<FScriptTypedElementList>(InElementList))
 	{
 	}
diff --git a/Engine/Source/Runtime/TypedElementRuntime/Private/Elements/Framework/TypedElementSelectionSetLibrary.h b/Engine/Source/Runtime/TypedElementRuntime/Private/Elements/Framework/TypedElementSelectionSetLibrary.h
--- a/Engine/Source/Runtime/TypedElementRuntime/Private/Elements/Framework/TypedElementSelectionSetLibrary.h
+++ b/Engine/Source/Runtime/TypedElementRuntime/Private/Elements/Framework/TypedElementSelectionSetLibrary.h
@@ -40,7 +40,7 @@ public:
 	 * @return True if the selection was changed, false otherwise.
 	 */
 	UFUNCTION(BlueprintCallable, Category="TypedElementFramework|Selection", meta=(ScriptMethod))
-	static bool SelectElementsFromList(UTypedElementSelectionSet* SelectionSet, const FManagedTypedElementListProxy ElementList, const FTypedElementSelectionOptions SelectionOptions)
+	static bool SelectElementsFromList(UTypedElementSelectionSet* SelectionSet, const FScriptTypedElementListProxy ElementList, const FTypedElementSelectionOptions SelectionOptions)
 	{
 		if (!UE::TypedElementFramework::Private::CheckForInvalidSelectionSet(SelectionSet))
 		{
@@ -62,7 +62,7 @@ public:
 	 * @return True if the selection was changed, false otherwise.
 	 */
 	UFUNCTION(BlueprintCallable, Category="TypedElementFramework|Selection", meta=(ScriptMethod))
-	static bool DeselectElementsFromList(UTypedElementSelectionSet* SelectionSet, const FManagedTypedElementListProxy ElementList, const FTypedElementSelectionOptions SelectionOptions)
+	static bool DeselectElementsFromList(UTypedElementSelectionSet* SelectionSet, const FScriptTypedElementListProxy ElementList, const FTypedElementSelectionOptions SelectionOptions)
 	{
 		if (!UE::TypedElementFramework::Private::CheckForInvalidSelectionSet(SelectionSet))
 		{
@@ -85,7 +85,7 @@ public:
 	 * @return True if the selection was changed, false otherwise.
 	 */
 	UFUNCTION(BlueprintCallable, Category="TypedElementFramework|Selection", meta=(ScriptMethod))
-	static bool SetSelectionFromList(UTypedElementSelectionSet* SelectionSet, const FManagedTypedElementListProxy ElementList, const FTypedElementSelectionOptions SelectionOptions)
+	static bool SetSelectionFromList(UTypedElementSelectionSet* SelectionSet, const FScriptTypedElementListProxy ElementList, const FTypedElementSelectionOptions SelectionOptions)
 	{
 		if (!UE::TypedElementFramework::Private::CheckForInvalidSelectionSet(SelectionSet))
 		{
@@ -103,15 +103,15 @@ public:
 	 * This will do things like expand out groups, and resolve any parent<->child elements so that duplication operations aren't performed on both the parent and the child.
 	 */
 	UFUNCTION(BlueprintPure, Category="TypedElementFramework|Selection", meta=(ScriptMethod))
-	static FManagedTypedElementListProxy GetNormalizedSelection(UTypedElementSelectionSet* SelectionSet, const FTypedElementSelectionNormalizationOptions NormalizationOptions)
+	static FScriptTypedElementListProxy GetNormalizedSelection(UTypedElementSelectionSet* SelectionSet, const FTypedElementSelectionNormalizationOptions NormalizationOptions)
 	{
 		if (!UE::TypedElementFramework::Private::CheckForInvalidSelectionSet(SelectionSet))
 		{
-			return FManagedTypedElementListProxy();
+			return FScriptTypedElementListProxy();
 		}
 
 		FScriptTypedElementListPtr NormalizedSelection = UE::TypedElementFramework::ConvertToScriptTypedElementList(SelectionSet->GetNormalizedSelection(NormalizationOptions));
-		return NormalizedSelection ? NormalizedSelection.ToSharedRef() : FManagedTypedElementListProxy();
+		return NormalizedSelection ? NormalizedSelection.ToSharedRef() : FScriptTypedElementListProxy();
 	}
 
 	/**
@@ -119,14 +119,14 @@ public:
 	 * This will do things like expand out groups, and resolve any parent<->child elements so that duplication operations aren't performed on both the parent and the child.
 	 */
 	UFUNCTION(BlueprintPure, Category="TypedElementFramework|Selection", meta=(ScriptMethod))
-	static FManagedTypedElementListProxy GetNormalizedElementList(UTypedElementSelectionSet* SelectionSet, const FManagedTypedElementListProxy ElementList, const FTypedElementSelectionNormalizationOptions NormalizationOptions)
+	static FScriptTypedElementListProxy GetNormalizedElementList(UTypedElementSelectionSet* SelectionSet, const FScriptTypedElementListProxy ElementList, const FTypedElementSelectionNormalizationOptions NormalizationOptions)
 	{
 		if (!UE::TypedElementFramework::Private::CheckForInvalidSelectionSet(SelectionSet))
 		{
-			return FManagedTypedElementListProxy();
+			return FScriptTypedElementListProxy();
 		}
 
 		FTypedElementListPtr NativeList = UE::TypedElementFramework::ConvertToNativeTypedElementList(ElementList.GetElementList());
-		return NativeList ?  UE::TypedElementFramework::ConvertToScriptTypedElementList(SelectionSet->GetNormalizedElementList(NativeList.ToSharedRef(), NormalizationOptions)).ToSharedRef() : FManagedTypedElementListProxy();
+		return NativeList ?  UE::TypedElementFramework::ConvertToScriptTypedElementList(SelectionSet->GetNormalizedElementList(NativeList.ToSharedRef(), NormalizationOptions)).ToSharedRef() : FScriptTypedElementListProxy();
 	}
 };
diff --git a/Engine/Source/Runtime/TypedElementRuntime/Private/Elements/Interfaces/TypedElementSelectionInterface.cpp b/Engine/Source/Runtime/TypedElementRuntime/Private/Elements/Interfaces/TypedElementSelectionInterface.cpp
--- a/Engine/Source/Runtime/TypedElementRuntime/Private/Elements/Interfaces/TypedElementSelectionInterface.cpp
+++ b/Engine/Source/Runtime/TypedElementRuntime/Private/Elements/Interfaces/TypedElementSelectionInterface.cpp
@@ -23,7 +23,7 @@ bool ITypedElementSelectionInterface::DeselectElement(const FTypedElementHandle&
 	return InSelectionSet && InSelectionSet->Remove(InElementHandle);
 }
 
-bool ITypedElementSelectionInterface::IsElementSelected(const FScriptTypedElementHandle& InElementHandle, const FManagedTypedElementListProxy InSelectionSet, const FTypedElementIsSelectedOptions& InSelectionOptions)
+bool ITypedElementSelectionInterface::IsElementSelected(const FScriptTypedElementHandle& InElementHandle, const FScriptTypedElementListProxy InSelectionSet, const FTypedElementIsSelectedOptions& InSelectionOptions)
 {
 	FTypedElementHandle NativeHandle = InElementHandle.GetTypedElementHandle();
 	if (!NativeHandle)
@@ -66,7 +66,7 @@ bool ITypedElementSelectionInterface::CanDeselectElement(const FScriptTypedEleme
 	return CanDeselectElement(NativeHandle, InSelectionOptions);
 }
 
-bool ITypedElementSelectionInterface::SelectElement(const FScriptTypedElementHandle& InElementHandle, FManagedTypedElementListProxy InSelectionSet, const FTypedElementSelectionOptions& InSelectionOptions)
+bool ITypedElementSelectionInterface::SelectElement(const FScriptTypedElementHandle& InElementHandle, FScriptTypedElementListProxy InSelectionSet, const FTypedElementSelectionOptions& InSelectionOptions)
 {
 	FTypedElementHandle NativeHandle = InElementHandle.GetTypedElementHandle();
 	if (!NativeHandle)
@@ -85,7 +85,7 @@ bool ITypedElementSelectionInterface::SelectElement(const FScriptTypedElementHan
 	return SelectElement(NativeHandle, NativeList, InSelectionOptions);
 }
 
-bool ITypedElementSelectionInterface::DeselectElement(const FScriptTypedElementHandle& InElementHandle, FManagedTypedElementListProxy InSelectionSet, const FTypedElementSelectionOptions& InSelectionOptions)
+bool ITypedElementSelectionInterface::DeselectElement(const FScriptTypedElementHandle& InElementHandle, FScriptTypedElementListProxy InSelectionSet, const FTypedElementSelectionOptions& InSelectionOptions)
 {
 	FTypedElementHandle NativeHandle = InElementHandle.GetTypedElementHandle();
 	if (!NativeHandle)
@@ -104,7 +104,7 @@ bool ITypedElementSelectionInterface::DeselectElement(const FScriptTypedElementH
 	return DeselectElement(NativeHandle, NativeList, InSelectionOptions);
 }
 
-bool ITypedElementSelectionInterface::AllowSelectionModifiers(const FScriptTypedElementHandle& InElementHandle, const FManagedTypedElementListProxy InSelectionSet)
+bool ITypedElementSelectionInterface::AllowSelectionModifiers(const FScriptTypedElementHandle& InElementHandle, const FScriptTypedElementListProxy InSelectionSet)
 {
 	FTypedElementHandle NativeHandle = InElementHandle.GetTypedElementHandle();
 	if (!NativeHandle)
@@ -123,7 +123,7 @@ bool ITypedElementSelectionInterface::AllowSelectionModifiers(const FScriptTyped
 	return AllowSelectionModifiers(NativeHandle, NativeList);
 }
 
-FScriptTypedElementHandle ITypedElementSelectionInterface::GetSelectionElement(const FScriptTypedElementHandle& InElementHandle, const FManagedTypedElementListProxy InCurrentSelection, const ETypedElementSelectionMethod InSelectionMethod)
+FScriptTypedElementHandle ITypedElementSelectionInterface::GetSelectionElement(const FScriptTypedElementHandle& InElementHandle, const FScriptTypedElementListProxy InCurrentSelection, const ETypedElementSelectionMethod InSelectionMethod)
 {
 	FTypedElementHandle NativeHandle = InElementHandle.GetTypedElementHandle();
 	if (!NativeHandle)
diff --git a/Engine/Source/Runtime/TypedElementRuntime/Public/Elements/Interfaces/TypedElementSelectionInterface.h b/Engine/Source/Runtime/TypedElementRuntime/Public/Elements/Interfaces/TypedElementSelectionInterface.h
--- a/Engine/Source/Runtime/TypedElementRuntime/Public/Elements/Interfaces/TypedElementSelectionInterface.h
+++ b/Engine/Source/Runtime/TypedElementRuntime/Public/Elements/Interfaces/TypedElementSelectionInterface.h
@@ -217,7 +217,7 @@ public:
 	 * Test to see whether the given element is currently considered selected.
 	 */
 	UFUNCTION(BlueprintCallable, Category = "TypedElementInterfaces|Selection")
-	virtual bool IsElementSelected(const FScriptTypedElementHandle& InElementHandle, const FManagedTypedElementListProxy InSelectionSet, const FTypedElementIsSelectedOptions& InSelectionOptions);
+	virtual bool IsElementSelected(const FScriptTypedElementHandle& InElementHandle, const FScriptTypedElementListProxy InSelectionSet, const FTypedElementIsSelectedOptions& InSelectionOptions);
 
 	/**
 	 * Test to see whether the given element can be selected.
@@ -236,26 +236,26 @@ public:
 	 * @return True if the selection was changed, false otherwise.
 	 */
 	UFUNCTION(BlueprintCallable, Category="TypedElementInterfaces|Selection")
-	virtual bool SelectElement(const FScriptTypedElementHandle& InElementHandle, FManagedTypedElementListProxy InSelectionSet, const FTypedElementSelectionOptions& InSelectionOptions);
+	virtual bool SelectElement(const FScriptTypedElementHandle& InElementHandle, FScriptTypedElementListProxy InSelectionSet, const FTypedElementSelectionOptions& InSelectionOptions);
 
 	/**
 	 * Attempt to deselect the given element.
 	 * @return True if the selection was changed, false otherwise.
 	 */
 	UFUNCTION(BlueprintCallable, Category="TypedElementInterfaces|Selection")
-	virtual bool DeselectElement(const FScriptTypedElementHandle& InElementHandle, FManagedTypedElementListProxy InSelectionSet, const FTypedElementSelectionOptions& InSelectionOptions);
+	virtual bool DeselectElement(const FScriptTypedElementHandle& InElementHandle, FScriptTypedElementListProxy InSelectionSet, const FTypedElementSelectionOptions& InSelectionOptions);
 
 	/**
 	 * Test to see whether selection modifiers (Ctrl or Shift) are allowed while selecting this element.
 	 */
 	UFUNCTION(BlueprintCallable, Category="TypedElementInterfaces|Selection")
-	virtual bool AllowSelectionModifiers(const FScriptTypedElementHandle& InElementHandle, const FManagedTypedElementListProxy InSelectionSet);
+	virtual bool AllowSelectionModifiers(const FScriptTypedElementHandle& InElementHandle, const FScriptTypedElementListProxy InSelectionSet);
 
 	/**
 	 * Given an element, return the element that should actually perform a selection operation.
 	 */
 	UFUNCTION(BlueprintCallable, Category="TypedElementInterfaces|Selection")
-	virtual FScriptTypedElementHandle GetSelectionElement(const FScriptTypedElementHandle& InElementHandle, const FManagedTypedElementListProxy InCurrentSelection, const ETypedElementSelectionMethod InSelectionMethod);
+	virtual FScriptTypedElementHandle GetSelectionElement(const FScriptTypedElementHandle& InElementHandle, const FScriptTypedElementListProxy InCurrentSelection, const ETypedElementSelectionMethod InSelectionMethod);
 
 protected:
 	/**
