Make TypeElement ref tracking available in all editor build targets (even shipping)

diff --git a/Engine/Source/Runtime/TypedElementFramework/Public/Elements/Framework/TypedElementLimits.h b/Engine/Source/Runtime/TypedElementFramework/Public/Elements/Framework/TypedElementLimits.h
--- a/Engine/Source/Runtime/TypedElementFramework/Public/Elements/Framework/TypedElementLimits.h
+++ b/Engine/Source/Runtime/TypedElementFramework/Public/Elements/Framework/TypedElementLimits.h
@@ -6,7 +6,7 @@
 #include "Math/NumericLimits.h"
 
 #define UE_TYPED_ELEMENT_HAS_REFCOUNTING (1)
-#define UE_TYPED_ELEMENT_HAS_REFTRACKING (UE_BUILD_DEBUG || UE_BUILD_DEVELOPMENT)
+#define UE_TYPED_ELEMENT_HAS_REFTRACKING (UE_BUILD_DEBUG || UE_BUILD_DEVELOPMENT || WITH_EDITOR)
 #define UE_TYPED_ELEMENT_HAS_REFERENCING (UE_TYPED_ELEMENT_HAS_REFCOUNTING || UE_TYPED_ELEMENT_HAS_REFTRACKING)
 
 /**
