Modular builds: TTypedElementList: Added a missing TYPEDELEMENTFRAMEWORK_API to resolve missing symbols during linking on some platforms.

#rnx

diff --git a/Engine/Source/Runtime/TypedElementFramework/Public/Elements/Framework/TypedElementList.h b/Engine/Source/Runtime/TypedElementFramework/Public/Elements/Framework/TypedElementList.h
--- a/Engine/Source/Runtime/TypedElementFramework/Public/Elements/Framework/TypedElementList.h
+++ b/Engine/Source/Runtime/TypedElementFramework/Public/Elements/Framework/TypedElementList.h
@@ -44,7 +44,7 @@ FORCEINLINE TTypedElement<BaseInterfaceType> GetElement(const UTypedElementRegis
  * Provides high-level access to groups of elements, including accessing elements that implement specific interfaces.
  */
 template<class HandleType>
-class TTypedElementList final : public TSharedFromThis<TTypedElementList<HandleType>>
+class TYPEDELEMENTFRAMEWORK_API TTypedElementList final : public TSharedFromThis<TTypedElementList<HandleType>>
 {
 public:
 
