Initial support for adding tag and data to the Typed Elements Data Storage and creating queries.

This is very early work-in-progress. Everythng in this change list is very likely to be updated and improved in the (near) future. This was CL was created now to avoid lengthy reviews later on and to get early feedback, in particular about coding standards and UE specific patterns.

This change list adds the following:
- The ability to add tags by type desscription or by name.
- The ability to add and retrieve column data by type description or by name.
- The abiilty to create and add queries. This CL only covers simple queries (those that can go straight into MASS), but some code has also be added for future development. Note that the provided query system doesn't support all the features MASS' queries provide. Most of these missing features will be added over time in a variety of ways. While MASS' queries put an emphasis on performance, the goal for the Typed Elements Data Storage is to create very accessible and familiar looking query statements while minimizing performance cost as a secondary challenge.
- Reorganized the Data Storage Interface header to be easier to read.
- Added a handle store to manage queries. This container similarities with SparseArray, but doesn't allow compacting in favor of being able to keep track if the generation of an entry to quickly check if it's alive or not. The current implementation is a first step and updates are planned.
- Utility functions to help set column data through strings. The driving reason for this is to provide other widgets to set a batch set of values without needing to know the particulars of each child widgets. Child widgets can in the future be generated through a centralized registration so the parent widget may at that point not  know what a widget type is, but might still want to set generic values like background color, font size, etc.
- A query builder to make creating queries much easier than using the interface directly. Please see the TypedElementQueryBuilder header for an example of a query.

diff --git a/Engine/Plugins/TypedElementsDataStorage/Source/TypedElementsDataStorage/Private/TypedElementDatabase.cpp b/Engine/Plugins/TypedElementsDataStorage/Source/TypedElementsDataStorage/Private/TypedElementDatabase.cpp
--- a/Engine/Plugins/TypedElementsDataStorage/Source/TypedElementsDataStorage/Private/TypedElementDatabase.cpp
+++ b/Engine/Plugins/TypedElementsDataStorage/Source/TypedElementsDataStorage/Private/TypedElementDatabase.cpp
@@ -60,26 +60,6 @@ TSharedPtr<const FMassEntityManager> UTypedElementDatabase::GetActiveEditorEntit
 	return ActiveEditorEntityManager;
 }
 
-FTypedElementOnDataStorageCreation& UTypedElementDatabase::OnCreation()
-{
-	return OnCreationDelegate;
-}
-
-FTypedElementOnDataStorageDestruction& UTypedElementDatabase::OnDestruction()
-{
-	return OnDestructionDelegate;
-}
-
-FTypedElementOnDataStorageUpdate& UTypedElementDatabase::OnUpdate()
-{
-	return OnUpdateDelegate;
-}
-
-bool UTypedElementDatabase::IsAvailable() const
-{
-	return ActiveEditorWorld && ActiveEditorEntityManager;
-}
-
 TypedElementTableHandle UTypedElementDatabase::RegisterTable(TConstArrayView<const UScriptStruct*> ColumnList)
 {
 	return RegisterTable(ColumnList, {});
@@ -158,14 +138,226 @@ void UTypedElementDatabase::RemoveRow(TypedElementRowHandle Row)
 	}
 }
 
+void UTypedElementDatabase::AddTag(TypedElementRowHandle Row, const UScriptStruct* TagType)
+{
+	checkf(TagType && TagType->IsChildOf(FMassTag::StaticStruct()),
+		TEXT("Tag type '%s' is invalid as it needs to be set or derived from FMassTag."), *GetPathNameSafe(TagType));
+
+	FMassEntityHandle Entity = FMassEntityHandle::FromNumber(Row);
+	if (ActiveEditorEntityManager && ActiveEditorEntityManager->IsEntityValid(Entity))
+	{
+		ActiveEditorEntityManager->AddTagToEntity(Entity, TagType);
+	}
+}
+
+void UTypedElementDatabase::AddTag(TypedElementRowHandle Row, FTopLevelAssetPath TagName)
+{
+	bool bExactMatch = true;
+	UScriptStruct* TagStructInfo = Cast<UScriptStruct>(StaticFindObject(UScriptStruct::StaticClass(), TagName, bExactMatch));
+	if (TagStructInfo)
+	{
+		AddTag(Row, TagStructInfo);
+	}
+}
+
+void* UTypedElementDatabase::AddOrGetColumnData(TypedElementRowHandle Row, const UScriptStruct* ColumnType)
+{
+	checkf(ColumnType && ColumnType->IsChildOf(FMassFragment::StaticStruct()), 
+		TEXT("Colum type '%s' is invalid as it needs to be set or derived from FMassFragment."), *GetPathNameSafe(ColumnType));
+
+	FMassEntityHandle Entity = FMassEntityHandle::FromNumber(Row);
+	if (ActiveEditorEntityManager && ActiveEditorEntityManager->IsEntityValid(Entity))
+	{
+		FStructView Column = ActiveEditorEntityManager->GetFragmentDataStruct(Entity, ColumnType);
+		if (!Column.IsValid())
+		{
+			ActiveEditorEntityManager->AddFragmentToEntity(Entity, ColumnType);
+			Column = ActiveEditorEntityManager->GetFragmentDataStruct(Entity, ColumnType);
+			checkf(Column.IsValid(), TEXT("Added a new column to the Typed Element's data storae, but it couldn't be retrieved."));
+
+		}
+		return Column.GetMutableMemory();
+	}
+	return nullptr;
+}
+
+ColumnDataResult UTypedElementDatabase::AddOrGetColumnData(TypedElementRowHandle Row, FTopLevelAssetPath ColumnName)
+{
+	constexpr bool bExactMatch = true;
+	UScriptStruct* FragmentStructInfo = Cast<UScriptStruct>(StaticFindObject(UScriptStruct::StaticClass(), ColumnName, bExactMatch));
+	return FragmentStructInfo ?
+		ColumnDataResult{ FragmentStructInfo, AddOrGetColumnData(Row, FragmentStructInfo) }:
+		ColumnDataResult{ nullptr, nullptr };
+}
+
+ColumnDataResult UTypedElementDatabase::AddOrGetColumnData(TypedElementRowHandle Row, FTopLevelAssetPath ColumnName,
+	TConstArrayView<TypedElement::ColumnUtils::Argument> Arguments)
+{
+	ColumnDataResult Result = AddOrGetColumnData(Row, ColumnName);
+	if (Result.Description && Result.Data)
+	{
+		TypedElement::ColumnUtils::SetColumnValues(Result.Data, Result.Description, Arguments);
+		return Result;
+	}
+	else
+	{
+		return ColumnDataResult{ nullptr, nullptr };
+	}
+}
+
 void* UTypedElementDatabase::GetColumnData(TypedElementRowHandle Row, const UScriptStruct* ColumnType)
 {
-	check(ColumnType);
+	checkf(ColumnType && ColumnType->IsChildOf(FMassFragment::StaticStruct()),
+		TEXT("Colum type '%s' is invalid as it needs to be set or derived from FMassFragment."), *GetPathNameSafe(ColumnType));
 
 	FMassEntityHandle Entity = FMassEntityHandle::FromNumber(Row);
-	return ActiveEditorEntityManager && ActiveEditorEntityManager->IsEntityValid(Entity) ?
-		ActiveEditorEntityManager->GetFragmentDataStruct(FMassEntityHandle::FromNumber(Row), ColumnType).GetMutableMemory() :
-		nullptr;
+	if (ActiveEditorEntityManager && ActiveEditorEntityManager->IsEntityValid(Entity))
+	{
+		FStructView Column = ActiveEditorEntityManager->GetFragmentDataStruct(Entity, ColumnType);
+		if (Column.IsValid())
+		{
+			return Column.GetMutableMemory();
+		}
+	}
+	return nullptr;
+}
+
+ColumnDataResult UTypedElementDatabase::GetColumnData(TypedElementRowHandle Row, FTopLevelAssetPath ColumnName)
+{
+	FMassEntityHandle Entity = FMassEntityHandle::FromNumber(Row);
+	if (ActiveEditorEntityManager && ActiveEditorEntityManager->IsEntityValid(Entity))
+	{
+		const UScriptStruct* FragmentType = nullptr;
+		FMassArchetypeHandle Archetype = ActiveEditorEntityManager->GetArchetypeForEntityUnsafe(Entity);
+		ActiveEditorEntityManager->ForEachArchetypeFragmentType(Archetype, 
+			[ColumnName, &FragmentType](const UScriptStruct* Fragment)
+			{
+				if (Fragment->GetStructPathName() == ColumnName)
+				{
+					FragmentType = Fragment;
+				}
+			});
+
+		if (FragmentType != nullptr)
+		{
+			FStructView Column = ActiveEditorEntityManager->GetFragmentDataStruct(Entity, FragmentType);
+			if (Column.IsValid())
+			{
+				return ColumnDataResult{ FragmentType, Column.GetMutableMemory() };
+			}
+		}
+	}
+	return ColumnDataResult{ nullptr, nullptr };
+}
+
+TypedElementQueryHandle UTypedElementDatabase::RegisterQuery(const QueryDescription& Query)
+{
+	auto LocalToNativeAccess = [](QueryDescription::EAccessType Access) -> EMassFragmentAccess
+	{
+		switch (Access)
+		{
+		case QueryDescription::EAccessType::ReadOnly:
+			return EMassFragmentAccess::ReadOnly;
+		case QueryDescription::EAccessType::ReadWrite:
+			return EMassFragmentAccess::ReadWrite;
+		default:
+			checkf(false, TEXT("Invalid query access type: %i."), static_cast<uint32>(Access));
+			return EMassFragmentAccess::MAX;
+		}
+	};
+
+	QueryStore::Handle Result = Queries.Emplace();
+	FMassEntityQuery& NativeQuery = Queries.GetMutable(Result).NativeQuery;
+
+	for (const QueryDescription::AccessControlledStruct& SelectEntry : Query.Selection)
+	{
+		checkf(SelectEntry.Type, TEXT("Provided query selection type can not be null."));
+		checkf(SelectEntry.Type->IsChildOf(FMassFragment::StaticStruct()), 
+			TEXT("Provided query selection type '%s' is not based on FMassFragment."), *SelectEntry.Type->GetStructPathName().ToString());
+		NativeQuery.AddRequirement(SelectEntry.Type, LocalToNativeAccess(SelectEntry.Access));
+	}
+
+	if (Query.bSimpleQuery) // This backend currently only supports simple queries.
+	{
+		checkf(Query.ConditionTypes.Num() == Query.ConditionOperators.Num(),
+			TEXT("The types and operators for a typed element query have gone out of sync."));
+		
+		const QueryDescription::Operator* Operand = Query.ConditionOperators.GetData();
+		for (QueryDescription::OperatorType Type : Query.ConditionTypes)
+		{
+			EMassFragmentPresence Presence;
+			switch (Type)
+			{
+			case QueryDescription::OperatorType::SimpleAll:
+				Presence = EMassFragmentPresence::All;
+				break;
+			case QueryDescription::OperatorType::SimpleAny:
+				Presence = EMassFragmentPresence::Any;
+				break;
+			case QueryDescription::OperatorType::SimpleNone:
+				Presence = EMassFragmentPresence::None;
+				break;
+			default:
+				continue;
+			}
+
+			if (Operand->Type->IsChildOf(FMassTag::StaticStruct()))
+			{
+				NativeQuery.AddTagRequirement(*(Operand->Type), Presence);
+			}
+			else if (Operand->Type->IsChildOf(FMassFragment::StaticStruct()))
+			{
+				NativeQuery.AddRequirement(Operand->Type, EMassFragmentAccess::None, Presence);
+			}
+
+			++Operand;
+		}
+	}
+
+	for (const QueryDescription::AccessControlledClass& DependencyEntry : Query.Dependencies)
+	{
+		checkf(DependencyEntry.Type, TEXT("Provided query dependcy type can not be null."));
+		checkf(DependencyEntry.Type->IsChildOf<UWorldSubsystem>(),
+			TEXT("Provided query dependency type '%s' is not based on UWorldSubSystem."), *DependencyEntry.Type->GetStructPathName().ToString());
+		
+		constexpr bool bGameThreadOnly = true;
+		NativeQuery.AddSubsystemRequirement(const_cast<UClass*>(DependencyEntry.Type), LocalToNativeAccess(DependencyEntry.Access), bGameThreadOnly);
+	}
+
+	return Result.Handle;
+}
+
+void UTypedElementDatabase::UnregisterQuery(TypedElementQueryHandle Query)
+{
+	QueryStore::Handle Handle;
+	Handle.Handle = Query;
+
+	if (Queries.IsAlive(Handle))
+	{
+		Queries.Get(Handle).NativeQuery.Clear();
+	}
+
+	Queries.Remove(Handle);
+}
+
+FTypedElementOnDataStorageCreation& UTypedElementDatabase::OnCreation()
+{
+	return OnCreationDelegate;
+}
+
+FTypedElementOnDataStorageDestruction& UTypedElementDatabase::OnDestruction()
+{
+	return OnDestructionDelegate;
+}
+
+FTypedElementOnDataStorageUpdate& UTypedElementDatabase::OnUpdate()
+{
+	return OnUpdateDelegate;
+}
+
+bool UTypedElementDatabase::IsAvailable() const
+{
+	return ActiveEditorWorld && ActiveEditorEntityManager;
 }
 
 void* UTypedElementDatabase::GetExternalSystemAddress(UClass* Target)
diff --git a/Engine/Plugins/TypedElementsDataStorage/Source/TypedElementsDataStorage/Private/TypedElementDatabase.h b/Engine/Plugins/TypedElementsDataStorage/Source/TypedElementsDataStorage/Private/TypedElementDatabase.h
--- a/Engine/Plugins/TypedElementsDataStorage/Source/TypedElementsDataStorage/Private/TypedElementDatabase.h
+++ b/Engine/Plugins/TypedElementsDataStorage/Source/TypedElementsDataStorage/Private/TypedElementDatabase.h
@@ -6,6 +6,8 @@
 #include "Templates/SharedPointer.h"
 #include "Elements/Interfaces/TypedElementDataStorageInterface.h"
 #include "MassArchetypeTypes.h"
+#include "MassEntityQuery.h"
+#include "TypedElementHandleStore.h"
 
 #include "TypedElementDatabase.generated.h"
 
@@ -19,6 +21,11 @@ class TYPEDELEMENTSDATASTORAGE_API UTypedElementDatabase
 {
 	GENERATED_BODY()
 public:
+	struct FExtendedQuery
+	{
+		FMassEntityQuery NativeQuery;
+	};
+
 	~UTypedElementDatabase() override = default;
 
 	void Initialize();
@@ -27,11 +34,6 @@ public:
 	TSharedPtr<FMassEntityManager> GetActiveMutableEditorEntityManager();
 	TSharedPtr<const FMassEntityManager> GetActiveEditorEntityManager() const;
 
-	FTypedElementOnDataStorageCreation& OnCreation() override;
-	FTypedElementOnDataStorageDestruction& OnDestruction() override;
-	FTypedElementOnDataStorageUpdate& OnUpdate() override;
-	bool IsAvailable() const override;
-
 	TypedElementTableHandle RegisterTable(TConstArrayView<const UScriptStruct*> ColumnList) override;
 	TypedElementTableHandle RegisterTable(TConstArrayView<const UScriptStruct*> ColumnList, const FName Name) override;
 
@@ -41,16 +43,34 @@ public:
 	bool BatchAddRow(FName TableName, int32 Count, TypedElementDataStorageCreationCallbackRef OnCreated) override;
 	void RemoveRow(TypedElementRowHandle Row) override;
 
+	void AddTag(TypedElementRowHandle Row, const UScriptStruct* TagType) override;
+	void AddTag(TypedElementRowHandle Row, FTopLevelAssetPath TagName) override;
+	void* AddOrGetColumnData(TypedElementRowHandle Row, const UScriptStruct* ColumnType) override;
+	ColumnDataResult AddOrGetColumnData(TypedElementRowHandle Row, FTopLevelAssetPath ColumnName) override;
 	void* GetColumnData(TypedElementRowHandle Row, const UScriptStruct* ColumnType) override;
+	ColumnDataResult AddOrGetColumnData(TypedElementRowHandle Row, FTopLevelAssetPath ColumnName,
+		TConstArrayView<TypedElement::ColumnUtils::Argument> Arguments) override;
+	ColumnDataResult GetColumnData(TypedElementRowHandle Row, FTopLevelAssetPath ColumnName) override;
+
+	TypedElementQueryHandle RegisterQuery(const QueryDescription& Query) override;
+	void UnregisterQuery(TypedElementQueryHandle Query) override;
 
+	FTypedElementOnDataStorageCreation& OnCreation() override;
+	FTypedElementOnDataStorageDestruction& OnDestruction() override;
+	FTypedElementOnDataStorageUpdate& OnUpdate() override;
+	bool IsAvailable() const override;
 	void* GetExternalSystemAddress(UClass* Target) override;
 
 private:
+	using QueryStore = TTypedElementHandleStore<FExtendedQuery>;
+
 	void Reset();
 	
 	TArray<FMassArchetypeHandle> Tables;
 	TMap<FName, TypedElementTableHandle> TableNameLookup;
 
+	QueryStore Queries;
+
 	FTypedElementOnDataStorageCreation OnCreationDelegate;
 	FTypedElementOnDataStorageDestruction OnDestructionDelegate;
 	FTypedElementOnDataStorageUpdate OnUpdateDelegate;
diff --git a/Engine/Plugins/TypedElementsDataStorage/Source/TypedElementsDataStorage/Private/TypedElementHandleStore.h b/Engine/Plugins/TypedElementsDataStorage/Source/TypedElementsDataStorage/Private/TypedElementHandleStore.h
new file mode 100644
--- /dev/null
+++ b/Engine/Plugins/TypedElementsDataStorage/Source/TypedElementsDataStorage/Private/TypedElementHandleStore.h
@@ -0,0 +1,111 @@
+// Copyright Epic Games, Inc. All Rights Reserved.
+
+#pragma once
+
+#include "Containers/Array.h"
+#include "Containers/Deque.h"
+
+template<typename DataType, uint32 ReservationSize = 64>
+class TTypedElementHandleStore
+{
+public:
+	struct HandleData
+	{
+		uint32 Index;
+		uint32 Generation;
+	};
+	union Handle
+	{
+		HandleData Data;
+		uint64 Handle;
+	};
+	
+	template<typename... Args>
+	Handle Emplace(Args... Arguments);
+	
+	DataType& Get(Handle Entry);
+	DataType& GetMutable(Handle Entry);
+	const DataType& Get(Handle Entry) const;
+
+	/** Removes the entry at the given handle if alive. This doesn't destroy the entry until the slot is reused. */
+	void Remove(Handle Entry);
+
+	bool IsAlive(Handle Entry) const;
+
+private:
+	TArray<DataType> Data;
+	TArray<uint32> Generations;
+	TDeque<uint32> RecycleBin; // Using a deque to avoid the same slot being constantly reused.
+};
+
+
+// Implementations
+
+template<typename DataType, uint32 ReservationSize>
+template<typename... Args>
+auto TTypedElementHandleStore<DataType, ReservationSize>::Emplace(Args... Arguments) -> Handle
+{
+	if (RecycleBin.IsEmpty())
+	{
+		uint32 NewSize = Data.Num() + ReservationSize;
+		Data.Reserve(NewSize);
+		Generations.Reserve(NewSize);
+
+		for (uint32 Index = Data.Num(); Index < NewSize; ++Index)
+		{
+			Data.AddZeroed();
+			Generations.Add(0);
+			RecycleBin.EmplaceLast(Index);
+		}
+	}
+
+	uint32 Index = RecycleBin.Last();
+	RecycleBin.PopLast();
+	
+	DataType& Entry = Data[Index];
+	new(&Entry) DataType( Forward<Args>(Arguments)... );
+
+	Handle Result;
+	Result.Data.Index = Index;
+	Result.Data.Generation = Generations[Index];
+	return Result;
+}
+
+template<typename DataType, uint32 ReservationSize>
+DataType& TTypedElementHandleStore<DataType, ReservationSize>::Get(Handle Entry)
+{
+	checkf(IsAlive(Entry), TEXT("Attempting to retrieve a dead entry from a Typed Element Handle Store."));
+	return Data[Entry.Data.Index];
+}
+
+template<typename DataType, uint32 ReservationSize>
+DataType& TTypedElementHandleStore<DataType, ReservationSize>::GetMutable(Handle Entry)
+{
+	checkf(IsAlive(Entry), TEXT("Attempting to retrieve a dead entry from a Typed Element Handle Store."));
+	return Data[Entry.Data.Index];
+}
+
+template<typename DataType, uint32 ReservationSize>
+const DataType& TTypedElementHandleStore<DataType, ReservationSize>::Get(Handle Entry) const
+{
+	checkf(IsAlive(Entry), TEXT("Attempting to retrieve a dead entry from a Typed Element Handle Store."));
+	return Data[Entry.Data.Index];
+}
+
+template<typename DataType, uint32 ReservationSize>
+void TTypedElementHandleStore<DataType, ReservationSize>::Remove(Handle Entry)
+{
+	if (IsAlive(Entry))
+	{
+		++Generations[Entry.Data.Index];
+		RecycleBin.EmplaceLast(Entry.Data.Index);
+	}
+}
+
+template<typename DataType, uint32 ReservationSize>
+bool TTypedElementHandleStore<DataType, ReservationSize>::IsAlive(Handle Entry) const
+{
+	return
+		Entry.Data.Index < static_cast<uint32>(Generations.Num()) &&
+		Generations[Entry.Data.Index] == Entry.Data.Generation;
+}
\ No newline at end of file
diff --git a/Engine/Source/Runtime/TypedElementFramework/Private/Elements/Framework/TypedElementColumnUtils.cpp b/Engine/Source/Runtime/TypedElementFramework/Private/Elements/Framework/TypedElementColumnUtils.cpp
new file mode 100644
--- /dev/null
+++ b/Engine/Source/Runtime/TypedElementFramework/Private/Elements/Framework/TypedElementColumnUtils.cpp
@@ -0,0 +1,30 @@
+// Copyright Epic Games, Inc. All Rights Reserved.
+
+#include "Elements/Framework/TypedElementColumnUtils.h"
+
+#include "UObject/Class.h"
+#include "UObject/UnrealType.h"
+
+namespace TypedElement::ColumnUtils
+{
+	void SetColumnValue(void* ColumnData, const UScriptStruct* ColumnType, FName ArgumentName, const FString& ArgumentValue)
+	{
+		FProperty* Property = ColumnType->FindPropertyByName(ArgumentName);
+		if (!Property)
+		{
+			Property = ColumnType->CustomFindProperty(ArgumentName);
+		}
+		if (Property)
+		{
+			Property->ImportText_Direct(*ArgumentValue, Property->ContainerPtrToValuePtr<uint8>(ColumnData, 0), nullptr, 0);
+		}
+	}
+
+	void SetColumnValues(void* ColumnData, const UScriptStruct* ColumnType, TConstArrayView<Argument> Arguments)
+	{
+		for (const Argument& Arg : Arguments)
+		{
+			SetColumnValue(ColumnData, ColumnType, Arg.Name, Arg.Value);
+		}
+	}
+} // namespace TypedElement::ColumnUtils
\ No newline at end of file
diff --git a/Engine/Source/Runtime/TypedElementFramework/Private/Elements/Framework/TypedElementQueryBuilder.cpp b/Engine/Source/Runtime/TypedElementFramework/Private/Elements/Framework/TypedElementQueryBuilder.cpp
new file mode 100644
--- /dev/null
+++ b/Engine/Source/Runtime/TypedElementFramework/Private/Elements/Framework/TypedElementQueryBuilder.cpp
@@ -0,0 +1,200 @@
+// Copyright Epic Games, Inc. All Rights Reserved.
+
+#include "Elements/Framework/TypedElementQueryBuilder.h"
+
+namespace TypedElementQueryBuilder
+{
+	const UScriptStruct* Type(FTopLevelAssetPath Name)
+	{
+		const UScriptStruct* StructInfo = TypeOptional(Name);
+		checkf(StructInfo, TEXT("Type name '%s' used as part of building a typed element query was not found."), *Name.ToString());
+		return StructInfo;
+	}
+
+	const UScriptStruct* TypeOptional(FTopLevelAssetPath Name)
+	{
+		constexpr bool bExactMatch = true;
+		return static_cast<UScriptStruct*>(StaticFindObject(UScriptStruct::StaticClass(), Name, bExactMatch));
+	}
+
+	const UScriptStruct* operator""_Type(const char* Name, std::size_t NameSize)
+	{
+		return Type(FTopLevelAssetPath{ FAnsiStringView{ Name, IntCastChecked<int32>(NameSize) } });
+	}
+
+	const UScriptStruct* operator""_TypeOptional(const char* Name, std::size_t NameSize)
+	{
+		return TypeOptional(FTopLevelAssetPath{ FAnsiStringView{ Name, IntCastChecked<int32>(NameSize) } });
+	}
+
+	/**
+	 * DependsOn
+	 */
+
+	Dependency::Dependency(ITypedElementDataStorageInterface::QueryDescription* Query)
+		: Query(Query)
+	{
+	}
+
+	Dependency& Dependency::ReadOnly(const UClass* Target)
+	{
+		checkf(Target, TEXT("The Dependency section in the Typed Elements query builder doesn't support nullptrs as Read-Only input."));
+		Query->Dependencies.Emplace(Target, EAccessType::ReadOnly);
+		return *this;
+	}
+
+	Dependency& Dependency::ReadOnly(std::initializer_list<const UClass*> Targets)
+	{
+		for (const UClass* Target : Targets)
+		{
+			ReadOnly(Targets);
+		}
+		return *this;
+	}
+
+	Dependency& Dependency::ReadWrite(const UClass* Target)
+	{
+		checkf(Target, TEXT("The Dependency section in the Typed Elements query builder doesn't support nullptrs as Read/Write input."));
+		Query->Dependencies.Emplace(Target, EAccessType::ReadWrite);
+		return *this;
+	}
+
+	Dependency& Dependency::ReadWrite(std::initializer_list<const UClass*> Targets)
+	{
+		for (const UClass* Target : Targets)
+		{
+			ReadWrite(Targets);
+		}
+		return *this;
+	}
+
+	ITypedElementDataStorageInterface::QueryDescription&& Dependency::Commit()
+	{
+		return MoveTemp(*Query);
+	}
+
+
+	/**
+	 * SimpleQuery
+	 */
+
+	SimpleQuery::SimpleQuery(ITypedElementDataStorageInterface::QueryDescription* Query)
+		: Query(Query)
+	{
+		Query->bSimpleQuery = true;
+	}
+
+	SimpleQuery& SimpleQuery::All(const UScriptStruct* Target)
+	{
+		if (Target)
+		{
+			Query->ConditionTypes.Add(ITypedElementDataStorageInterface::QueryDescription::OperatorType::SimpleAll);
+			Query->ConditionOperators.AddZeroed_GetRef().Type = Target;
+		}
+		return *this;
+	}
+
+	SimpleQuery& SimpleQuery::All(std::initializer_list<const UScriptStruct*> Targets)
+	{
+		for (const UScriptStruct* Target : Targets)
+		{
+			All(Target);
+		}
+		return *this;
+	}
+
+	SimpleQuery& SimpleQuery::Any(const UScriptStruct* Target)
+	{
+		if (Target)
+		{
+			Query->ConditionTypes.Add(ITypedElementDataStorageInterface::QueryDescription::OperatorType::SimpleAny);
+			Query->ConditionOperators.AddZeroed_GetRef().Type = Target;
+		}
+		return *this;
+	}
+
+	SimpleQuery& SimpleQuery::Any(std::initializer_list<const UScriptStruct*> Targets)
+	{
+		for (const UScriptStruct* Target : Targets)
+		{
+			Any(Target);
+		}
+		return *this;
+	}
+
+	SimpleQuery& SimpleQuery::None(const UScriptStruct* Target)
+	{
+		if (Target)
+		{
+			Query->ConditionTypes.Add(ITypedElementDataStorageInterface::QueryDescription::OperatorType::SimpleNone);
+			Query->ConditionOperators.AddZeroed_GetRef().Type = Target;
+		}
+		return *this;
+	}
+
+	SimpleQuery& SimpleQuery::None(std::initializer_list<const UScriptStruct*> Targets)
+	{
+		for (const UScriptStruct* Target : Targets)
+		{
+			None(Target);
+		}
+		return *this;
+	}
+
+	Dependency SimpleQuery::DependsOn()
+	{
+		return Dependency{ Query };
+	}
+
+
+	/**
+	 * Select
+	 */
+	
+	Select& Select::ReadOnly(const UScriptStruct* Target)
+	{
+		checkf(Target, TEXT("The Select section in the Typed Elements query builder doesn't support nullptrs as Read-Only input."));
+		Query.Selection.Emplace(Target, EAccessType::ReadOnly);
+		return *this;
+	}
+
+	Select& Select::ReadOnly(std::initializer_list<const UScriptStruct*> Targets)
+	{
+		for (const UScriptStruct* Target : Targets)
+		{
+			ReadOnly(Target);
+		}
+		return *this;
+	}
+
+	Select& Select::ReadWrite(const UScriptStruct* Target)
+	{
+		checkf(Target, TEXT("The Select section in the Typed Elements query builder doesn't support nullptrs as Read/Write input."));
+		Query.Selection.Emplace(Target, EAccessType::ReadWrite);
+		return *this;
+	}
+
+	Select& Select::ReadWrite(std::initializer_list<const UScriptStruct*> Targets)
+	{
+		for (const UScriptStruct* Target : Targets)
+		{
+			ReadWrite(Target);
+		}
+		return *this;
+	}
+
+	SimpleQuery Select::Where()
+	{
+		return SimpleQuery{ &Query };
+	}
+
+	Dependency Select::DependsOn()
+	{
+		return Dependency{ &Query };
+	}
+
+	ITypedElementDataStorageInterface::QueryDescription&& Select::Commit()
+	{
+		return MoveTemp(Query);
+	}
+}
\ No newline at end of file
diff --git a/Engine/Source/Runtime/TypedElementFramework/Public/Elements/Framework/TypedElementColumnUtils.h b/Engine/Source/Runtime/TypedElementFramework/Public/Elements/Framework/TypedElementColumnUtils.h
new file mode 100644
--- /dev/null
+++ b/Engine/Source/Runtime/TypedElementFramework/Public/Elements/Framework/TypedElementColumnUtils.h
@@ -0,0 +1,45 @@
+// Copyright Epic Games, Inc. All Rights Reserved.
+
+#pragma once
+
+#include "CoreMinimal.h"
+#include "Containers/ArrayView.h"
+
+class UScriptStruct;
+
+namespace TypedElement::ColumnUtils
+{
+	struct Argument
+	{
+		FName Name;
+		FString Value;
+	};
+
+	/** 
+	 * If found in the column, the variable under the provided name will be updated to the target value. 
+	 * The target variable must be convertable from string. Any incompatible or missing variables will
+	 * be silently ignored.
+	 */
+	TYPEDELEMENTFRAMEWORK_API void SetColumnValue(
+		void* ColumnData, const UScriptStruct* ColumnType, FName ArgumentName, const FString& ArgumentValue);
+	/**
+	 * If found in the column, the variables under the provided names will be updated to the target values.
+	 * The target variables must be convertable from string. Any incompatible or missing variables will
+	 * be silently ignored.
+	 */
+	TYPEDELEMENTFRAMEWORK_API void SetColumnValues(void* ColumnData, const UScriptStruct* ColumnType, TConstArrayView<Argument> Arguments);
+
+	/** Utility function for SetColumnValue that automatically detects the type of the column. */
+	template<typename ColumnType>
+	void SetColumnValue(ColumnType& Column, FName ArgumentName, const FString& ArgumentValue)
+	{
+		SetColumnValue(&Column, ColumnType::StaticStruct(), ArgumentName, ArgumentValue);
+	}
+
+	/** Utility function for SetColumnValues that automatically detects the type of the column. */
+	template<typename ColumnType>
+	void SetColumnValues(ColumnType& Column, TConstArrayView<Argument> Arguments)
+	{
+		SetColumnValues(&Column, ColumnType::StaticStruct(), Arguments);
+	}
+} // namespace TypedElement::ColumnUtils
\ No newline at end of file
diff --git a/Engine/Source/Runtime/TypedElementFramework/Public/Elements/Framework/TypedElementQueryBuilder.h b/Engine/Source/Runtime/TypedElementFramework/Public/Elements/Framework/TypedElementQueryBuilder.h
new file mode 100644
--- /dev/null
+++ b/Engine/Source/Runtime/TypedElementFramework/Public/Elements/Framework/TypedElementQueryBuilder.h
@@ -0,0 +1,198 @@
+// Copyright Epic Games, Inc. All Rights Reserved.
+
+#pragma once
+
+#include "CoreMinimal.h"
+#include "Elements/Interfaces/TypedElementDataStorageInterface.h"
+
+class UScriptStruct;
+
+/**
+ * The TypedElementQueryBuilder allows for the construction of queries for use by the Typed Element Data Storage.
+ * There are two types of queries, simple and normal. Simple queries are guaranteed to be supported by the data
+ * storage backend and guaranteed to have no performance side effects. <Normal queries pending development.>
+ * 
+ * Queries are constructed with the following section:
+ * - Select		A list of the data objects that are returned as the result of the query.
+ * - Where		A list of conditions that restrict what's accepted by the query.
+ * - DependsOn	A list of systems outside the data storage that will be accessed by the query('s user).
+ * - Commit		Compiles the query into its final form and can be used afterwards.
+ * 
+ * Calls to the sections become increasingly restrictive, e.g. after calling Where only DependsOn can be
+ * called again.
+ * 
+ * Arguments to the various functions take a pointer to a description of a UStruct. These can be provived
+ * in the follow ways:
+ * - By using the templated version, e.g. Any<FStructExample>()
+ * - By callling the static StaticStruct() function on the UStruct, e.g. FStructExample::StaticStruct();
+ * - By name using the Type or TypeOptional string operator, e.g. "/Script/ExamplePackage.FStructExample"_Type or
+ *		"/Script/OptionalPackage.FStructOptional"_TypeOptional
+ * All functions allow for a single type to be added or a list of types, e.g. ReadOnly(Type<FStructExample>() or
+ *		ReadOnly({ Type<FStructExample1>(), FStructExample2::StaticStruct(), "/Script/ExamplePackage.FStructExample3"_Type });
+ * 
+ * "Select" is constructed with: 
+ * - ReadOnly: Indicates that the data object will only be read from
+ * - ReadWrite: Indicated that the data object will be read and written to.
+ * 
+ * "Where" is constructed with:
+ * - All: The query will be accepted only if all the types listed here are present in a table.
+ * - Any: The query will be accepted if at least one of the listed types is present in a table.
+ * - Not: The query will be accepted if none of the listed types are present in a table.
+ * The above construction calls can be mixed and be called multiple times.
+ * All functions accept a nullptr for the type in which case the call will have no effect. This can be used to
+ *		reference types in plugins that may not be loaded when using the TypeOptional string operator.
+
+ * "DependsOn" is constructed with:
+ * - ReadOnly: Indicates that the external system will only be used to read data from.
+ * - ReadWrite: Indicates that the external system will be used to write data to.
+ * 
+ * Usage example:
+ * ITypedElementDataStorageInterface::QueryDescription Query =
+ *		Select{}
+ *			.ReadWrite({ FDataExample1::StaticStruct() })
+ *			.ReadWrite<FDataExample2, FDataExample3>()
+ *			.ReadOnly<FDataExample4>()
+ *		.Where()
+ *			.All<FTagExample1, FDataExample5>()
+ *			.Any("/Script/ExamplePackage.FStructExample"_TypeOptional)
+ *			.None(FTagExample2::StaticStruct())
+ *		.DependsOn()
+ *			.ReadOnly<USystemExample1, USystemExample2>()
+ *			.ReadWrite(USystemExample2::StaticClass())
+ *		.Commit();
+ * 
+ * Creating a query is expensive on the builder and the backend side. It's therefore recommended to create a query
+ * and store its compiled form for repeated use instead of rebuilding the query on every update.
+ */
+
+namespace TypedElementQueryBuilder
+{
+	using EAccessType = ITypedElementDataStorageInterface::QueryDescription::EAccessType;
+
+	TYPEDELEMENTFRAMEWORK_API const UScriptStruct* Type(FTopLevelAssetPath Name);
+	TYPEDELEMENTFRAMEWORK_API const UScriptStruct* TypeOptional(FTopLevelAssetPath Name);
+	TYPEDELEMENTFRAMEWORK_API const UScriptStruct* operator""_Type(const char* Name, std::size_t NameSize);
+	TYPEDELEMENTFRAMEWORK_API const UScriptStruct* operator""_TypeOptional(const char* Name, std::size_t NameSize);
+
+	class TYPEDELEMENTFRAMEWORK_API Dependency final
+	{
+		friend class Select;
+		friend class SimpleQuery;
+	public:
+		template<typename... TargetTypes>
+		Dependency& ReadOnly();
+		Dependency& ReadOnly(const UClass* Target);
+		Dependency& ReadOnly(std::initializer_list<const UClass*> Targets);
+		template<typename... TargetTypes>
+		Dependency& ReadWrite();
+		Dependency& ReadWrite(const UClass* Target);
+		Dependency& ReadWrite(std::initializer_list<const UClass*> Targets);
+
+		ITypedElementDataStorageInterface::QueryDescription&& Commit();
+
+	private:
+		explicit Dependency(ITypedElementDataStorageInterface::QueryDescription* Query);
+		
+		ITypedElementDataStorageInterface::QueryDescription* Query;
+	};
+
+	class TYPEDELEMENTFRAMEWORK_API SimpleQuery final
+	{
+	public:
+		friend class Select;
+		Dependency DependsOn();
+
+		template<typename... TargetTypes>
+		SimpleQuery& All();
+		SimpleQuery& All(const UScriptStruct* Target);
+		SimpleQuery& All(std::initializer_list<const UScriptStruct*> Targets);
+		template<typename... TargetTypes>
+		SimpleQuery& Any();
+		SimpleQuery& Any(const UScriptStruct* Target);
+		SimpleQuery& Any(std::initializer_list<const UScriptStruct*> Targets);
+		template<typename... TargetTypes>
+		SimpleQuery& None();
+		SimpleQuery& None(const UScriptStruct* Target);
+		SimpleQuery& None(std::initializer_list<const UScriptStruct*> Targets);
+
+	private:
+		explicit SimpleQuery(ITypedElementDataStorageInterface::QueryDescription* Query);
+
+		ITypedElementDataStorageInterface::QueryDescription* Query;
+	};
+
+	class TYPEDELEMENTFRAMEWORK_API Select final
+	{
+	public:
+		template<typename... TargetTypes>
+		Select& ReadOnly();
+		Select& ReadOnly(const UScriptStruct* Target);
+		Select& ReadOnly(std::initializer_list<const UScriptStruct*> Targets);
+		template<typename... TargetTypes>
+		Select& ReadWrite();
+		Select& ReadWrite(const UScriptStruct* Target);
+		Select& ReadWrite(std::initializer_list<const UScriptStruct*> Targets);
+
+		ITypedElementDataStorageInterface::QueryDescription&& Commit();
+		SimpleQuery Where();
+		Dependency DependsOn();
+
+	private:
+		ITypedElementDataStorageInterface::QueryDescription Query;
+	};
+
+
+	// Implementations
+
+	template<typename... TargetTypes>
+	Select& Select::ReadOnly()
+	{
+		ReadOnly({ TargetTypes::StaticStruct()... });
+		return *this;
+	}
+
+	template<typename... TargetTypes>
+	Select& Select::ReadWrite()
+	{
+		ReadWrite({ TargetTypes::StaticStruct()... });
+		return *this;
+	}
+
+
+	template<typename... TargetTypes>
+	SimpleQuery& SimpleQuery::All()
+	{
+		All({ TargetTypes::StaticStruct()... });
+		return *this;
+	}
+
+	template<typename... TargetTypes>
+	SimpleQuery& SimpleQuery::Any()
+	{
+		Any({ TargetTypes::StaticStruct()... });
+		return *this;
+	}
+
+	template<typename... TargetTypes>
+	SimpleQuery& SimpleQuery::None()
+	{
+		None({ TargetTypes::StaticStruct()... });
+		return *this;
+	}
+
+
+	template<typename... TargetTypes>
+	Dependency& Dependency::ReadOnly()
+	{
+		ReadOnly( { TargetTypes::StaticClass()... } );
+		return *this;
+	}
+
+	template<typename... TargetTypes>
+	Dependency& Dependency::ReadWrite()
+	{
+		ReadWrite({ TargetTypes::StaticClass()... });
+		return *this;
+	}
+
+} // namespace TypedElementQueryBuilder
\ No newline at end of file
diff --git a/Engine/Source/Runtime/TypedElementFramework/Public/Elements/Interfaces/TypedElementDataStorageInterface.h b/Engine/Source/Runtime/TypedElementFramework/Public/Elements/Interfaces/TypedElementDataStorageInterface.h
--- a/Engine/Source/Runtime/TypedElementFramework/Public/Elements/Interfaces/TypedElementDataStorageInterface.h
+++ b/Engine/Source/Runtime/TypedElementFramework/Public/Elements/Interfaces/TypedElementDataStorageInterface.h
@@ -2,8 +2,10 @@
 
 #pragma once
 
+#include "Containers/Array.h"
 #include "Containers/ContainersFwd.h"
 #include "Delegates/Delegate.h"
+#include "Elements/Framework/TypedElementColumnUtils.h"
 #include "Math/NumericLimits.h"
 #include "UObject/Interface.h"
 #include "UObject/NameTypes.h"
@@ -14,10 +16,20 @@
 class UClass;
 class UScriptStruct;
 
+struct ColumnDataResult
+{
+	/** Pointer to the structure that holds the description of the returned data. */
+	const UScriptStruct* Description;
+	/** Pointer to the column data. The type is guaranteed to match type described in Description. */
+	void* Data;
+};
+
 using TypedElementTableHandle = uint64;
 static constexpr auto TypedElementInvalidTableHandle = TNumericLimits<TypedElementTableHandle>::Max();
 using TypedElementRowHandle = uint64;
 static constexpr auto TypedElementInvalidRowHandle = TNumericLimits<TypedElementRowHandle>::Max();
+using TypedElementQueryHandle = uint64;
+static constexpr auto TypedElementInvalidQueryHandle = TNumericLimits<TypedElementQueryHandle>::Max();
 
 using FTypedElementOnDataStorageCreation = FSimpleMulticastDelegate;
 using FTypedElementOnDataStorageDestruction = FSimpleMulticastDelegate;
@@ -37,45 +49,28 @@ class TYPEDELEMENTFRAMEWORK_API ITypedElementDataStorageInterface
 
 public:
 	/**
-	 * Delegate that will be called when the data storage has been created. For various reasons the data storage
-	 * can be destroyed and recreated during the lifetime of the editor. This delegate allows for the opportunity
-	 * to register for this event so information such as table definitions can be registered again.
-	 * Note that by the time this delegate is available there will already be a data storage created. It's therefore
-	 * important that any information is registered with the data storage before registering a callback with the
-	 * delegate.
-	 */
-	virtual FTypedElementOnDataStorageCreation& OnCreation() = 0;
-	/**
-	 * Delegate that will be called when the data storage is about to be destroyed. This can be use as an opportunity
-	 * to store any data in the data storage before it becomes unavailable.
-	 */
-	virtual FTypedElementOnDataStorageDestruction& OnDestruction() = 0;
-	/**
-	 * Called periodically when the storage is available. This provides an opportunity to do any repeated processing
-	 * for the data storage.
-	 */
-	virtual FTypedElementOnDataStorageUpdate& OnUpdate() = 0;
-
-	/**
-	 * Whether or not the data storage is available. The data storage is available most of the time, but can be
-	 * unavailable for a brief time between being destroyed and a new one created.
-	 */
-	virtual bool IsAvailable() const = 0;
-
-	/**
+	 * @section Table management
+	 * 
+	 * @description
 	 * Tables are automatically created by taking an existing table and adding/removing columns. For
-	 * performance its however better to create a tables before hand and add objects to the table. This
+	 * performance its however better to create a table before adding objects to the table. This
 	 * doesn't prevent those objects from having columns added/removed at a later time.
-	 * To make debugging and profiling easier it's also recommended to add a name.
+	 * To make debugging and profiling easier it's also recommended to give tables a name.
 	 */
+
+	/** Creates a new table for with the provided columns. Optionally a name can be given which is useful for retrieval later. */
 	virtual TypedElementTableHandle RegisterTable(TConstArrayView<const UScriptStruct*> ColumnList) = 0;
 	virtual TypedElementTableHandle RegisterTable(TConstArrayView<const UScriptStruct*> ColumnList, const FName Name) = 0;
 
+	
 	/**
-	 * Adds a row to a table.
+	 * @section Row management
 	 */
+
+	/** Adds a new row to the provided table. */
 	virtual TypedElementRowHandle AddRow(TypedElementTableHandle Table) = 0;
 	virtual TypedElementRowHandle AddRow(FName TableName) = 0;
+	
 	/**
 	 * Add multiple rows at once. For each new row the OnCreated callback is called. Callers are expected to use the callback to
 	 * initialize the row if needed.
@@ -83,22 +78,138 @@ public:
 	virtual bool BatchAddRow(TypedElementTableHandle Table, int32 Count, TypedElementDataStorageCreationCallbackRef OnCreated) = 0;
 	virtual bool BatchAddRow(FName TableName, int32 Count, TypedElementDataStorageCreationCallbackRef OnCreated) = 0;
 	
+	/** Removes a previously added row. */
+	virtual void RemoveRow(TypedElementRowHandle Row) = 0;
+
+	
 	/**
-	 * Removes a previously added row.
+	 * @section Column management
 	 */
-	virtual void RemoveRow(TypedElementRowHandle Row) = 0;
 
+	/** Adds a tag to a row or does nothing if already added. */
+	virtual void AddTag(TypedElementRowHandle Row, const UScriptStruct* TagType) = 0;
+	virtual void AddTag(TypedElementRowHandle Row, FTopLevelAssetPath TagName) = 0;
+	
+	/** Adds a new column to a row. If the column already exists it will be returned instead. */
+	virtual void* AddOrGetColumnData(TypedElementRowHandle Row, const UScriptStruct* ColumnType) = 0;
+	virtual ColumnDataResult AddOrGetColumnData(TypedElementRowHandle Row, FTopLevelAssetPath ColumnName) = 0;
+	
 	/**
-	 * Retrieves a pointer to the column of the given row or a nullptr if not found.
+	 * Sets the data of a column using the provided argument bag. This is only meant for simple initialization for 
+	 * fragments that use UPROPERTY to expose properties. For complex initialization or when the fragment type is known
+	 * it's recommende to use call that work directly on the type for better performance and a wider range of configuration options.
 	 */
+	virtual ColumnDataResult AddOrGetColumnData(TypedElementRowHandle Row, FTopLevelAssetPath ColumnName,
+		TConstArrayView<TypedElement::ColumnUtils::Argument> Arguments) = 0;
+	
+	/** Retrieves a pointer to the column of the given row or a nullptr if not found. */
 	virtual void* GetColumnData(TypedElementRowHandle Row, const UScriptStruct* ColumnType) = 0;
+	virtual ColumnDataResult GetColumnData(TypedElementRowHandle Row, FTopLevelAssetPath ColumnName) = 0;
+	
 
 	/**
-	 * Returns a pointer to the registered external system if found, otherwise null.
+	 * @section Query
+	 * @description
+	 * Queries can be constructed using the Query Builder. Note that the Query Builder allows for the creation of queries that
+	 * are more complex than the backend may support. The backend is allowed to simplify the query, in which case the query
+	 * can be used directly in the processor to do additional filtering. This will however impact performance and it's 
+	 * therefore recommended to try to simplify the query first before relying on extended query filtering in a processor.
 	 */
+	struct QueryDescription final
+	{
+		using OperatorIndex = int32;
+		enum class OperatorType
+		{
+			SimpleAll,			// Unary: Type
+			SimpleAny,			// Unary: Type
+			SimpleNone,			// Unary: Type
+			SimpleOptional,		// Unary: Type
+			And,				// Binary: left operator index, right operator index
+			Or,					// Binary: left operator index, right operator index
+			Not,				// Unary: condition index
+			Type				// Unary: Type
+		};
+
+		struct BinaryOperator final
+		{
+			OperatorIndex Left;
+			OperatorIndex Right;
+		};
+		
+		union Operator
+		{
+			BinaryOperator Binary;
+			OperatorIndex Unary;
+			const UScriptStruct* Type;
+		};
+
+		enum class EAccessType { ReadOnly, ReadWrite };
+
+		struct AccessControlledStruct
+		{
+			AccessControlledStruct() = default;
+			inline AccessControlledStruct(const UScriptStruct* Type, EAccessType Access) : Type(Type), Access(Access) {}
+
+			const UScriptStruct* Type;
+			EAccessType Access;
+		};
+
+		struct AccessControlledClass
+		{
+			AccessControlledClass() = default;
+			inline AccessControlledClass(const UClass* Type, EAccessType Access) : Type(Type), Access(Access) {}
+
+			const UClass* Type;
+			EAccessType Access;
+		};
+
+		TArray<AccessControlledStruct, TInlineAllocator<16>> Selection;
+		TArray<OperatorType, TInlineAllocator<32>> ConditionTypes;
+		TArray<Operator, TInlineAllocator<32>> ConditionOperators;
+		TArray<AccessControlledClass, TInlineAllocator<4>> Dependencies;
+		/** If true, this query only has simple operations and is guaranteed to be executed fully and at optimal performance. */
+		bool bSimpleQuery{ false };
+	};
+	virtual TypedElementQueryHandle RegisterQuery(const QueryDescription& Query) = 0;
+	virtual void UnregisterQuery(TypedElementQueryHandle Query) = 0;
+	
+	/**
+	 * @section Misc
+	 */
+	
+	/**
+	 * Delegate that will be called when the data storage has been created. For various reasons the data storage
+	 * can be destroyed and recreated during the lifetime of the editor. This delegate allows for the opportunity
+	 * to register for this event so information such as table definitions can be registered again.
+	 * Note that by the time this delegate is available there will already be a data storage created. It's therefore
+	 * important that any information is registered with the data storage before registering a callback with the
+	 * delegate.
+	 */
+	virtual FTypedElementOnDataStorageCreation& OnCreation() = 0;
+	
+	/**
+	 * Delegate that will be called when the data storage is about to be destroyed. This can be use as an opportunity
+	 * to store any data in the data storage before it becomes unavailable.
+	 */
+	virtual FTypedElementOnDataStorageDestruction& OnDestruction() = 0;
+	
+	/**
+	 * Called periodically when the storage is available. This provides an opportunity to do any repeated processing
+	 * for the data storage.
+	 */
+	virtual FTypedElementOnDataStorageUpdate& OnUpdate() = 0;
+
+	/**
+	 * Whether or not the data storage is available. The data storage is available most of the time, but can be
+	 * unavailable for a brief time between being destroyed and a new one created.
+	 */
+	virtual bool IsAvailable() const = 0;
+
+	/** Returns a pointer to the registered external system if found, otherwise null. */
 	virtual void* GetExternalSystemAddress(UClass* Target) = 0;
 
 	
+
 	/**
 	 *
 	 * The following are utility functions that are not part of the interface but are provided in order to make using the
@@ -107,20 +218,47 @@ public:
 	 */
 	
 	
+	/** Adds a tag to a row or does nothing if already added. */
+	template<typename TagType>
+	void AddTag();
+	
 	/**
-	 * Returns a pointer to the column of the given row or a nullptr if the type couldn't be found or the row doesn't exist.
+	 * Returns a pointer to the column of the given row or creates a new one if not found. Optionally arguments can be provided
+	 * to update or initialize the column's data.
 	 */
+	template<typename ColumnType, typename... Args>
+	ColumnType* AddOrGetColumn(TypedElementRowHandle Row, Args... Arguments);
+	
+	/** Returns a pointer to the column of the given row or a nullptr if the type couldn't be found or the row doesn't exist. */
 	template<typename ColumnType>
 	ColumnType* GetColumn(TypedElementRowHandle Row);
 
-	/**
-	 * Returns a pointer to the registered external system if found, otherwise null.
-	 */
+	/** Returns a pointer to the registered external system if found, otherwise null. */
 	template<typename SystemType>
 	SystemType* GetExternalSystem();
 };
 
 // Implementations
+template<typename TagType>
+void ITypedElementDataStorageInterface::AddTag()
+{
+	AddTag(TagType::StaticStruct());
+}
+
+template<typename ColumnType, typename... Args>
+ColumnType* ITypedElementDataStorageInterface::AddOrGetColumn(TypedElementRowHandle Row, Args... Arguments)
+{
+	auto* Result = reinterpret_cast<ColumnType*>(AddOrGetColumnData(Row, ColumnType::StaticStruct()));
+	if constexpr (sizeof...(Arguments) > 0)
+	{
+		if (Result)
+		{
+			new(Result) ColumnType{ Forward<Arguments>... };
+		}
+	}
+	return Result;
+}
+
 template<typename ColumnType>
 ColumnType* ITypedElementDataStorageInterface::GetColumn(TypedElementRowHandle Row)
 {
