The Typed Element queries have been updated with support for a count query and the initial step to run queries.

Count queries are queries that don't select any columns but instead return a count of the total number of rows that pass the filter conditions. A call was also added that allows these queries to be run, which is currently limited to retrieving the result of a count query as the other query type, select, can't be direclty run.
This change list also introduces FTypedElementDataStorageColumn and FTypedElementDataStorageTag as specific base classes for columns tags. These are used in upcoming changes.

diff --git a/Engine/Plugins/TypedElementsDataStorage/Source/TypedElementsDataStorage/Private/TypedElementDatabase.cpp b/Engine/Plugins/TypedElementsDataStorage/Source/TypedElementsDataStorage/Private/TypedElementDatabase.cpp
--- a/Engine/Plugins/TypedElementsDataStorage/Source/TypedElementsDataStorage/Private/TypedElementDatabase.cpp
+++ b/Engine/Plugins/TypedElementsDataStorage/Source/TypedElementsDataStorage/Private/TypedElementDatabase.cpp
@@ -250,15 +250,15 @@ ColumnDataResult UTypedElementDatabase::GetColumnData(TypedElementRowHandle Row,
 	return ColumnDataResult{ nullptr, nullptr };
 }
 
-TypedElementQueryHandle UTypedElementDatabase::RegisterQuery(const QueryDescription& Query)
+TypedElementQueryHandle UTypedElementDatabase::RegisterQuery(const FQueryDescription& Query)
 {
-	auto LocalToNativeAccess = [](QueryDescription::EAccessType Access) -> EMassFragmentAccess
+	auto LocalToNativeAccess = [](FQueryDescription::EAccessType Access) -> EMassFragmentAccess
 	{
 		switch (Access)
 		{
-		case QueryDescription::EAccessType::ReadOnly:
+		case FQueryDescription::EAccessType::ReadOnly:
 			return EMassFragmentAccess::ReadOnly;
-		case QueryDescription::EAccessType::ReadWrite:
+		case FQueryDescription::EAccessType::ReadWrite:
 			return EMassFragmentAccess::ReadWrite;
 		default:
 			checkf(false, TEXT("Invalid query access type: %i."), static_cast<uint32>(Access));
@@ -267,14 +267,25 @@ TypedElementQueryHandle UTypedElementDatabase::RegisterQuery(const QueryDescript
 	};
 
 	QueryStore::Handle Result = Queries.Emplace();
-	FMassEntityQuery& NativeQuery = Queries.GetMutable(Result).NativeQuery;
+	FExtendedQuery& StoredQuery = Queries.GetMutable(Result);
+	StoredQuery.Action = Query.Action;
+	StoredQuery.bSimpleQuery = Query.bSimpleQuery;
 
-	for (const QueryDescription::AccessControlledStruct& SelectEntry : Query.Selection)
+	FMassEntityQuery& NativeQuery = StoredQuery.NativeQuery;
+
+	if (Query.Action == FQueryDescription::EActionType::Count)
+	{
+		checkf(Query.Selection.IsEmpty(), TEXT("Count queries for the Typed Elements Data Storage can't have entries for selection."));
+	}
+	else
 	{
-		checkf(SelectEntry.Type, TEXT("Provided query selection type can not be null."));
-		checkf(SelectEntry.Type->IsChildOf(FMassFragment::StaticStruct()), 
-			TEXT("Provided query selection type '%s' is not based on FMassFragment."), *SelectEntry.Type->GetStructPathName().ToString());
-		NativeQuery.AddRequirement(SelectEntry.Type, LocalToNativeAccess(SelectEntry.Access));
+		for (const FQueryDescription::FAccessControlledStruct& SelectEntry : Query.Selection)
+		{
+			checkf(SelectEntry.Type, TEXT("Provided query selection type can not be null."));
+			checkf(SelectEntry.Type->IsChildOf(FTypedElementDataStorageColumn::StaticStruct()),
+				TEXT("Provided query selection type '%s' is not based on FTypedElementDataStorageColumn."), *SelectEntry.Type->GetStructPathName().ToString());
+			NativeQuery.AddRequirement(SelectEntry.Type, LocalToNativeAccess(SelectEntry.Access));
+		}
 	}
 
 	if (Query.bSimpleQuery) // This backend currently only supports simple queries.
@@ -282,19 +293,19 @@ TypedElementQueryHandle UTypedElementDatabase::RegisterQuery(const QueryDescript
 		checkf(Query.ConditionTypes.Num() == Query.ConditionOperators.Num(),
 			TEXT("The types and operators for a typed element query have gone out of sync."));
 		
-		const QueryDescription::Operator* Operand = Query.ConditionOperators.GetData();
-		for (QueryDescription::OperatorType Type : Query.ConditionTypes)
+		const FQueryDescription::FOperator* Operand = Query.ConditionOperators.GetData();
+		for (FQueryDescription::EOperatorType Type : Query.ConditionTypes)
 		{
 			EMassFragmentPresence Presence;
 			switch (Type)
 			{
-			case QueryDescription::OperatorType::SimpleAll:
+			case FQueryDescription::EOperatorType::SimpleAll:
 				Presence = EMassFragmentPresence::All;
 				break;
-			case QueryDescription::OperatorType::SimpleAny:
+			case FQueryDescription::EOperatorType::SimpleAny:
 				Presence = EMassFragmentPresence::Any;
 				break;
-			case QueryDescription::OperatorType::SimpleNone:
+			case FQueryDescription::EOperatorType::SimpleNone:
 				Presence = EMassFragmentPresence::None;
 				break;
 			default:
@@ -314,7 +325,7 @@ TypedElementQueryHandle UTypedElementDatabase::RegisterQuery(const QueryDescript
 		}
 	}
 
-	for (const QueryDescription::AccessControlledClass& DependencyEntry : Query.Dependencies)
+	for (const FQueryDescription::FAccessControlledClass& DependencyEntry : Query.Dependencies)
 	{
 		checkf(DependencyEntry.Type, TEXT("Provided query dependcy type can not be null."));
 		checkf(DependencyEntry.Type->IsChildOf<UWorldSubsystem>(),
@@ -340,6 +351,50 @@ void UTypedElementDatabase::UnregisterQuery(TypedElementQueryHandle Query)
 	Queries.Remove(Handle);
 }
 
+ITypedElementDataStorageInterface::FQueryResult UTypedElementDatabase::RunQuery(TypedElementQueryHandle Query)
+{
+	FQueryResult Result;
+
+	QueryStore::Handle Handle;
+	Handle.Handle = Query;
+
+	if (Queries.IsAlive(Handle))
+	{
+		FExtendedQuery& QueryData = Queries.Get(Handle);
+		if (QueryData.bSimpleQuery)
+		{
+			switch (QueryData.Action)
+			{
+			case FQueryDescription::EActionType::None:
+				Result.Completed = FQueryResult::ECompletion::Fully;
+				break;
+			case FQueryDescription::EActionType::Select:
+				Result.Completed = FQueryResult::ECompletion::Unsupported;
+				break;
+			case FQueryDescription::EActionType::Count:
+				checkf(ActiveEditorEntityManager, 
+					TEXT("Unable to run Typed Element Data Storage query before a MASS Entity Manager has been assigned."));
+				Result.Count = QueryData.NativeQuery.GetNumMatchingEntities(*ActiveEditorEntityManager);
+				Result.Completed = FQueryResult::ECompletion::Fully;
+				break;
+			default:
+				Result.Completed = FQueryResult::ECompletion::Unsupported;
+				break;
+			}
+		}
+		else
+		{
+			Result.Completed = FQueryResult::ECompletion::Unsupported;
+		}
+	}
+	else
+	{
+		Result.Completed = FQueryResult::ECompletion::Unavailable;
+	}
+
+	return Result;
+}
+
 FTypedElementOnDataStorageCreation& UTypedElementDatabase::OnCreation()
 {
 	return OnCreationDelegate;
diff --git a/Engine/Plugins/TypedElementsDataStorage/Source/TypedElementsDataStorage/Private/TypedElementDatabase.h b/Engine/Plugins/TypedElementsDataStorage/Source/TypedElementsDataStorage/Private/TypedElementDatabase.h
--- a/Engine/Plugins/TypedElementsDataStorage/Source/TypedElementsDataStorage/Private/TypedElementDatabase.h
+++ b/Engine/Plugins/TypedElementsDataStorage/Source/TypedElementsDataStorage/Private/TypedElementDatabase.h
@@ -24,6 +24,8 @@ public:
 	struct FExtendedQuery
 	{
 		FMassEntityQuery NativeQuery;
+		FQueryDescription::EActionType Action{ FQueryDescription::EActionType::None };
+		bool bSimpleQuery{ false };
 	};
 
 	~UTypedElementDatabase() override = default;
@@ -52,8 +54,9 @@ public:
 		TConstArrayView<TypedElement::ColumnUtils::Argument> Arguments) override;
 	ColumnDataResult GetColumnData(TypedElementRowHandle Row, FTopLevelAssetPath ColumnName) override;
 
-	TypedElementQueryHandle RegisterQuery(const QueryDescription& Query) override;
+	TypedElementQueryHandle RegisterQuery(const FQueryDescription& Query) override;
 	void UnregisterQuery(TypedElementQueryHandle Query) override;
+	FQueryResult RunQuery(TypedElementQueryHandle Query) override;
 
 	FTypedElementOnDataStorageCreation& OnCreation() override;
 	FTypedElementOnDataStorageDestruction& OnDestruction() override;
diff --git a/Engine/Source/Runtime/TypedElementFramework/Private/Elements/Framework/TypedElementQueryBuilder.cpp b/Engine/Source/Runtime/TypedElementFramework/Private/Elements/Framework/TypedElementQueryBuilder.cpp
--- a/Engine/Source/Runtime/TypedElementFramework/Private/Elements/Framework/TypedElementQueryBuilder.cpp
+++ b/Engine/Source/Runtime/TypedElementFramework/Private/Elements/Framework/TypedElementQueryBuilder.cpp
@@ -31,19 +31,19 @@ namespace TypedElementQueryBuilder
 	 * DependsOn
 	 */
 
-	Dependency::Dependency(ITypedElementDataStorageInterface::QueryDescription* Query)
+	FDependency::FDependency(ITypedElementDataStorageInterface::FQueryDescription* Query)
 		: Query(Query)
 	{
 	}
 
-	Dependency& Dependency::ReadOnly(const UClass* Target)
+	FDependency& FDependency::ReadOnly(const UClass* Target)
 	{
 		checkf(Target, TEXT("The Dependency section in the Typed Elements query builder doesn't support nullptrs as Read-Only input."));
 		Query->Dependencies.Emplace(Target, EAccessType::ReadOnly);
 		return *this;
 	}
 
-	Dependency& Dependency::ReadOnly(std::initializer_list<const UClass*> Targets)
+	FDependency& FDependency::ReadOnly(std::initializer_list<const UClass*> Targets)
 	{
 		for (const UClass* Target : Targets)
 		{
@@ -52,14 +52,14 @@ namespace TypedElementQueryBuilder
 		return *this;
 	}
 
-	Dependency& Dependency::ReadWrite(const UClass* Target)
+	FDependency& FDependency::ReadWrite(const UClass* Target)
 	{
 		checkf(Target, TEXT("The Dependency section in the Typed Elements query builder doesn't support nullptrs as Read/Write input."));
 		Query->Dependencies.Emplace(Target, EAccessType::ReadWrite);
 		return *this;
 	}
 
-	Dependency& Dependency::ReadWrite(std::initializer_list<const UClass*> Targets)
+	FDependency& FDependency::ReadWrite(std::initializer_list<const UClass*> Targets)
 	{
 		for (const UClass* Target : Targets)
 		{
@@ -68,33 +68,33 @@ namespace TypedElementQueryBuilder
 		return *this;
 	}
 
-	ITypedElementDataStorageInterface::QueryDescription&& Dependency::Commit()
+	ITypedElementDataStorageInterface::FQueryDescription&& FDependency::Compile()
 	{
 		return MoveTemp(*Query);
 	}
 
 
 	/**
-	 * SimpleQuery
+	 * Simple Query
 	 */
 
-	SimpleQuery::SimpleQuery(ITypedElementDataStorageInterface::QueryDescription* Query)
+	FSimpleQuery::FSimpleQuery(ITypedElementDataStorageInterface::FQueryDescription* Query)
 		: Query(Query)
 	{
 		Query->bSimpleQuery = true;
 	}
 
-	SimpleQuery& SimpleQuery::All(const UScriptStruct* Target)
+	FSimpleQuery& FSimpleQuery::All(const UScriptStruct* Target)
 	{
 		if (Target)
 		{
-			Query->ConditionTypes.Add(ITypedElementDataStorageInterface::QueryDescription::OperatorType::SimpleAll);
+			Query->ConditionTypes.Add(ITypedElementDataStorageInterface::FQueryDescription::EOperatorType::SimpleAll);
 			Query->ConditionOperators.AddZeroed_GetRef().Type = Target;
 		}
 		return *this;
 	}
 
-	SimpleQuery& SimpleQuery::All(std::initializer_list<const UScriptStruct*> Targets)
+	FSimpleQuery& FSimpleQuery::All(std::initializer_list<const UScriptStruct*> Targets)
 	{
 		for (const UScriptStruct* Target : Targets)
 		{
@@ -103,17 +103,17 @@ namespace TypedElementQueryBuilder
 		return *this;
 	}
 
-	SimpleQuery& SimpleQuery::Any(const UScriptStruct* Target)
+	FSimpleQuery& FSimpleQuery::Any(const UScriptStruct* Target)
 	{
 		if (Target)
 		{
-			Query->ConditionTypes.Add(ITypedElementDataStorageInterface::QueryDescription::OperatorType::SimpleAny);
+			Query->ConditionTypes.Add(ITypedElementDataStorageInterface::FQueryDescription::EOperatorType::SimpleAny);
 			Query->ConditionOperators.AddZeroed_GetRef().Type = Target;
 		}
 		return *this;
 	}
 
-	SimpleQuery& SimpleQuery::Any(std::initializer_list<const UScriptStruct*> Targets)
+	FSimpleQuery& FSimpleQuery::Any(std::initializer_list<const UScriptStruct*> Targets)
 	{
 		for (const UScriptStruct* Target : Targets)
 		{
@@ -122,17 +122,17 @@ namespace TypedElementQueryBuilder
 		return *this;
 	}
 
-	SimpleQuery& SimpleQuery::None(const UScriptStruct* Target)
+	FSimpleQuery& FSimpleQuery::None(const UScriptStruct* Target)
 	{
 		if (Target)
 		{
-			Query->ConditionTypes.Add(ITypedElementDataStorageInterface::QueryDescription::OperatorType::SimpleNone);
+			Query->ConditionTypes.Add(ITypedElementDataStorageInterface::FQueryDescription::EOperatorType::SimpleNone);
 			Query->ConditionOperators.AddZeroed_GetRef().Type = Target;
 		}
 		return *this;
 	}
 
-	SimpleQuery& SimpleQuery::None(std::initializer_list<const UScriptStruct*> Targets)
+	FSimpleQuery& FSimpleQuery::None(std::initializer_list<const UScriptStruct*> Targets)
 	{
 		for (const UScriptStruct* Target : Targets)
 		{
@@ -141,15 +141,25 @@ namespace TypedElementQueryBuilder
 		return *this;
 	}
 
-	Dependency SimpleQuery::DependsOn()
+	FDependency FSimpleQuery::DependsOn()
 	{
-		return Dependency{ Query };
+		return FDependency{ Query };
+	}
+
+	ITypedElementDataStorageInterface::FQueryDescription&& FSimpleQuery::Compile()
+	{
+		return MoveTemp(*Query);
 	}
 
 
 	/**
 	 * Select
 	 */
+
+	Select::Select()
+	{
+		Query.Action = ITypedElementDataStorageInterface::FQueryDescription::EActionType::Select;
+	}
 	
 	Select& Select::ReadOnly(const UScriptStruct* Target)
 	{
@@ -183,18 +193,38 @@ namespace TypedElementQueryBuilder
 		return *this;
 	}
 
-	SimpleQuery Select::Where()
+	FSimpleQuery Select::Where()
 	{
-		return SimpleQuery{ &Query };
+		return FSimpleQuery{ &Query };
 	}
 
-	Dependency Select::DependsOn()
+	FDependency Select::DependsOn()
 	{
-		return Dependency{ &Query };
+		return FDependency{ &Query };
 	}
 
-	ITypedElementDataStorageInterface::QueryDescription&& Select::Commit()
+	ITypedElementDataStorageInterface::FQueryDescription&& Select::Compile()
 	{
 		return MoveTemp(Query);
 	}
+
+
+	/**
+	 * Count
+	 */
+
+	Count::Count()
+	{
+		Query.Action = ITypedElementDataStorageInterface::FQueryDescription::EActionType::Count;
+	}
+
+	FSimpleQuery Count::Where()
+	{
+		return FSimpleQuery{ &Query };
+	}
+
+	FDependency Count::DependsOn()
+	{
+		return FDependency{ &Query };
+	}
 }
\ No newline at end of file
diff --git a/Engine/Source/Runtime/TypedElementFramework/Public/Elements/Framework/TypedElementQueryBuilder.h b/Engine/Source/Runtime/TypedElementFramework/Public/Elements/Framework/TypedElementQueryBuilder.h
--- a/Engine/Source/Runtime/TypedElementFramework/Public/Elements/Framework/TypedElementQueryBuilder.h
+++ b/Engine/Source/Runtime/TypedElementFramework/Public/Elements/Framework/TypedElementQueryBuilder.h
@@ -14,9 +14,10 @@ class UScriptStruct;
  * 
  * Queries are constructed with the following section:
  * - Select		A list of the data objects that are returned as the result of the query.
+ * - Count		Counts the total number or rows that pass the filter.
  * - Where		A list of conditions that restrict what's accepted by the query.
  * - DependsOn	A list of systems outside the data storage that will be accessed by the query('s user).
- * - Commit		Compiles the query into its final form and can be used afterwards.
+ * - Compile	Compiles the query into its final form and can be used afterwards.
  * 
  * Calls to the sections become increasingly restrictive, e.g. after calling Where only DependsOn can be
  * called again.
@@ -34,6 +35,8 @@ class UScriptStruct;
  * - ReadOnly: Indicates that the data object will only be read from
  * - ReadWrite: Indicated that the data object will be read and written to.
  * 
+ * "Count" does not have any construction options.
+ * 
  * "Where" is constructed with:
  * - All: The query will be accepted only if all the types listed here are present in a table.
  * - Any: The query will be accepted if at least one of the listed types is present in a table.
@@ -47,8 +50,8 @@ class UScriptStruct;
  * - ReadWrite: Indicates that the external system will be used to write data to.
  * 
  * Usage example:
- * ITypedElementDataStorageInterface::QueryDescription Query =
- *		Select{}
+ * ITypedElementDataStorageInterface::FQueryDescription Query =
+ *		Select()
  *			.ReadWrite({ FDataExample1::StaticStruct() })
  *			.ReadWrite<FDataExample2, FDataExample3>()
  *			.ReadOnly<FDataExample4>()
@@ -59,7 +62,7 @@ class UScriptStruct;
  *		.DependsOn()
  *			.ReadOnly<USystemExample1, USystemExample2>()
  *			.ReadWrite(USystemExample2::StaticClass())
- *		.Commit();
+ *		.Compile();
  * 
  * Creating a query is expensive on the builder and the backend side. It's therefore recommended to create a query
  * and store its compiled form for repeated use instead of rebuilding the query on every update.
@@ -67,63 +70,70 @@ class UScriptStruct;
 
 namespace TypedElementQueryBuilder
 {
-	using EAccessType = ITypedElementDataStorageInterface::QueryDescription::EAccessType;
+	using EAccessType = ITypedElementDataStorageInterface::FQueryDescription::EAccessType;
 
 	TYPEDELEMENTFRAMEWORK_API const UScriptStruct* Type(FTopLevelAssetPath Name);
 	TYPEDELEMENTFRAMEWORK_API const UScriptStruct* TypeOptional(FTopLevelAssetPath Name);
 	TYPEDELEMENTFRAMEWORK_API const UScriptStruct* operator""_Type(const char* Name, std::size_t NameSize);
 	TYPEDELEMENTFRAMEWORK_API const UScriptStruct* operator""_TypeOptional(const char* Name, std::size_t NameSize);
 
-	class TYPEDELEMENTFRAMEWORK_API Dependency final
+	class TYPEDELEMENTFRAMEWORK_API FDependency final
 	{
+		friend class Count;
 		friend class Select;
-		friend class SimpleQuery;
+		friend class FSimpleQuery;
 	public:
 		template<typename... TargetTypes>
-		Dependency& ReadOnly();
-		Dependency& ReadOnly(const UClass* Target);
-		Dependency& ReadOnly(std::initializer_list<const UClass*> Targets);
+		FDependency& ReadOnly();
+		FDependency& ReadOnly(const UClass* Target);
+		FDependency& ReadOnly(std::initializer_list<const UClass*> Targets);
 		template<typename... TargetTypes>
-		Dependency& ReadWrite();
-		Dependency& ReadWrite(const UClass* Target);
-		Dependency& ReadWrite(std::initializer_list<const UClass*> Targets);
+		FDependency& ReadWrite();
+		FDependency& ReadWrite(const UClass* Target);
+		FDependency& ReadWrite(std::initializer_list<const UClass*> Targets);
 
-		ITypedElementDataStorageInterface::QueryDescription&& Commit();
+		ITypedElementDataStorageInterface::FQueryDescription&& Compile();
 
 	private:
-		explicit Dependency(ITypedElementDataStorageInterface::QueryDescription* Query);
+		explicit FDependency(ITypedElementDataStorageInterface::FQueryDescription* Query);
 		
-		ITypedElementDataStorageInterface::QueryDescription* Query;
+		ITypedElementDataStorageInterface::FQueryDescription* Query;
 	};
 
-	class TYPEDELEMENTFRAMEWORK_API SimpleQuery final
+	class TYPEDELEMENTFRAMEWORK_API FSimpleQuery final
 	{
 	public:
+		friend class Count;
 		friend class Select;
-		Dependency DependsOn();
+
+		FDependency DependsOn();
+		ITypedElementDataStorageInterface::FQueryDescription&& Compile();
 
 		template<typename... TargetTypes>
-		SimpleQuery& All();
-		SimpleQuery& All(const UScriptStruct* Target);
-		SimpleQuery& All(std::initializer_list<const UScriptStruct*> Targets);
+		FSimpleQuery& All();
+		FSimpleQuery& All(const UScriptStruct* Target);
+		FSimpleQuery& All(std::initializer_list<const UScriptStruct*> Targets);
 		template<typename... TargetTypes>
-		SimpleQuery& Any();
-		SimpleQuery& Any(const UScriptStruct* Target);
-		SimpleQuery& Any(std::initializer_list<const UScriptStruct*> Targets);
+		FSimpleQuery& Any();
+		FSimpleQuery& Any(const UScriptStruct* Target);
+		FSimpleQuery& Any(std::initializer_list<const UScriptStruct*> Targets);
 		template<typename... TargetTypes>
-		SimpleQuery& None();
-		SimpleQuery& None(const UScriptStruct* Target);
-		SimpleQuery& None(std::initializer_list<const UScriptStruct*> Targets);
+		FSimpleQuery& None();
+		FSimpleQuery& None(const UScriptStruct* Target);
+		FSimpleQuery& None(std::initializer_list<const UScriptStruct*> Targets);
 
 	private:
-		explicit SimpleQuery(ITypedElementDataStorageInterface::QueryDescription* Query);
+		explicit FSimpleQuery(ITypedElementDataStorageInterface::FQueryDescription* Query);
 
-		ITypedElementDataStorageInterface::QueryDescription* Query;
+		ITypedElementDataStorageInterface::FQueryDescription* Query;
 	};
 
+	// Explicitly not following the naming convention in order to keep readability consistent. It now reads like a query sentence.
 	class TYPEDELEMENTFRAMEWORK_API Select final
 	{
 	public:
+		Select();
+
 		template<typename... TargetTypes>
 		Select& ReadOnly();
 		Select& ReadOnly(const UScriptStruct* Target);
@@ -133,66 +143,90 @@ namespace TypedElementQueryBuilder
 		Select& ReadWrite(const UScriptStruct* Target);
 		Select& ReadWrite(std::initializer_list<const UScriptStruct*> Targets);
 
-		ITypedElementDataStorageInterface::QueryDescription&& Commit();
-		SimpleQuery Where();
-		Dependency DependsOn();
+		ITypedElementDataStorageInterface::FQueryDescription&& Compile();
+		FSimpleQuery Where();
+		FDependency DependsOn();
+
+	private:
+		ITypedElementDataStorageInterface::FQueryDescription Query;
+	};
+
+	// Explicitly not following the naming convention in order to keep readability consistent. It now reads like a query sentence.
+	class TYPEDELEMENTFRAMEWORK_API Count final
+	{
+	public:
+		Count();
+
+		FSimpleQuery Where();
+		FDependency DependsOn();
 
 	private:
-		ITypedElementDataStorageInterface::QueryDescription Query;
+		ITypedElementDataStorageInterface::FQueryDescription Query;
 	};
 
 
 	// Implementations
 
+	//
+	// FDependecy
+	//
+
 	template<typename... TargetTypes>
-	Select& Select::ReadOnly()
+	FDependency& FDependency::ReadOnly()
 	{
-		ReadOnly({ TargetTypes::StaticStruct()... });
+		ReadOnly({ TargetTypes::StaticClass()... });
 		return *this;
 	}
 
 	template<typename... TargetTypes>
-	Select& Select::ReadWrite()
+	FDependency& FDependency::ReadWrite()
 	{
-		ReadWrite({ TargetTypes::StaticStruct()... });
+		ReadWrite({ TargetTypes::StaticClass()... });
 		return *this;
 	}
 
 
+	//
+	// Select
+	//
+
 	template<typename... TargetTypes>
-	SimpleQuery& SimpleQuery::All()
+	Select& Select::ReadOnly()
 	{
-		All({ TargetTypes::StaticStruct()... });
+		ReadOnly({ TargetTypes::StaticStruct()... });
 		return *this;
 	}
 
 	template<typename... TargetTypes>
-	SimpleQuery& SimpleQuery::Any()
+	Select& Select::ReadWrite()
 	{
-		Any({ TargetTypes::StaticStruct()... });
+		ReadWrite({ TargetTypes::StaticStruct()... });
 		return *this;
 	}
 
+
+	//
+	// FSimpleQuery
+	//
+
 	template<typename... TargetTypes>
-	SimpleQuery& SimpleQuery::None()
+	FSimpleQuery& FSimpleQuery::All()
 	{
-		None({ TargetTypes::StaticStruct()... });
+		All({ TargetTypes::StaticStruct()... });
 		return *this;
 	}
 
-
 	template<typename... TargetTypes>
-	Dependency& Dependency::ReadOnly()
+	FSimpleQuery& FSimpleQuery::Any()
 	{
-		ReadOnly( { TargetTypes::StaticClass()... } );
+		Any({ TargetTypes::StaticStruct()... });
 		return *this;
 	}
 
 	template<typename... TargetTypes>
-	Dependency& Dependency::ReadWrite()
+	FSimpleQuery& FSimpleQuery::None()
 	{
-		ReadWrite({ TargetTypes::StaticClass()... });
+		None({ TargetTypes::StaticStruct()... });
 		return *this;
 	}
-
 } // namespace TypedElementQueryBuilder
\ No newline at end of file
diff --git a/Engine/Source/Runtime/TypedElementFramework/Public/Elements/Interfaces/TypedElementDataStorageInterface.h b/Engine/Source/Runtime/TypedElementFramework/Public/Elements/Interfaces/TypedElementDataStorageInterface.h
--- a/Engine/Source/Runtime/TypedElementFramework/Public/Elements/Interfaces/TypedElementDataStorageInterface.h
+++ b/Engine/Source/Runtime/TypedElementFramework/Public/Elements/Interfaces/TypedElementDataStorageInterface.h
@@ -37,6 +37,24 @@ using FTypedElementOnDataStorageUpdate = FSimpleMulticastDelegate;
 
 using TypedElementDataStorageCreationCallbackRef = TFunctionRef<void(TypedElementRowHandle Row)>;
 
+/**
+ * Base for the data structures for a column.
+ */
+USTRUCT()
+struct FTypedElementDataStorageColumn
+{
+	GENERATED_BODY()
+};
+
+/**
+ * Base for the data structures that act as tags to rows. Tags should not have any data.
+ */
+USTRUCT()
+struct FTypedElementDataStorageTag
+{
+	GENERATED_BODY()
+};
+
 UINTERFACE(MinimalAPI)
 class UTypedElementDataStorageInterface : public UInterface
 {
@@ -115,10 +133,17 @@ public:
 	 * can be used directly in the processor to do additional filtering. This will however impact performance and it's 
 	 * therefore recommended to try to simplify the query first before relying on extended query filtering in a processor.
 	 */
-	struct QueryDescription final
+	struct FQueryDescription final
 	{
+		enum class EActionType : uint8
+		{
+			None,	/** Do nothing. */
+			Select,	/** Selects a set of columns for further processing. */
+			Count	/** Counts the number of entries that match the filter condition. */
+		};
+
 		using OperatorIndex = int32;
-		enum class OperatorType
+		enum class EOperatorType
 		{
 			SimpleAll,			// Unary: Type
 			SimpleAny,			// Unary: Type
@@ -130,48 +155,70 @@ public:
 			Type				// Unary: Type
 		};
 
-		struct BinaryOperator final
+		struct FBinaryOperator final
 		{
 			OperatorIndex Left;
 			OperatorIndex Right;
 		};
 		
-		union Operator
+		union FOperator
 		{
-			BinaryOperator Binary;
+			FBinaryOperator Binary;
 			OperatorIndex Unary;
 			const UScriptStruct* Type;
 		};
 
 		enum class EAccessType { ReadOnly, ReadWrite };
 
-		struct AccessControlledStruct
+		struct FAccessControlledStruct
 		{
-			AccessControlledStruct() = default;
-			inline AccessControlledStruct(const UScriptStruct* Type, EAccessType Access) : Type(Type), Access(Access) {}
+			FAccessControlledStruct() = default;
+			inline FAccessControlledStruct(const UScriptStruct* Type, EAccessType Access) : Type(Type), Access(Access) {}
 
 			const UScriptStruct* Type;
 			EAccessType Access;
 		};
 
-		struct AccessControlledClass
+		struct FAccessControlledClass
 		{
-			AccessControlledClass() = default;
-			inline AccessControlledClass(const UClass* Type, EAccessType Access) : Type(Type), Access(Access) {}
+			FAccessControlledClass() = default;
+			inline FAccessControlledClass(const UClass* Type, EAccessType Access) : Type(Type), Access(Access) {}
 
 			const UClass* Type;
 			EAccessType Access;
 		};
 
-		TArray<AccessControlledStruct, TInlineAllocator<16>> Selection;
-		TArray<OperatorType, TInlineAllocator<32>> ConditionTypes;
-		TArray<Operator, TInlineAllocator<32>> ConditionOperators;
-		TArray<AccessControlledClass, TInlineAllocator<4>> Dependencies;
+		TArray<FAccessControlledStruct, TInlineAllocator<16>> Selection;
+		TArray<EOperatorType, TInlineAllocator<32>> ConditionTypes;
+		TArray<FOperator, TInlineAllocator<32>> ConditionOperators;
+		TArray<FAccessControlledClass, TInlineAllocator<4>> Dependencies;
+		EActionType Action;
 		/** If true, this query only has simple operations and is guaranteed to be executed fully and at optimal performance. */
 		bool bSimpleQuery{ false };
 	};
-	virtual TypedElementQueryHandle RegisterQuery(const QueryDescription& Query) = 0;
+	virtual TypedElementQueryHandle RegisterQuery(const FQueryDescription& Query) = 0;
 	virtual void UnregisterQuery(TypedElementQueryHandle Query) = 0;
+	struct FQueryResult
+	{
+		enum class ECompletion
+		{
+			/** Query could be fully executed. */
+			Fully,
+			/** Only portions of the query were executed. This is caused by a problem that was encountered partway through processing. */
+			Partially,
+			/** 
+			 * The backend doesn't support the particular query. This may be a limitation in how/where the query is run or because
+			 * the query contains actions and/or operations that are not supported.
+			 */
+			Unsupported,
+			/** The provided query is no longer available. */
+			Unavailable
+		};
+		
+		uint32 Count{ 0 }; /** The number of rows were processed. */
+		ECompletion Completed{ ECompletion::Unavailable };
+	};
+	virtual FQueryResult RunQuery(TypedElementQueryHandle Query) = 0;
 	
 	/**
 	 * @section Misc
@@ -220,7 +267,7 @@ public:
 	
 	/** Adds a tag to a row or does nothing if already added. */
 	template<typename TagType>
-	void AddTag();
+	void AddTag(TypedElementRowHandle Row);
 	
 	/**
 	 * Returns a pointer to the column of the given row or creates a new one if not found. Optionally arguments can be provided
@@ -240,9 +287,9 @@ public:
 
 // Implementations
 template<typename TagType>
-void ITypedElementDataStorageInterface::AddTag()
+void ITypedElementDataStorageInterface::AddTag(TypedElementRowHandle Row)
 {
-	AddTag(TagType::StaticStruct());
+	AddTag(Row, TagType::StaticStruct());
 }
 
 template<typename ColumnType, typename... Args>
