Introduces the intial step to supporting UIs from the Typed Element Data Storage.

This change list provides the initial functionality to register and create UI elements with the Typed Element Data Storage.
This also enables FTypedElementDataStorageColumn and FTypedElementDataStorageTag to be used by making them appear as if they were FMassFragments and FMassTags respectively.

diff --git a/Engine/Plugins/TypedElementsDataStorage/Source/TypedElementsDataStorage/Private/TypedElementDatabaseUI.cpp b/Engine/Plugins/TypedElementsDataStorage/Source/TypedElementsDataStorage/Private/TypedElementDatabaseUI.cpp
new file mode 100644
--- /dev/null
+++ b/Engine/Plugins/TypedElementsDataStorage/Source/TypedElementsDataStorage/Private/TypedElementDatabaseUI.cpp
@@ -0,0 +1,116 @@
+// Copyright Epic Games, Inc. All Rights Reserved.
+
+#include "TypedElementDatabaseUI.h"
+
+#include "Elements/Columns/TypedElementSlateWidgetColumns.h"
+#include "GenericPlatform/GenericPlatformMemory.h"
+
+void UTypedElementDatabaseUi::Initialize(ITypedElementDataStorageInterface* StorageInterface)
+{
+	checkf(StorageInterface, TEXT("Typed Element's Database compatibility manager is being initialized with an invalid storage target."));
+
+	Storage = StorageInterface;
+	CreateStandardArchetypes();
+
+	StorageInterface->OnCreation().AddUObject(this, &UTypedElementDatabaseUi::CreateStandardArchetypes);
+}
+
+void UTypedElementDatabaseUi::Deinitialize()
+{
+}
+
+void UTypedElementDatabaseUi::RegisterWidgetFactory(FName Purpose, const UScriptStruct* Constructor)
+{
+	checkf(Constructor->IsChildOf(FTypedElementWidgetConstructor::StaticStruct()),
+		TEXT("Attempting to register a Typed Elements widget constructor '%s' that isn't deriving from FTypedElementWidgetConstructor."),
+		*Constructor->GetFullName());
+	WidgetFactoryStructs.Add(Purpose, Constructor);
+}
+
+void UTypedElementDatabaseUi::RegisterWidgetFactory(
+	FName Purpose, const UScriptStruct* ConstructorType, TUniquePtr<FTypedElementWidgetConstructor>&& Constructor)
+{
+	WidgetFactoryInstances.Add(Purpose, { MoveTemp(Constructor), ConstructorType });
+}
+
+void UTypedElementDatabaseUi::ConstructWidgets(FName Purpose, TConstArrayView<TypedElement::ColumnUtils::Argument> Arguments,
+	const WidgetCreatedCallback& ConstructionCallback)
+{
+	// Construct with structs
+	for (TMultiMap<FName, const UScriptStruct*>::TKeyIterator KeyValue = WidgetFactoryStructs.CreateKeyIterator(Purpose); 
+		KeyValue; ++KeyValue)
+	{
+		const UScriptStruct* Target = KeyValue.Value();
+		CreateWidgetInstanceFromDescription(Target, Arguments, ConstructionCallback);
+	}
+
+	// Construct with instances
+	for (TMultiMap<FName, FInstanceConstructor>::TKeyIterator KeyValue = WidgetFactoryInstances.CreateKeyIterator(Purpose); 
+		KeyValue; ++KeyValue)
+	{
+		FTypedElementWidgetConstructor* Source = KeyValue.Value().Constructor.Get();
+		const UScriptStruct* Target = KeyValue.Value().ConstructorType;
+		CreateWidgetInstanceFromInstance(Source, Target, Arguments, ConstructionCallback);
+	}
+}
+
+void UTypedElementDatabaseUi::CreateWidgetInstanceFromDescription(
+	const UScriptStruct* Target,
+	TConstArrayView<TypedElement::ColumnUtils::Argument> Arguments,
+	const WidgetCreatedCallback& ConstructionCallback)
+{
+	FTypedElementWidgetConstructor* Constructor = reinterpret_cast<FTypedElementWidgetConstructor*>(
+		FMemory_Alloca_Aligned(Target->GetStructureSize(), Target->GetMinAlignment()));
+	Target->InitializeStruct(Constructor);
+
+	CreateWidgetInstance(*Constructor, Arguments, ConstructionCallback, Target);
+}
+
+void UTypedElementDatabaseUi::CreateWidgetInstanceFromInstance(
+	FTypedElementWidgetConstructor* SourceConstructor,
+	const UScriptStruct* Target,
+	TConstArrayView<TypedElement::ColumnUtils::Argument> Arguments,
+	const WidgetCreatedCallback& ConstructionCallback)
+{
+	FTypedElementWidgetConstructor* Constructor = reinterpret_cast<FTypedElementWidgetConstructor*>(
+		FMemory_Alloca_Aligned(Target->GetStructureSize(), Target->GetMinAlignment()));
+	Target->InitializeStruct(Constructor);
+	Target->CopyScriptStruct(Constructor, SourceConstructor);
+
+	CreateWidgetInstance(*Constructor, Arguments, ConstructionCallback, Target);
+}
+
+void UTypedElementDatabaseUi::CreateWidgetInstance(
+	FTypedElementWidgetConstructor& Constructor, 
+	TConstArrayView<TypedElement::ColumnUtils::Argument> Arguments, 
+	const WidgetCreatedCallback& ConstructionCallback, 
+	const UScriptStruct* Target)
+{
+	TypedElementRowHandle Row = Storage->AddRow(WidgetTable);
+	TSharedPtr<SWidget> Widget = Constructor.Construct(Row, Storage, this, Arguments);
+	if (Widget)
+	{
+		ConstructionCallback(Widget.ToSharedRef(), Row);
+	}
+	else
+	{
+		Storage->RemoveRow(Row);
+	}
+
+	Target->DestroyStruct(&Constructor);
+}
+
+TSharedPtr<SWidget> UTypedElementDatabaseUi::ConstructWidget(TypedElementRowHandle Row, FTypedElementWidgetConstructor& Constructor,
+	TConstArrayView<TypedElement::ColumnUtils::Argument> Arguments)
+{
+	return Constructor.Construct(Row, Storage, this, Arguments);
+}
+
+void UTypedElementDatabaseUi::CreateStandardArchetypes()
+{
+	WidgetTable = Storage->RegisterTable(MakeArrayView(
+		{
+			FTypedElementSlateWidgetReferenceColumn::StaticStruct(),
+			FTypedElementSlateWidgetReferenceDeletesRowTag::StaticStruct()
+		}), FName("Editor_WidgetTable"));
+}
\ No newline at end of file
diff --git a/Engine/Plugins/TypedElementsDataStorage/Source/TypedElementsDataStorage/Private/TypedElementDatabaseUI.h b/Engine/Plugins/TypedElementsDataStorage/Source/TypedElementsDataStorage/Private/TypedElementDatabaseUI.h
new file mode 100644
--- /dev/null
+++ b/Engine/Plugins/TypedElementsDataStorage/Source/TypedElementsDataStorage/Private/TypedElementDatabaseUI.h
@@ -0,0 +1,66 @@
+// Copyright Epic Games, Inc. All Rights Reserved.
+
+#pragma once
+
+#include "Containers/Map.h"
+#include "Elements/Interfaces/TypedElementDataStorageInterface.h"
+#include "Elements/Interfaces/TypedElementDataStorageUiInterface.h"
+#include "UObject/ObjectMacros.h"
+
+#include "TypedElementDatabaseUI.generated.h"
+
+UCLASS()
+class TYPEDELEMENTSDATASTORAGE_API UTypedElementDatabaseUi final
+	: public UObject
+	, public ITypedElementDataStorageUiInterface
+{
+	GENERATED_BODY()
+
+public:
+	~UTypedElementDatabaseUi() override = default;
+
+	void Initialize(ITypedElementDataStorageInterface* StorageInterface);
+	void Deinitialize();
+
+	void RegisterWidgetFactory(FName Purpose, const UScriptStruct* Constructor) override;
+	
+	void ConstructWidgets(FName Purpose, TConstArrayView<TypedElement::ColumnUtils::Argument> Arguments,
+		const WidgetCreatedCallback& ConstructionCallback) override;
+	TSharedPtr<SWidget> ConstructWidget(TypedElementRowHandle Row, FTypedElementWidgetConstructor& Constructor,
+		TConstArrayView<TypedElement::ColumnUtils::Argument> Arguments) override;
+
+protected:
+	void RegisterWidgetFactory(FName Purpose, const UScriptStruct* ConstructorType,
+		TUniquePtr<FTypedElementWidgetConstructor>&& Constructor) override;
+
+private:
+	struct FInstanceConstructor
+	{
+		TUniquePtr<FTypedElementWidgetConstructor> Constructor;
+		const UScriptStruct* ConstructorType;
+	};
+
+	void CreateStandardArchetypes();
+
+	void CreateWidgetInstanceFromDescription(
+		const UScriptStruct* Target,
+		TConstArrayView<TypedElement::ColumnUtils::Argument> Arguments,
+		const WidgetCreatedCallback& ConstructionCallback);
+
+	void CreateWidgetInstanceFromInstance(
+		FTypedElementWidgetConstructor* SourceConstructor,
+		const UScriptStruct* Target,
+		TConstArrayView<TypedElement::ColumnUtils::Argument> Arguments,
+		const WidgetCreatedCallback& ConstructionCallback);
+
+	void CreateWidgetInstance(
+		FTypedElementWidgetConstructor& Constructor,
+		TConstArrayView<TypedElement::ColumnUtils::Argument> Arguments,
+		const WidgetCreatedCallback& ConstructionCallback,
+		const UScriptStruct* Target);
+
+	ITypedElementDataStorageInterface* Storage{ nullptr };
+	TypedElementTableHandle WidgetTable{ TypedElementInvalidTableHandle };
+	TMultiMap<FName, const UScriptStruct*> WidgetFactoryStructs;
+	TMultiMap<FName, FInstanceConstructor> WidgetFactoryInstances;
+};
\ No newline at end of file
diff --git a/Engine/Plugins/TypedElementsDataStorage/Source/TypedElementsDataStorage/Private/TypedElementsDataStorage.cpp b/Engine/Plugins/TypedElementsDataStorage/Source/TypedElementsDataStorage/Private/TypedElementsDataStorage.cpp
--- a/Engine/Plugins/TypedElementsDataStorage/Source/TypedElementsDataStorage/Private/TypedElementsDataStorage.cpp
+++ b/Engine/Plugins/TypedElementsDataStorage/Source/TypedElementsDataStorage/Private/TypedElementsDataStorage.cpp
@@ -3,13 +3,50 @@
 #include "TypedElementsDataStorage.h"
 
 #include "Elements/Framework/TypedElementRegistry.h"
+#include "MassEntityTypes.h"
 #include "Misc/CoreDelegates.h"
+#include "Templates/IsPolymorphic.h"
 #include "TypedElementDatabase.h"
 #include "TypedElementDatabaseCompatibility.h"
+#include "TypedElementDatabaseUI.h"
 #include "UObject/UObjectGlobals.h"
 
 #define LOCTEXT_NAMESPACE "FTypedElementsDataStorageModule"
 
+// MASS uses DCO in a few places, making it a difficult to consistently register Type Element's Columns and Tags
+// as they may have not been set up to impersonate MASS' Fragments and Tags yet. To get around this, do the 
+// registration during static initialization, which is before DCOs.
+/**
+ * Typed Elements provides base classes for columns and tags. These directly map to fragments and tags in MASS.
+ * To avoid deep and tight coupling between both systems, columns and tags don't directly inhered from MASS, but
+ * are otherwise fully compatible. To allow MASS to do its type safety checks, this class updates the type
+ * information so Typed Elements columns and tags present as MASS fragments and tags from MASS's perspective.
+ */
+struct FTagAndFragmentImpersonators
+{
+	FTagAndFragmentImpersonators()
+	{
+		// Have FTypedElementDataStorageColumn impersonate a FMassFragment, which is the actual data storage when using MASS as a backend.
+		static_assert(sizeof(FTypedElementDataStorageColumn) == sizeof(FMassFragment),
+			"In order for FTypedElementDataStorageColumn to impersonate FMassFragment they need to be identical.");
+		static_assert(!TIsPolymorphic<FMassFragment>::Value,
+			"In order to be able to impersonate FMassFragment it can't have any virtual functions.");
+		static_assert(!TIsPolymorphic<FTypedElementDataStorageColumn>::Value,
+			"In order to be able to use FTypedElementDataStorageColumn to impersonate FMassFragment it can't have any virtual functions.");
+		FTypedElementDataStorageColumn::StaticStruct()->SetSuperStruct(FMassFragment::StaticStruct());
+
+		// Have FTypedElementDataStorageTag impersonate a FMassTag, which is the tag type when using MASS as a backend.
+		static_assert(sizeof(FTypedElementDataStorageTag) == sizeof(FMassTag),
+			"In order for FTypedElementDataStorageTag to impersonate FMassTag they need to be identical.");
+		static_assert(!TIsPolymorphic<FMassTag>::Value,
+			"In order to be able to impersonate FMassTag it can't have any virtual functions.");
+		static_assert(!TIsPolymorphic<FTypedElementDataStorageTag>::Value,
+			"In order to be able to use FTypedElementDataStorageTag to impersonate FMassTag it can't have any virtual functions.");
+		FTypedElementDataStorageTag::StaticStruct()->SetSuperStruct(FMassTag::StaticStruct());
+	}
+};
+static FTagAndFragmentImpersonators TagAndFragmentImpersonators;
+
 void FTypedElementsDataStorageModule::StartupModule()
 {
 	FCoreDelegates::OnAllModuleLoadingPhasesComplete.AddLambda(
@@ -23,10 +60,15 @@ void FTypedElementsDataStorageModule::StartupModule()
 				DatabaseCompatibility = NewObject<UTypedElementDatabaseCompatibility>();
 				DatabaseCompatibility->Initialize(Database.Get());
 
+				DatabaseUi = NewObject<UTypedElementDatabaseUi>();
+				DatabaseUi->Initialize(Database.Get());
+
 				UTypedElementRegistry* Registry = UTypedElementRegistry::GetInstance();
-				check(Registry);
+				checkf(Registry, TEXT(
+					"FTypedElementsDataStorageModule tried to register itself, but there was no Typed Element Registry to register to."));
 				Registry->SetDataStorage(Database.Get());
 				Registry->SetDataStorageCompatibility(DatabaseCompatibility.Get());
+				Registry->SetDataStorageUi(DatabaseUi.Get());
 
 				bInitialized = true;
 			}
@@ -40,10 +82,13 @@ void FTypedElementsDataStorageModule::ShutdownModule()
 	{
 		Registry->SetDataStorage(nullptr);
 		Registry->SetDataStorageCompatibility(nullptr);
+		Registry->SetDataStorageUi(nullptr);
 	}
 
+	DatabaseUi->Deinitialize();
 	DatabaseCompatibility->Deinitialize();
 	Database->Deinitialize();
+
 	bInitialized = false;
 }
 
@@ -53,6 +98,7 @@ void FTypedElementsDataStorageModule::AddReferencedObjects(FReferenceCollector&
 	{
 		Collector.AddReferencedObject(Database);
 		Collector.AddReferencedObject(DatabaseCompatibility);
+		Collector.AddReferencedObject(DatabaseUi);
 	}
 }
 
diff --git a/Engine/Plugins/TypedElementsDataStorage/Source/TypedElementsDataStorage/Public/TypedElementsDataStorage.h b/Engine/Plugins/TypedElementsDataStorage/Source/TypedElementsDataStorage/Public/TypedElementsDataStorage.h
--- a/Engine/Plugins/TypedElementsDataStorage/Source/TypedElementsDataStorage/Public/TypedElementsDataStorage.h
+++ b/Engine/Plugins/TypedElementsDataStorage/Source/TypedElementsDataStorage/Public/TypedElementsDataStorage.h
@@ -12,6 +12,7 @@ class FReferenceCollector;
 class ITypedElementDataStorageInterface;
 class UTypedElementDatabase;
 class UTypedElementDatabaseCompatibility;
+class UTypedElementDatabaseUi;
 
 class FTypedElementsDataStorageModule : public IModuleInterface, public FGCObject
 {
@@ -36,5 +37,6 @@ public:
 private:
 	TObjectPtr<UTypedElementDatabase> Database;
 	TObjectPtr<UTypedElementDatabaseCompatibility> DatabaseCompatibility;
+	TObjectPtr<UTypedElementDatabaseUi> DatabaseUi;
 	bool bInitialized{ false };
 };
diff --git a/Engine/Source/Runtime/TypedElementFramework/Private/Elements/Framework/TypedElementRegistry.cpp b/Engine/Source/Runtime/TypedElementFramework/Private/Elements/Framework/TypedElementRegistry.cpp
--- a/Engine/Source/Runtime/TypedElementFramework/Private/Elements/Framework/TypedElementRegistry.cpp
+++ b/Engine/Source/Runtime/TypedElementFramework/Private/Elements/Framework/TypedElementRegistry.cpp
@@ -75,6 +75,7 @@ const ITypedElementDataStorageInterface* UTypedElementRegistry::GetDataStorage()
 void UTypedElementRegistry::SetDataStorage(ITypedElementDataStorageInterface* Storage)
 {
 	DataStorage = Storage;
+	CallDataStorageInterfacesSetDelegateIfNeeded();
 }
 
 ITypedElementDataStorageCompatibilityInterface* UTypedElementRegistry::GetMutableDataStorageCompatibility()
@@ -90,6 +91,28 @@ const ITypedElementDataStorageCompatibilityInterface* UTypedElementRegistry::Get
 void UTypedElementRegistry::SetDataStorageCompatibility(ITypedElementDataStorageCompatibilityInterface* Storage)
 {
 	DataStorageCompatibility = Storage;
+	CallDataStorageInterfacesSetDelegateIfNeeded();
+}
+
+ITypedElementDataStorageUiInterface* UTypedElementRegistry::GetMutableDataStorageUi()
+{
+	return DataStorageUi;
+}
+
+const ITypedElementDataStorageUiInterface* UTypedElementRegistry::GetDataStorageUi() const
+{
+	return DataStorageUi;
+}
+
+void UTypedElementRegistry::SetDataStorageUi(ITypedElementDataStorageUiInterface* Storage)
+{
+	DataStorageUi = Storage;
+	CallDataStorageInterfacesSetDelegateIfNeeded();
+}
+
+bool UTypedElementRegistry::AreDataStorageInterfacesSet() const
+{
+	return DataStorage && DataStorageCompatibility && DataStorageUi;
 }
 
 void UTypedElementRegistry::FinishDestroy()
@@ -379,3 +402,10 @@ void UTypedElementRegistry::OnPostGarbageCollect()
 	}
 }
 
+void UTypedElementRegistry::CallDataStorageInterfacesSetDelegateIfNeeded()
+{
+	if (OnDataStorageInterfacesSetDelegate.IsBound() && AreDataStorageInterfacesSet())
+	{
+		OnDataStorageInterfacesSetDelegate.Broadcast();
+	}
+}
\ No newline at end of file
diff --git a/Engine/Source/Runtime/TypedElementFramework/Private/Elements/Interfaces/TypedElementDataStorageUiInterface.cpp b/Engine/Source/Runtime/TypedElementFramework/Private/Elements/Interfaces/TypedElementDataStorageUiInterface.cpp
new file mode 100644
--- /dev/null
+++ b/Engine/Source/Runtime/TypedElementFramework/Private/Elements/Interfaces/TypedElementDataStorageUiInterface.cpp
@@ -0,0 +1,34 @@
+// Copyright Epic Games, Inc. All Rights Reserved.
+
+#include "Elements/Interfaces/TypedElementDataStorageUiInterface.h"
+#include "Elements/Framework/TypedElementColumnUtils.h"
+#include "Elements/Columns/TypedElementSlateWidgetColumns.h"
+
+TSharedPtr<SWidget> FTypedElementWidgetConstructor::Construct(
+	TypedElementRowHandle Row,
+	ITypedElementDataStorageInterface* DataStorage,
+	ITypedElementDataStorageUiInterface* DataStorageUi,
+	TConstArrayView<TypedElement::ColumnUtils::Argument> Arguments)
+{
+	ApplyArguments(Arguments);
+	TSharedPtr<SWidget> Widget = CreateWidget();
+	if (Widget)
+	{
+		AddColumns(DataStorage, Row, Widget);
+	}
+	return Widget;
+}
+
+void FTypedElementWidgetConstructor::ApplyArguments(TConstArrayView<TypedElement::ColumnUtils::Argument> Arguments)
+{
+	SetColumnValues(*this, Arguments);
+}
+
+void FTypedElementWidgetConstructor::AddColumns(
+	ITypedElementDataStorageInterface* DataStorage, TypedElementRowHandle Row, const TSharedPtr<SWidget>& Widget)
+{
+	DataStorage->AddTag<FTypedElementSlateWidgetReferenceDeletesRowTag>(Row);
+	FTypedElementSlateWidgetReferenceColumn* StorageColumn = DataStorage->AddOrGetColumn<FTypedElementSlateWidgetReferenceColumn>(Row);
+	checkf(StorageColumn, TEXT("Expected Typed Element Data Storage to contain a FTypedElementSlateWidgetReferenceFragment for row %llu."), Row);
+	StorageColumn->Widget = Widget;
+}
\ No newline at end of file
diff --git a/Engine/Source/Runtime/TypedElementFramework/Public/Elements/Columns/TypedElementSlateWidgetColumns.h b/Engine/Source/Runtime/TypedElementFramework/Public/Elements/Columns/TypedElementSlateWidgetColumns.h
new file mode 100644
--- /dev/null
+++ b/Engine/Source/Runtime/TypedElementFramework/Public/Elements/Columns/TypedElementSlateWidgetColumns.h
@@ -0,0 +1,36 @@
+// Copyright Epic Games, Inc. All Rights Reserved.
+
+#pragma once
+
+#include "Elements/Interfaces/TypedElementDataStorageInterface.h"
+#include "Templates/SharedPointer.h"
+#include "UObject/ObjectMacros.h"
+
+#include "TypedElementSlateWidgetColumns.generated.h"
+
+class SWidget;
+
+/**
+ * Stores a widget reference in the data storage. At the start of processing any
+ * columns that are not pointing to a valid widget will be removed. If the
+ * FTypedElementSlateWidgetDeletesRowTag is found then the entire row will
+ * be deleted.
+ */
+USTRUCT()
+struct FTypedElementSlateWidgetReferenceColumn : public FTypedElementDataStorageColumn
+{
+	GENERATED_BODY()
+
+	TWeakPtr<SWidget> Widget;
+};
+
+/**
+ * Tag to indicate that the entire row needs to be deleted when the widget in
+ * FTypedElementSlateWidgetReferenceColumn is no longer valid, otherwise only
+ * the column will be removed.
+ */
+USTRUCT()
+struct FTypedElementSlateWidgetReferenceDeletesRowTag : public FTypedElementDataStorageTag
+{
+	GENERATED_BODY()
+};
\ No newline at end of file
diff --git a/Engine/Source/Runtime/TypedElementFramework/Public/Elements/Framework/TypedElementRegistry.h b/Engine/Source/Runtime/TypedElementFramework/Public/Elements/Framework/TypedElementRegistry.h
--- a/Engine/Source/Runtime/TypedElementFramework/Public/Elements/Framework/TypedElementRegistry.h
+++ b/Engine/Source/Runtime/TypedElementFramework/Public/Elements/Framework/TypedElementRegistry.h
@@ -19,6 +19,7 @@
 #include "Elements/Framework/TypedElementListFwd.h"
 #include "Elements/Interfaces/TypedElementDataStorageInterface.h"
 #include "Elements/Interfaces/TypedElementDataStorageCompatibilityInterface.h"
+#include "Elements/Interfaces/TypedElementDataStorageUiInterface.h"
 #include "HAL/CriticalSection.h"
 #include "HAL/PlatformCrt.h"
 #include "Logging/LogVerbosity.h"
@@ -85,6 +86,21 @@ public:
 	const ITypedElementDataStorageCompatibilityInterface* GetDataStorageCompatibility() const;
 	void SetDataStorageCompatibility(ITypedElementDataStorageCompatibilityInterface* Storage);
 
+	ITypedElementDataStorageUiInterface* GetMutableDataStorageUi();
+	const ITypedElementDataStorageUiInterface* GetDataStorageUi() const;
+	void SetDataStorageUi(ITypedElementDataStorageUiInterface* Storage);
+
+	bool AreDataStorageInterfacesSet() const;
+
+	/**
+	 * Event fired when all Data Storage Interfaces have been set. 
+	 */
+	DECLARE_MULTICAST_DELEGATE(FOnDataStorageInterfacesSet);
+	FOnDataStorageInterfacesSet& OnDataStorageInterfacesSet()
+	{
+		return OnDataStorageInterfacesSetDelegate;
+	}
+
 	/**
 	 * Event fired when references to one element should be replaced with a reference to a different element.
 	 */
@@ -666,6 +682,8 @@ private:
 		--DisableElementDestructionOnGCCount;
 	}
 
+	void CallDataStorageInterfacesSetDelegateIfNeeded();
+
 	mutable FRWLock RegisteredElementTypesRW;
 	TUniquePtr<FRegisteredElementType> RegisteredElementTypes[TypedHandleMaxTypeId - 1];
 	TSortedMap<FName, FTypedHandleTypeId, FDefaultAllocator, FNameFastLess> RegisteredElementTypesNameToId;
@@ -681,7 +699,9 @@ private:
 	FOnElementReplaced OnElementReplacedDelegate;
 	FOnElementUpdated OnElementUpdatedDelegate;
 	FSimpleMulticastDelegate OnProcessingDeferredElementsToDestroyDelegate;
-
+	
+	FOnDataStorageInterfacesSet OnDataStorageInterfacesSetDelegate;
 	ITypedElementDataStorageInterface* DataStorage = nullptr;
 	ITypedElementDataStorageCompatibilityInterface* DataStorageCompatibility = nullptr;
+	ITypedElementDataStorageUiInterface* DataStorageUi = nullptr;
 };
diff --git a/Engine/Source/Runtime/TypedElementFramework/Public/Elements/Interfaces/TypedElementDataStorageUiInterface.h b/Engine/Source/Runtime/TypedElementFramework/Public/Elements/Interfaces/TypedElementDataStorageUiInterface.h
new file mode 100644
--- /dev/null
+++ b/Engine/Source/Runtime/TypedElementFramework/Public/Elements/Interfaces/TypedElementDataStorageUiInterface.h
@@ -0,0 +1,132 @@
+// Copyright Epic Games, Inc. All Rights Reserved.
+
+#pragma once
+
+#include "Containers/ContainersFwd.h"
+#include "Elements/Framework/TypedElementColumnUtils.h"
+#include "Elements/Interfaces/TypedElementDataStorageInterface.h"
+#include "Templates/UnrealTypeTraits.h"
+#include "UObject/Interface.h"
+
+#include "TypedElementDataStorageUiInterface.generated.h"
+
+class ITypedElementDataStorageUiInterface;
+class SWidget;
+
+/**
+ * Base class used to construct Typed Element widgets with.
+ * It's recommended to expose any construction variables as properties so they
+ * can either be set by a user or set using the passed in arguments. The
+ * Arguments can be directly used in case complex operations need to be done that
+ * prevent automatically setting construction variables.
+ * See below for the options to register a constructor with the Data Storage. For
+ * either registration case a new instance/copy of the constructor is created so
+ * arguments can be safely applied.
+ */
+USTRUCT()
+struct TYPEDELEMENTFRAMEWORK_API FTypedElementWidgetConstructor
+{
+	GENERATED_BODY()
+
+public:
+	virtual ~FTypedElementWidgetConstructor() = default;
+
+	/**
+	 * Constructs the widget according to the provided information. Information is collected by calling
+	 * the below functions ApplyArguments, CreateWidget and AddColumns. It's recommeded to overload those
+	 * functions to build widgets according to a standard recipe and to reduce the amount of code needed.
+	 * If a complexer situation is called for this function can also be directly overwritten.
+	 */
+	virtual TSharedPtr<SWidget> Construct(
+		TypedElementRowHandle Row, /** The row the widget will be stored in. */
+		ITypedElementDataStorageInterface* DataStorage,
+		ITypedElementDataStorageUiInterface* DataStorageUi,
+		TConstArrayView<TypedElement::ColumnUtils::Argument> Arguments);
+
+protected:
+	/**
+	 * Uses the type system to apply the provided arguments to constructor's properties. Overwrite this 
+	 * function if there are non-properties that can be configured through the arguments or if the configuration
+	 * contains arguments that are too complex to initialize (fully) through the type system.
+	 */
+	virtual void ApplyArguments(TConstArrayView<TypedElement::ColumnUtils::Argument> Arguments);
+	/** Create a new instance of the target widget. This is a required function. */
+	virtual TSharedPtr<SWidget> CreateWidget() PURE_VIRTUAL(FTypedElementWidgetConstructor::CreateWidget, return nullptr; );
+	/**
+	 * Add any columns to the provided row in the data storage to support the widget. By default a column will be added
+	 * that keeps track of the lifecycle of the widget.
+	 */
+	virtual void AddColumns(ITypedElementDataStorageInterface* DataStorage, TypedElementRowHandle Row, const TSharedPtr<SWidget>& Widget);
+};
+
+template<>
+struct TStructOpsTypeTraits<FTypedElementWidgetConstructor> : public TStructOpsTypeTraitsBase2<FTypedElementWidgetConstructor>
+{
+	enum
+	{
+		WithPureVirtual = true,
+	};
+};
+
+UINTERFACE(MinimalAPI)
+class UTypedElementDataStorageUiInterface : public UInterface
+{
+	GENERATED_BODY()
+};
+
+class TYPEDELEMENTFRAMEWORK_API ITypedElementDataStorageUiInterface
+{
+	GENERATED_BODY()
+
+public:
+	using WidgetCreatedCallback = TFunctionRef<void(const TSharedRef<SWidget>& NewWidget, TypedElementRowHandle Row)>;
+
+	/** 
+	 * Registers a widget factory that will be called when the purpose it's registered under is requested.
+	 * This version registers a generic type. Construction using these are typically cheaper as they can avoid
+	 * copying the Constructor and take up less memory. The downside is that they can't store additional configuration
+	 * options.
+	 */
+	virtual void RegisterWidgetFactory(FName Purpose, const UScriptStruct* Constructor) = 0;
+	
+	/**
+	 * Registers a widget factory that will be called when the purpose it's registered under is requested.
+	 * This version uses a pre-created instance of the Constructor. The benefit of this is that it store
+	 * configuration options. The downside is that it takes up more memory and requires copying when it's
+	 * used.
+	 */
+	template<typename ConstructorType>
+	void RegisterWidgetFactory(FName Purpose, TUniquePtr<ConstructorType>&& Constructor);
+	
+	/**
+	 * Creates all the widgets registered under the provided name. This may be a large number of widgets for a wide name
+	 * or exactly one when the exact name of the widget is registered. Arguments can be provided, but widgets are free
+	 * to ignore them.
+	 */
+	virtual void ConstructWidgets(FName Purpose, TConstArrayView<TypedElement::ColumnUtils::Argument> Arguments, 
+		const WidgetCreatedCallback& ConstructionCallback) = 0;
+
+	/** Creates a single widget using the provided constructor. Arguments can optionally be used to intialize the constructor. */
+	virtual TSharedPtr<SWidget> ConstructWidget(TypedElementRowHandle Row, FTypedElementWidgetConstructor& Constructor,
+		TConstArrayView<TypedElement::ColumnUtils::Argument> Arguments) = 0;
+
+protected:
+	/**
+	 * Registers a widget factory that will be called when the purpose it's registered under is requested.
+	 * This version uses a pre-created instance of the Constructor. The benefit of this is that it store
+	 * configuration options. The downside is that it takes up more memory and requires copying when it's
+	 * used.
+	 */
+	virtual void RegisterWidgetFactory(FName Purpose, const UScriptStruct* ConstructorType,
+		TUniquePtr<FTypedElementWidgetConstructor>&& Constructor) = 0;
+};
+
+
+// Implementations
+template<typename ConstructorType>
+void ITypedElementDataStorageUiInterface::RegisterWidgetFactory(FName Purpose, TUniquePtr<ConstructorType>&& Constructor)
+{
+	static_assert(TIsDerivedFrom<ConstructorType, FTypedElementWidgetConstructor>::Value, 
+		"Provided Typed Element Widget Constructor doesn't derive from FTypedElementWidgetConstructor.");
+	RegisterWidgetFactory(Purpose, ConstructorType::StaticStruct(), MoveTemp(Constructor));
+}
\ No newline at end of file
