It's now possible to register callbacks with a query for the Typed Elements Data Storage.

The Select query now accepts a callback which will be called every tick if there are any column combinations that match the query. Alternatively the callback can be triggered whenever a specific column is added or removed from a table. Information such as the required columns and dependencies are extracted automatically from the function signature.

diff --git a/Engine/Plugins/Experimental/TypedElementsDataStorage/Source/TypedElementsDataStorage/Private/Processors/TypedElementProcessorAdaptors.cpp b/Engine/Plugins/Experimental/TypedElementsDataStorage/Source/TypedElementsDataStorage/Private/Processors/TypedElementProcessorAdaptors.cpp
new file mode 100644
--- /dev/null
+++ b/Engine/Plugins/Experimental/TypedElementsDataStorage/Source/TypedElementsDataStorage/Private/Processors/TypedElementProcessorAdaptors.cpp
@@ -0,0 +1,286 @@
+// Copyright Epic Games, Inc. All Rights Reserved.
+
+#include "Processors/TypedElementProcessorAdaptors.h"
+
+#include "MassCommonTypes.h"
+#include "MassExecutionContext.h"
+#include "TypedElementDatabase.h"
+
+struct FMassContextForwarder final : public ITypedElementDataStorageInterface::FQueryContext
+{
+	explicit FMassContextForwarder(FMassExecutionContext& Context)
+		: Context(Context)
+		, RowCount(Context.GetNumEntities())
+	{}
+
+	const void* GetColumn(const UScriptStruct* ColumnType) const override
+	{
+		return Context.GetFragmentView(ColumnType).GetData();
+	}
+
+	void* GetMutableColumn(const UScriptStruct* ColumnType) override
+	{
+		return Context.GetMutableFragmentView(ColumnType).GetData();
+	}
+
+	void GetColumns(TArrayView<char*> RetrievedAddresses, TConstArrayView<const UScriptStruct*> ColumnTypes,
+		TConstArrayView<ITypedElementDataStorageInterface::EQueryAccessType> AccessTypes) override
+	{
+		checkf(RetrievedAddresses.Num() == ColumnTypes.Num(), TEXT("Unable to retrieve a batch of columns as the number of addresses "
+			"doesn't match the number of requested column."));
+		checkf(RetrievedAddresses.Num() == AccessTypes.Num(), TEXT("Unable to retrieve a batch of columns as the number of addresses "
+			"doesn't match the number of access types."));
+		
+		GetColumnsUnguarded(ColumnTypes.Num(), RetrievedAddresses.GetData(), ColumnTypes.GetData(), AccessTypes.GetData());
+	}
+
+	void GetColumnsUnguarded(int32 TypeCount, char** RetrievedAddresses, const UScriptStruct* const* ColumnTypes,
+		const ITypedElementDataStorageInterface::EQueryAccessType* AccessTypes) override
+	{
+		for (int32 Index = 0; Index < TypeCount; ++Index)
+		{
+			*RetrievedAddresses = *AccessTypes == ITypedElementDataStorageInterface::EQueryAccessType::ReadOnly
+				? const_cast<char*>(reinterpret_cast<const char*>(Context.GetFragmentView(*ColumnTypes).GetData()))
+				: reinterpret_cast<char*>(Context.GetMutableFragmentView(*ColumnTypes).GetData());
+
+			++RetrievedAddresses;
+			++ColumnTypes;
+			++AccessTypes;
+		}
+	}
+
+	USubsystem* GetMutableSubsystem(const TSubclassOf<USubsystem> SubsystemClass) override
+	{
+		return Context.GetMutableSubsystem<USubsystem>(SubsystemClass);
+	}
+
+	const USubsystem* GetSubsystem(const TSubclassOf<USubsystem> SubsystemClass) override
+	{
+		return Context.GetSubsystem<USubsystem>(SubsystemClass);
+	}
+
+	void GetSubsystems(TArrayView<char*> RetrievedAddresses, TConstArrayView<const TSubclassOf<USubsystem>> SubsystemTypes, 
+		TConstArrayView<ITypedElementDataStorageInterface::EQueryAccessType> AccessTypes) override
+	{
+		checkf(RetrievedAddresses.Num() == SubsystemTypes.Num(), TEXT("Unable to retrieve a batch of subsystem as the number of addresses "
+			"doesn't match the number of requested subsystem types."));
+
+		GetSubsystemsUnguarded(RetrievedAddresses.Num(), RetrievedAddresses.GetData(), SubsystemTypes.GetData(), AccessTypes.GetData());
+	}
+
+	void GetSubsystemsUnguarded(int32 SubsystemCount, char** RetrievedAddresses, const TSubclassOf<USubsystem>* SubsystemTypes,
+		const ITypedElementDataStorageInterface::EQueryAccessType* AccessTypes)
+	{
+		for (int32 Index = 0; Index < SubsystemCount; ++Index)
+		{
+			*RetrievedAddresses = *AccessTypes == ITypedElementDataStorageInterface::EQueryAccessType::ReadOnly
+				? const_cast<char*>(reinterpret_cast<const char*>(Context.GetSubsystem<USubsystem>(*SubsystemTypes)))
+				: reinterpret_cast<char*>(Context.GetMutableSubsystem<USubsystem>(*SubsystemTypes));
+
+			++RetrievedAddresses;
+			++SubsystemTypes;
+			++AccessTypes;
+		}
+	}
+
+	uint32 GetRowCount() const override
+	{
+		return RowCount;
+	}
+
+	FMassExecutionContext& Context;
+	uint32 RowCount{ 0 };
+};
+
+
+/**
+ * FTypedElementQueryProcessorData
+ */
+FTypedElementQueryProcessorData::FTypedElementQueryProcessorData(UMassProcessor& Owner)
+	: Query(Owner)
+{
+
+}
+
+EMassProcessingPhase FTypedElementQueryProcessorData::MapToMassProcessingPhase(ITypedElementDataStorageInterface::EQueryTickPhase Phase) const
+{
+	switch(Phase)
+	{
+	case ITypedElementDataStorageInterface::EQueryTickPhase::PrePhysics:
+		return EMassProcessingPhase::PrePhysics;
+	case ITypedElementDataStorageInterface::EQueryTickPhase::DuringPhysics:
+		return EMassProcessingPhase::DuringPhysics;
+	case ITypedElementDataStorageInterface::EQueryTickPhase::PostPhysics:
+		return EMassProcessingPhase::PostPhysics;
+	case ITypedElementDataStorageInterface::EQueryTickPhase::FrameEnd:
+		return EMassProcessingPhase::FrameEnd;
+	default:
+		checkf(false, TEXT("Query tick phase '%i' is unsupported."), static_cast<int>(Phase));
+		return EMassProcessingPhase::MAX;
+	};
+}
+
+FString FTypedElementQueryProcessorData::GetProcessorName(FString RootProcessorName) const
+{
+	if (ParentQuery)
+	{
+		FString Result = ParentQuery->Callback.Name.ToString();
+		Result += TEXT(" [");
+		Result += RootProcessorName;
+		Result += TEXT("]");
+		return Result;
+	}
+	else
+	{
+		return MoveTemp(RootProcessorName);
+	}
+}
+
+void FTypedElementQueryProcessorData::Execute(FMassEntityManager& EntityManager, FMassExecutionContext& Context)
+{
+	if (ParentQuery)
+	{
+		Query.ForEachEntityChunk(EntityManager, Context, [this](FMassExecutionContext& Context)
+			{
+				FMassContextForwarder QueryContext(Context);
+				ParentQuery->Callback.Function(QueryContext);
+			}
+		);
+	}
+}
+
+
+
+/**
+ * UTypedElementQueryProcessorCallbackAdapterProcessor
+ */
+
+UTypedElementQueryProcessorCallbackAdapterProcessor::UTypedElementQueryProcessorCallbackAdapterProcessor()
+	: Data(*this)
+{
+	bAllowMultipleInstances = true;
+}
+
+FMassEntityQuery& UTypedElementQueryProcessorCallbackAdapterProcessor::GetQuery()
+{
+	return Data.Query;
+}
+
+void UTypedElementQueryProcessorCallbackAdapterProcessor::ConfigureQueryCallback(FTypedElementDatabaseExtendedQuery& TargetParentQuery)
+{
+	Data.ParentQuery = &TargetParentQuery;
+
+	bRequiresGameThreadExecution = TargetParentQuery.Callback.bForceToGameThread;
+	ExecutionFlags = static_cast<int32>(EProcessorExecutionFlags::Editor); 
+	ExecutionOrder.ExecuteInGroup = TargetParentQuery.Callback.Group;
+	if (!TargetParentQuery.Callback.BeforeGroup.IsNone())
+	{
+		ExecutionOrder.ExecuteBefore.Add(TargetParentQuery.Callback.BeforeGroup);
+	}
+	if (!TargetParentQuery.Callback.AfterGroup.IsNone())
+	{
+		ExecutionOrder.ExecuteAfter.Add(TargetParentQuery.Callback.AfterGroup);
+	}
+	ProcessingPhase = Data.MapToMassProcessingPhase(TargetParentQuery.Callback.Phase);
+
+	Super::PostInitProperties();
+}
+
+void UTypedElementQueryProcessorCallbackAdapterProcessor::ConfigureQueries()
+{
+	// When the extended query information is provided the native query will already be fully configured.
+}
+
+void UTypedElementQueryProcessorCallbackAdapterProcessor::PostInitProperties()
+{
+	Super::Super::PostInitProperties();
+}
+
+FString UTypedElementQueryProcessorCallbackAdapterProcessor::GetProcessorName() const
+{
+	return Data.GetProcessorName(Super::GetProcessorName());
+}
+
+void UTypedElementQueryProcessorCallbackAdapterProcessor::Execute(FMassEntityManager& EntityManager, FMassExecutionContext& Context)
+{
+	Data.Execute(EntityManager, Context);
+}
+
+
+
+/**
+ * UTypedElementQueryObserverCallbackAdapterProcessor
+ */
+
+UTypedElementQueryObserverCallbackAdapterProcessor::UTypedElementQueryObserverCallbackAdapterProcessor()
+	: Data(*this)
+{
+	bAllowMultipleInstances = true;
+}
+
+FMassEntityQuery& UTypedElementQueryObserverCallbackAdapterProcessor::GetQuery()
+{
+	return Data.Query;
+}
+
+const UScriptStruct* UTypedElementQueryObserverCallbackAdapterProcessor::GetObservedType() const
+{
+	return ObservedType;
+}
+
+EMassObservedOperation UTypedElementQueryObserverCallbackAdapterProcessor::GetObservedOperation() const
+{
+	return Operation;
+}
+
+void UTypedElementQueryObserverCallbackAdapterProcessor::ConfigureQueryCallback(FTypedElementDatabaseExtendedQuery& TargetParentQuery)
+{
+	Data.ParentQuery = &TargetParentQuery;
+
+	bRequiresGameThreadExecution = TargetParentQuery.Callback.bForceToGameThread;
+	ExecutionFlags = static_cast<int32>(EProcessorExecutionFlags::Editor);
+	ProcessingPhase = Data.MapToMassProcessingPhase(TargetParentQuery.Callback.Phase);
+
+	ObservedType = const_cast<UScriptStruct*>(TargetParentQuery.Callback.MonitoredType);
+	
+	switch (TargetParentQuery.Callback.Type)
+	{
+	case ITypedElementDataStorageInterface::EQueryCallbackType::ObserveAdd:
+		Operation = EMassObservedOperation::Add;
+		break;
+	case ITypedElementDataStorageInterface::EQueryCallbackType::ObserveRemove:
+		Operation = EMassObservedOperation::Remove;
+		break;
+	default:
+		checkf(false, TEXT("Query type %i is not supported from the observer processor adapter."),
+			static_cast<int>(TargetParentQuery.Callback.Type));
+		break;
+	}
+
+	Super::PostInitProperties();
+}
+
+void UTypedElementQueryObserverCallbackAdapterProcessor::ConfigureQueries()
+{
+	// When the extended query information is provided the native query will already be fully configured.
+}
+
+void UTypedElementQueryObserverCallbackAdapterProcessor::PostInitProperties()
+{
+	Super::Super::PostInitProperties();
+}
+
+void UTypedElementQueryObserverCallbackAdapterProcessor::Register()
+{ 
+	// Do nothing as this processor will be explicitly registered.
+}
+
+FString UTypedElementQueryObserverCallbackAdapterProcessor::GetProcessorName() const
+{
+	return Data.GetProcessorName(Super::GetProcessorName());
+}
+
+void UTypedElementQueryObserverCallbackAdapterProcessor::Execute(FMassEntityManager& EntityManager, FMassExecutionContext& Context)
+{
+	Data.Execute(EntityManager, Context);
+}
diff --git a/Engine/Plugins/Experimental/TypedElementsDataStorage/Source/TypedElementsDataStorage/Private/Processors/TypedElementProcessorAdaptors.h b/Engine/Plugins/Experimental/TypedElementsDataStorage/Source/TypedElementsDataStorage/Private/Processors/TypedElementProcessorAdaptors.h
new file mode 100644
--- /dev/null
+++ b/Engine/Plugins/Experimental/TypedElementsDataStorage/Source/TypedElementsDataStorage/Private/Processors/TypedElementProcessorAdaptors.h
@@ -0,0 +1,84 @@
+// Copyright Epic Games, Inc. All Rights Reserved.
+
+#pragma once
+
+#include "Elements/Interfaces/TypedElementDataStorageInterface.h"
+#include "MassObserverProcessor.h"
+#include "MassProcessor.h"
+#include "UObject/ObjectMacros.h"
+
+#include "TypedElementProcessorAdaptors.generated.h"
+
+struct FTypedElementDatabaseExtendedQuery;
+
+USTRUCT()
+struct FTypedElementQueryProcessorData
+{
+	GENERATED_BODY()
+
+	FTypedElementQueryProcessorData() = default;
+	explicit FTypedElementQueryProcessorData(UMassProcessor& Owner);
+
+	EMassProcessingPhase MapToMassProcessingPhase(ITypedElementDataStorageInterface::EQueryTickPhase Phase) const;
+	FString GetProcessorName(FString RootProcessorName) const;
+
+	void Execute(FMassEntityManager& EntityManager, FMassExecutionContext& Context);
+
+	FTypedElementDatabaseExtendedQuery* ParentQuery{ nullptr };
+	FMassEntityQuery Query;
+};
+
+/**
+ * Adapts processor queries callback for MASS.
+ */
+UCLASS()
+class UTypedElementQueryProcessorCallbackAdapterProcessor : public UMassProcessor
+{
+	GENERATED_BODY()
+
+public:
+	UTypedElementQueryProcessorCallbackAdapterProcessor();
+
+	FMassEntityQuery& GetQuery();
+	void ConfigureQueryCallback(FTypedElementDatabaseExtendedQuery& Query);
+
+protected:
+	void ConfigureQueries() override;
+	void Execute(FMassEntityManager& EntityManager, FMassExecutionContext& TargetParentQuery) override;
+	
+	void PostInitProperties() override;
+	FString GetProcessorName() const override;
+
+private:
+	UPROPERTY(transient)
+	FTypedElementQueryProcessorData Data;
+};
+
+/**
+ * Adapts observer queries callback for MASS.
+ */
+UCLASS()
+class UTypedElementQueryObserverCallbackAdapterProcessor : public UMassObserverProcessor
+{
+	GENERATED_BODY()
+
+public:
+	UTypedElementQueryObserverCallbackAdapterProcessor();
+
+	FMassEntityQuery& GetQuery();
+	const UScriptStruct* GetObservedType() const;
+	EMassObservedOperation GetObservedOperation() const;
+	void ConfigureQueryCallback(FTypedElementDatabaseExtendedQuery& Query);
+
+protected:
+	void ConfigureQueries() override;
+	void Execute(FMassEntityManager& EntityManager, FMassExecutionContext& TargetParentQuery) override;
+
+	void PostInitProperties() override;
+	void Register() override;
+	FString GetProcessorName() const override;
+
+private:
+	UPROPERTY(transient)
+	FTypedElementQueryProcessorData Data;
+};
diff --git a/Engine/Plugins/Experimental/TypedElementsDataStorage/Source/TypedElementsDataStorage/Private/TypedElementDatabase.cpp b/Engine/Plugins/Experimental/TypedElementsDataStorage/Source/TypedElementsDataStorage/Private/TypedElementDatabase.cpp
--- a/Engine/Plugins/Experimental/TypedElementsDataStorage/Source/TypedElementsDataStorage/Private/TypedElementDatabase.cpp
+++ b/Engine/Plugins/Experimental/TypedElementsDataStorage/Source/TypedElementsDataStorage/Private/TypedElementDatabase.cpp
@@ -4,10 +4,13 @@
 
 #include "Editor.h"
 #include "Engine/World.h"
+#include "MassCommonTypes.h"
 #include "MassEntityEditorSubsystem.h"
+#include "MassSimulationSubsystem.h"
+#include "MassSubsystemAccess.h"
+#include "Processors/TypedElementProcessorAdaptors.h"
 #include "Stats/Stats2.h"
 #include "TickTaskManagerInterface.h"
-#include "MassSubsystemAccess.h"
 
 void UTypedElementDatabase::Initialize()
 {
@@ -249,28 +252,61 @@ ColumnDataResult UTypedElementDatabase::GetColumnData(TypedElementRowHandle Row,
 	return ColumnDataResult{ nullptr, nullptr };
 }
 
-TypedElementQueryHandle UTypedElementDatabase::RegisterQuery(const FQueryDescription& Query)
+TypedElementQueryHandle UTypedElementDatabase::RegisterQuery(FQueryDescription&& Query)
 {
-	auto LocalToNativeAccess = [](FQueryDescription::EAccessType Access) -> EMassFragmentAccess
+	auto LocalToNativeAccess = [](EQueryAccessType Access) -> EMassFragmentAccess
 	{
 		switch (Access)
 		{
-		case FQueryDescription::EAccessType::ReadOnly:
-			return EMassFragmentAccess::ReadOnly;
-		case FQueryDescription::EAccessType::ReadWrite:
-			return EMassFragmentAccess::ReadWrite;
-		default:
-			checkf(false, TEXT("Invalid query access type: %i."), static_cast<uint32>(Access));
-			return EMassFragmentAccess::MAX;
+			case EQueryAccessType::ReadOnly:
+				return EMassFragmentAccess::ReadOnly;
+			case EQueryAccessType::ReadWrite:
+				return EMassFragmentAccess::ReadWrite;
+			default:
+				checkf(false, TEXT("Invalid query access type: %i."), static_cast<uint32>(Access));
+				return EMassFragmentAccess::MAX;
 		}
 	};
 
+	auto SetupNativeQuery = [](FQueryDescription& Query, FTypedElementDatabaseExtendedQuery& StoredQuery) -> FMassEntityQuery&
+	{
+		if (Query.Action == FQueryDescription::EActionType::Select)
+		{
+			switch (Query.Callback.Type)
+			{
+				case EQueryCallbackType::None:
+					break;
+				case EQueryCallbackType::Processor:
+				{
+					UTypedElementQueryProcessorCallbackAdapterProcessor* Processor = 
+						NewObject<UTypedElementQueryProcessorCallbackAdapterProcessor>();
+					StoredQuery.Processor.Reset(Processor);
+					return Processor->GetQuery();
+				}
+				case EQueryCallbackType::ObserveAdd:
+					// Fallthrough
+				case EQueryCallbackType::ObserveRemove:
+				{
+					UTypedElementQueryObserverCallbackAdapterProcessor* Observer = 
+						NewObject<UTypedElementQueryObserverCallbackAdapterProcessor>();
+					StoredQuery.Processor.Reset(Observer);
+					return Observer->GetQuery();
+				}
+				default:
+					checkf(false, TEXT("Unsupported query callback type %i."), static_cast<int>(Query.Callback.Type));
+					break;
+			}
+		}
+		return StoredQuery.NativeQuery;
+	};
+
 	QueryStore::Handle Result = Queries.Emplace();
-	FExtendedQuery& StoredQuery = Queries.GetMutable(Result);
+	FTypedElementDatabaseExtendedQuery& StoredQuery = Queries.GetMutable(Result);
 	StoredQuery.Action = Query.Action;
 	StoredQuery.bSimpleQuery = Query.bSimpleQuery;
-
-	FMassEntityQuery& NativeQuery = StoredQuery.NativeQuery;
+	StoredQuery.Callback = MoveTemp(Query.Callback);
+	
+	FMassEntityQuery& NativeQuery = SetupNativeQuery(Query, StoredQuery);
 
 	if (Query.Action == FQueryDescription::EActionType::Count)
 	{
@@ -282,7 +318,8 @@ TypedElementQueryHandle UTypedElementDatabase::RegisterQuery(const FQueryDescrip
 		{
 			checkf(SelectEntry.Type, TEXT("Provided query selection type can not be null."));
 			checkf(SelectEntry.Type->IsChildOf(FTypedElementDataStorageColumn::StaticStruct()),
-				TEXT("Provided query selection type '%s' is not based on FTypedElementDataStorageColumn."), *SelectEntry.Type->GetStructPathName().ToString());
+				TEXT("Provided query selection type '%s' is not based on FTypedElementDataStorageColumn."), 
+				*SelectEntry.Type->GetStructPathName().ToString());
 			NativeQuery.AddRequirement(SelectEntry.Type, LocalToNativeAccess(SelectEntry.Access));
 		}
 	}
@@ -298,17 +335,17 @@ TypedElementQueryHandle UTypedElementDatabase::RegisterQuery(const FQueryDescrip
 			EMassFragmentPresence Presence;
 			switch (Type)
 			{
-			case FQueryDescription::EOperatorType::SimpleAll:
-				Presence = EMassFragmentPresence::All;
-				break;
-			case FQueryDescription::EOperatorType::SimpleAny:
-				Presence = EMassFragmentPresence::Any;
-				break;
-			case FQueryDescription::EOperatorType::SimpleNone:
-				Presence = EMassFragmentPresence::None;
-				break;
-			default:
-				continue;
+				case FQueryDescription::EOperatorType::SimpleAll:
+					Presence = EMassFragmentPresence::All;
+					break;
+				case FQueryDescription::EOperatorType::SimpleAny:
+					Presence = EMassFragmentPresence::Any;
+					break;
+				case FQueryDescription::EOperatorType::SimpleNone:
+					Presence = EMassFragmentPresence::None;
+					break;
+				default:
+					continue;
 			}
 
 			if (Operand->Type->IsChildOf(FMassTag::StaticStruct()))
@@ -327,11 +364,39 @@ TypedElementQueryHandle UTypedElementDatabase::RegisterQuery(const FQueryDescrip
 	for (const FQueryDescription::FAccessControlledClass& DependencyEntry : Query.Dependencies)
 	{
 		checkf(DependencyEntry.Type, TEXT("Provided query dependcy type can not be null."));
-		checkf(DependencyEntry.Type->IsChildOf<UWorldSubsystem>(),
-			TEXT("Provided query dependency type '%s' is not based on UWorldSubSystem."), *DependencyEntry.Type->GetStructPathName().ToString());
+		checkf(DependencyEntry.Type->IsChildOf<USubsystem>(),
+			TEXT("Provided query dependency type '%s' is not based on USubSystem."), 
+			*DependencyEntry.Type->GetStructPathName().ToString());
 		
 		constexpr bool bGameThreadOnly = true;
-		NativeQuery.AddSubsystemRequirement(const_cast<UClass*>(DependencyEntry.Type), LocalToNativeAccess(DependencyEntry.Access), bGameThreadOnly);
+		NativeQuery.AddSubsystemRequirement(
+			const_cast<UClass*>(DependencyEntry.Type), LocalToNativeAccess(DependencyEntry.Access), bGameThreadOnly);
+	}
+
+	if (StoredQuery.Processor)
+	{
+		if (StoredQuery.Processor->IsA<UTypedElementQueryProcessorCallbackAdapterProcessor>())
+		{
+			if (UMassEntityEditorSubsystem* Mass = GEditor->GetEditorSubsystem<UMassEntityEditorSubsystem>())
+			{
+				static_cast<UTypedElementQueryProcessorCallbackAdapterProcessor*>(StoredQuery.Processor.Get())->
+					ConfigureQueryCallback(StoredQuery);
+				Mass->RegisterDynamicProcessor(*StoredQuery.Processor);
+			}
+		}
+		else if (StoredQuery.Processor->IsA<UTypedElementQueryObserverCallbackAdapterProcessor>())
+		{
+			UTypedElementQueryObserverCallbackAdapterProcessor* Observer =
+				static_cast<UTypedElementQueryObserverCallbackAdapterProcessor*>(StoredQuery.Processor.Get());
+			Observer->ConfigureQueryCallback(StoredQuery);
+			ActiveEditorEntityManager->GetObserverManager().AddObserverInstance(
+				*Observer->GetObservedType(), Observer->GetObservedOperation(), *Observer);
+		}
+		else
+		{
+			checkf(false, TEXT("Query processor %s is of unsupported type %s."), 
+				*StoredQuery.Callback.Name.ToString(), *StoredQuery.Processor->GetSparseClassDataStruct()->GetName());
+		}
 	}
 
 	return Result.Handle;
@@ -344,12 +409,51 @@ void UTypedElementDatabase::UnregisterQuery(TypedElementQueryHandle Query)
 
 	if (Queries.IsAlive(Handle))
 	{
-		Queries.Get(Handle).NativeQuery.Clear();
+		FTypedElementDatabaseExtendedQuery& QueryData = Queries.Get(Handle);
+		if (QueryData.Processor)
+		{
+			if (QueryData.Processor->IsA<UTypedElementQueryProcessorCallbackAdapterProcessor>())
+			{
+				if (UMassEntityEditorSubsystem* Mass = GEditor->GetEditorSubsystem<UMassEntityEditorSubsystem>())
+				{
+					Mass->UnregisterDynamicProcessor(*QueryData.Processor);
+				}
+			}
+			else if (QueryData.Processor->IsA<UTypedElementQueryObserverCallbackAdapterProcessor>())
+			{
+				checkf(false, TEXT("Observer queries can not be unregistered."));
+			}
+			else
+			{
+				checkf(false, TEXT("Query processor %s is of unsupported type %s."),
+					*QueryData.Callback.Name.ToString(), *QueryData.Processor->GetSparseClassDataStruct()->GetName());
+			}
+		}
+		else
+		{
+			QueryData.NativeQuery.Clear();
+		}
 	}
 
 	Queries.Remove(Handle);
 }
 
+FName UTypedElementDatabase::GetQueryTickGroupName(EQueryTickGroups Group) const
+{
+	switch (Group)
+	{
+		case EQueryTickGroups::Default:
+			return NAME_None;
+		case EQueryTickGroups::SyncExternalToDataStorage:
+			return UE::Mass::ProcessorGroupNames::SyncWorldToMass;
+		case EQueryTickGroups::SyncDataStorageToExternal:
+			return UE::Mass::ProcessorGroupNames::UpdateWorldFromMass;
+		default:
+			checkf(false, TEXT("EQueryTickGroups value %i can't be translated to a group name by this Data Storage backend."), static_cast<int>(Group));
+			return NAME_None;
+	}
+}
+
 ITypedElementDataStorageInterface::FQueryResult UTypedElementDatabase::RunQuery(TypedElementQueryHandle Query)
 {
 	FQueryResult Result;
@@ -359,27 +463,27 @@ ITypedElementDataStorageInterface::FQueryResult UTypedElementDatabase::RunQuery(
 
 	if (Queries.IsAlive(Handle))
 	{
-		FExtendedQuery& QueryData = Queries.Get(Handle);
+		FTypedElementDatabaseExtendedQuery& QueryData = Queries.Get(Handle);
 		if (QueryData.bSimpleQuery)
 		{
 			switch (QueryData.Action)
 			{
-			case FQueryDescription::EActionType::None:
-				Result.Completed = FQueryResult::ECompletion::Fully;
-				break;
-			case FQueryDescription::EActionType::Select:
-				checkf(false, TEXT("Support for this option will be coming in a future update."));
-				Result.Completed = FQueryResult::ECompletion::Unsupported;
-				break;
-			case FQueryDescription::EActionType::Count:
-				checkf(ActiveEditorEntityManager, 
-					TEXT("Unable to run Typed Element Data Storage query before a MASS Entity Manager has been assigned."));
-				Result.Count = QueryData.NativeQuery.GetNumMatchingEntities(*ActiveEditorEntityManager);
-				Result.Completed = FQueryResult::ECompletion::Fully;
-				break;
-			default:
-				Result.Completed = FQueryResult::ECompletion::Unsupported;
-				break;
+				case FQueryDescription::EActionType::None:
+					Result.Completed = FQueryResult::ECompletion::Fully;
+					break;
+				case FQueryDescription::EActionType::Select:
+					checkf(false, TEXT("Support for this option will be coming in a future update."));
+					Result.Completed = FQueryResult::ECompletion::Unsupported;
+					break;
+				case FQueryDescription::EActionType::Count:
+					checkf(ActiveEditorEntityManager, 
+						TEXT("Unable to run Typed Element Data Storage query before a MASS Entity Manager has been assigned."));
+					Result.Count = QueryData.NativeQuery.GetNumMatchingEntities(*ActiveEditorEntityManager);
+					Result.Completed = FQueryResult::ECompletion::Fully;
+					break;
+				default:
+					Result.Completed = FQueryResult::ECompletion::Unsupported;
+					break;
 			}
 		}
 		else
diff --git a/Engine/Plugins/Experimental/TypedElementsDataStorage/Source/TypedElementsDataStorage/Private/TypedElementDatabase.h b/Engine/Plugins/Experimental/TypedElementsDataStorage/Source/TypedElementsDataStorage/Private/TypedElementDatabase.h
--- a/Engine/Plugins/Experimental/TypedElementsDataStorage/Source/TypedElementsDataStorage/Private/TypedElementDatabase.h
+++ b/Engine/Plugins/Experimental/TypedElementsDataStorage/Source/TypedElementsDataStorage/Private/TypedElementDatabase.h
@@ -14,12 +14,22 @@
 struct FMassEntityManager;
 class UWorld;
 
+struct FTypedElementDatabaseExtendedQuery
+{
+	FMassEntityQuery NativeQuery;
+	ITypedElementDataStorageInterface::FQueryDescription::EActionType Action{ ITypedElementDataStorageInterface::FQueryDescription::EActionType::None };
+	ITypedElementDataStorageInterface::FQueryDescription::FCallbackData Callback;
+	TStrongObjectPtr<UMassProcessor> Processor{ nullptr };
+	bool bSimpleQuery{ false };
+};
+
 UCLASS()
 class TYPEDELEMENTSDATASTORAGE_API UTypedElementDatabase 
 	: public UObject
 	, public ITypedElementDataStorageInterface
 {
 	GENERATED_BODY()
+
 public:
 	~UTypedElementDatabase() override = default;
 	
@@ -54,8 +64,9 @@ public:
 		TConstArrayView<TypedElement::ColumnUtils::Argument> Arguments) override;
 	ColumnDataResult GetColumnData(TypedElementRowHandle Row, FTopLevelAssetPath ColumnName) override;
 
-	TypedElementQueryHandle RegisterQuery(const FQueryDescription& Query) override;
+	TypedElementQueryHandle RegisterQuery(FQueryDescription&& Query) override;
 	void UnregisterQuery(TypedElementQueryHandle Query) override;
+	FName GetQueryTickGroupName(EQueryTickGroups Group) const override;
 	FQueryResult RunQuery(TypedElementQueryHandle Query) override;
 
 	FTypedElementOnDataStorageUpdate& OnUpdate() override;
@@ -63,13 +74,7 @@ public:
 	void* GetExternalSystemAddress(UClass* Target) override;
 
 private:
-	struct FExtendedQuery
-	{
-		FMassEntityQuery NativeQuery;
-		FQueryDescription::EActionType Action{ FQueryDescription::EActionType::None };
-		bool bSimpleQuery{ false };
-	};
-	using QueryStore = TTypedElementHandleStore<FExtendedQuery>;
+	using QueryStore = TTypedElementHandleStore<FTypedElementDatabaseExtendedQuery>;
 
 	void Reset();
 	
diff --git a/Engine/Plugins/Experimental/TypedElementsDataStorage/Source/TypedElementsDataStorage/TypedElementsDataStorage.Build.cs b/Engine/Plugins/Experimental/TypedElementsDataStorage/Source/TypedElementsDataStorage/TypedElementsDataStorage.Build.cs
--- a/Engine/Plugins/Experimental/TypedElementsDataStorage/Source/TypedElementsDataStorage/TypedElementsDataStorage.Build.cs
+++ b/Engine/Plugins/Experimental/TypedElementsDataStorage/Source/TypedElementsDataStorage/TypedElementsDataStorage.Build.cs
@@ -23,6 +23,7 @@ public class TypedElementsDataStorage : ModuleRules
 					"MassEntity",
 					"MassEntityEditor",
 					"MassGameplayEditor",
+					"MassSimulation",
 					"TypedElementFramework",
 					"EditorSubsystem",
 					"UnrealEd",
diff --git a/Engine/Source/Runtime/TypedElementFramework/Private/Elements/Framework/TypedElementQueryBuilder.cpp b/Engine/Source/Runtime/TypedElementFramework/Private/Elements/Framework/TypedElementQueryBuilder.cpp
--- a/Engine/Source/Runtime/TypedElementFramework/Private/Elements/Framework/TypedElementQueryBuilder.cpp
+++ b/Engine/Source/Runtime/TypedElementFramework/Private/Elements/Framework/TypedElementQueryBuilder.cpp
@@ -39,7 +39,7 @@ namespace TypedElementQueryBuilder
 	FDependency& FDependency::ReadOnly(const UClass* Target)
 	{
 		checkf(Target, TEXT("The Dependency section in the Typed Elements query builder doesn't support nullptrs as Read-Only input."));
-		Query->Dependencies.Emplace(Target, EAccessType::ReadOnly);
+		Query->Dependencies.Emplace(Target, ITypedElementDataStorageInterface::EQueryAccessType::ReadOnly);
 		return *this;
 	}
 
@@ -47,7 +47,7 @@ namespace TypedElementQueryBuilder
 	{
 		for (const UClass* Target : Targets)
 		{
-			ReadOnly(Targets);
+			ReadOnly(Target);
 		}
 		return *this;
 	}
@@ -55,7 +55,7 @@ namespace TypedElementQueryBuilder
 	FDependency& FDependency::ReadWrite(const UClass* Target)
 	{
 		checkf(Target, TEXT("The Dependency section in the Typed Elements query builder doesn't support nullptrs as Read/Write input."));
-		Query->Dependencies.Emplace(Target, EAccessType::ReadWrite);
+		Query->Dependencies.Emplace(Target, ITypedElementDataStorageInterface::EQueryAccessType::ReadWrite);
 		return *this;
 	}
 
@@ -63,7 +63,7 @@ namespace TypedElementQueryBuilder
 	{
 		for (const UClass* Target : Targets)
 		{
-			ReadWrite(Targets);
+			ReadWrite(Target);
 		}
 		return *this;
 	}
@@ -152,6 +152,125 @@ namespace TypedElementQueryBuilder
 	}
 
 
+	/**
+	 * FProcessor
+	 */
+	FProcessor::FProcessor(ITypedElementDataStorageInterface::EQueryTickPhase Phase, FName Group)
+		: Phase(Phase)
+		, Group(Group)
+	{}
+
+	FProcessor& FProcessor::SetPhase(ITypedElementDataStorageInterface::EQueryTickPhase NewPhase)
+	{
+		Phase = NewPhase;
+		return *this;
+	}
+
+	FProcessor& FProcessor::SetGroup(FName GroupName)
+	{
+		Group = GroupName;
+		return *this;
+	}
+
+	FProcessor& FProcessor::SetBeforeGroup(FName GroupName)
+	{
+		BeforeGroup = GroupName;
+		return *this;
+	}
+
+	FProcessor& FProcessor::SetAfterGroup(FName GroupName)
+	{
+		AfterGroup = GroupName;
+		return *this;
+	}
+
+	FProcessor& FProcessor::ForceToGameThread(bool bForce)
+	{
+		bForceToGameThread = bForce;
+		return *this;
+	}
+
+
+	/**
+	 * FObserver
+	 */
+	
+	FObserver::FObserver(EEvent MonitorForEvent, const UScriptStruct* MonitoredColumn)
+		: Monitor(MonitoredColumn)
+		, Event(MonitorForEvent)
+	{}
+
+	FObserver& FObserver::SetEvent(EEvent MonitorForEvent)
+	{
+		Event = MonitorForEvent;
+		return *this;
+	}
+
+	FObserver& FObserver::SetMonitoredColumn(const UScriptStruct* MonitoredColumn)
+	{
+		Monitor = MonitoredColumn;
+		return *this;
+	}
+
+	FObserver& FObserver::ForceToGameThread(bool bForce)
+	{
+		bForceToGameThread = bForce;
+		return *this;
+	}
+
+
+	/**
+	 * FQueryContextForwarder
+	 */
+
+	FQueryContextForwarder::FQueryContextForwarder(ITypedElementDataStorageInterface::FQueryContext& InParentContext)
+		: ParentContext(InParentContext)
+	{}
+
+	const void* FQueryContextForwarder::GetColumn(const UScriptStruct* ColumnType) const
+	{
+		return ParentContext.GetColumn(ColumnType);
+	}
+
+	void* FQueryContextForwarder::GetMutableColumn(const UScriptStruct* ColumnType)
+	{
+		return ParentContext.GetMutableColumn(ColumnType);
+	}
+
+	void FQueryContextForwarder::GetColumns(TArrayView<char*> RetrievedAddresses, TConstArrayView<const UScriptStruct*> ColumnTypes,
+		TConstArrayView<ITypedElementDataStorageInterface::EQueryAccessType> AccessTypes)
+	{
+		ParentContext.GetColumns(RetrievedAddresses, ColumnTypes, AccessTypes);
+	}
+
+	void FQueryContextForwarder::GetColumnsUnguarded(int32 TypeCount, char** RetrievedAddresses, const UScriptStruct* const* ColumnTypes,
+		const ITypedElementDataStorageInterface::EQueryAccessType* AccessTypes)
+	{
+		ParentContext.GetColumnsUnguarded(TypeCount, RetrievedAddresses, ColumnTypes, AccessTypes);
+	}
+
+	USubsystem* FQueryContextForwarder::GetMutableSubsystem(const TSubclassOf<USubsystem> SubsystemClass)
+	{
+		return ParentContext.GetMutableSubsystem(SubsystemClass);
+	}
+
+	const USubsystem* FQueryContextForwarder::GetSubsystem(const TSubclassOf<USubsystem> SubsystemClass)
+	{
+		return ParentContext.GetSubsystem(SubsystemClass);
+	}
+
+	void FQueryContextForwarder::GetSubsystems(TArrayView<char*> RetrievedAddresses, TConstArrayView<const TSubclassOf<USubsystem>> SubsystemTypes,
+		TConstArrayView<ITypedElementDataStorageInterface::EQueryAccessType> AccessTypes)
+	{
+		ParentContext.GetSubsystems(RetrievedAddresses, SubsystemTypes, AccessTypes);
+	}
+
+	uint32 FQueryContextForwarder::GetRowCount() const
+	{
+		return ParentContext.GetRowCount();
+	}
+
+
 	/**
 	 * Select
 	 */
@@ -164,7 +283,7 @@ namespace TypedElementQueryBuilder
 	Select& Select::ReadOnly(const UScriptStruct* Target)
 	{
 		checkf(Target, TEXT("The Select section in the Typed Elements query builder doesn't support nullptrs as Read-Only input."));
-		Query.Selection.Emplace(Target, EAccessType::ReadOnly);
+		Query.Selection.Emplace(Target, ITypedElementDataStorageInterface::EQueryAccessType::ReadOnly);
 		return *this;
 	}
 
@@ -180,7 +299,7 @@ namespace TypedElementQueryBuilder
 	Select& Select::ReadWrite(const UScriptStruct* Target)
 	{
 		checkf(Target, TEXT("The Select section in the Typed Elements query builder doesn't support nullptrs as Read/Write input."));
-		Query.Selection.Emplace(Target, EAccessType::ReadWrite);
+		Query.Selection.Emplace(Target, ITypedElementDataStorageInterface::EQueryAccessType::ReadWrite);
 		return *this;
 	}
 
diff --git a/Engine/Source/Runtime/TypedElementFramework/Public/Elements/Framework/TypedElementQueryBuilder.h b/Engine/Source/Runtime/TypedElementFramework/Public/Elements/Framework/TypedElementQueryBuilder.h
--- a/Engine/Source/Runtime/TypedElementFramework/Public/Elements/Framework/TypedElementQueryBuilder.h
+++ b/Engine/Source/Runtime/TypedElementFramework/Public/Elements/Framework/TypedElementQueryBuilder.h
@@ -4,8 +4,10 @@
 
 #include "CoreMinimal.h"
 #include "Elements/Interfaces/TypedElementDataStorageInterface.h"
+#include "Templates/SubclassOf.h"
 
 class UScriptStruct;
+class USubsystem;
 
 /**
  * The TypedElementQueryBuilder allows for the construction of queries for use by the Typed Element Data Storage.
@@ -30,6 +32,24 @@ class UScriptStruct;
  *		"/Script/OptionalPackage.FStructOptional"_TypeOptional
  * All functions allow for a single type to be added or a list of types, e.g. ReadOnly(Type<FStructExample>() or
  *		ReadOnly({ Type<FStructExample1>(), FStructExample2::StaticStruct(), "/Script/ExamplePackage.FStructExample3"_Type });
+ *
+ * Some functions allow binding to a callback. In these cases the arguments to the provided callback are analyzed and
+ * added to the query automatically. Const arguments are added as ReadOnly, while non-const arguments are added as 
+ * ReadWrite. Callbacks can be periodically called if constructed as a processor, in which case the callback is
+ * repeatedly, usually once per frame, called for all row (ranges) that match the query. If constructed as an observer
+ * the provided target type is monitored for actions like addition or deletion into/from any table and will trigger the
+ * callback once if the query matches. The following function signatures are excepted by "Select":
+ *		- void(ITypedElementDataStorageInterface::FQueryContext&, [const]Column&...) 
+ *			- The provided callback will be called for each row that matches the provided columns.
+ *		- void(FCachedQueryContext<[const]Dependencies...>&, [const]Column&...)
+ *			- Same as the previous but each listed dependency is registered and cached to improve performance. Note that
+ *				dependencies marked as const will only be available for reading and otherwise are can be accessed for read/write.
+ * Usage example:
+ *		FProcessor Info(
+ *			ITypedElementDataStorageInterface::EQueryTickPhase::FrameEnd, 
+ *			DataStorage->GetQueryTickGroupName(ITypedElementDataStorageInterface::EQueryTickGroups::SyncExternalToDataStorage);
+ *		Query = Select(FName(TEXT("Example Callback")), Info, 
+ *				[](FCachedQueryContext<Subsystem1, const Subsystem2>&, const FDataExample1&, FDataExample2&) {});
  * 
  * "Select" is constructed with: 
  * - ReadOnly: Indicates that the data object will only be read from
@@ -48,7 +68,7 @@ class UScriptStruct;
  * "DependsOn" is constructed with:
  * - ReadOnly: Indicates that the external system will only be used to read data from.
  * - ReadWrite: Indicates that the external system will be used to write data to.
- * 
+ *
  * Usage example:
  * ITypedElementDataStorageInterface::FQueryDescription Query =
  *		Select()
@@ -63,15 +83,13 @@ class UScriptStruct;
  *			.ReadOnly<USystemExample1, USystemExample2>()
  *			.ReadWrite(USystemExample2::StaticClass())
  *		.Compile();
- * 
+ *
  * Creating a query is expensive on the builder and the backend side. It's therefore recommended to create a query
  * and store its compiled form for repeated use instead of rebuilding the query on every update.
  */
 
 namespace TypedElementQueryBuilder
 {
-	using EAccessType = ITypedElementDataStorageInterface::FQueryDescription::EAccessType;
-
 	TYPEDELEMENTFRAMEWORK_API const UScriptStruct* Type(FTopLevelAssetPath Name);
 	TYPEDELEMENTFRAMEWORK_API const UScriptStruct* TypeOptional(FTopLevelAssetPath Name);
 	TYPEDELEMENTFRAMEWORK_API const UScriptStruct* operator""_Type(const char* Name, std::size_t NameSize);
@@ -96,7 +114,7 @@ namespace TypedElementQueryBuilder
 
 	private:
 		explicit FDependency(ITypedElementDataStorageInterface::FQueryDescription* Query);
-		
+
 		ITypedElementDataStorageInterface::FQueryDescription* Query;
 	};
 
@@ -128,12 +146,97 @@ namespace TypedElementQueryBuilder
 		ITypedElementDataStorageInterface::FQueryDescription* Query;
 	};
 
-	// Explicitly not following the naming convention in order to keep readability consistent. It now reads like a query sentence.
+	struct TYPEDELEMENTFRAMEWORK_API FQueryCallbackType{};
+
+	struct TYPEDELEMENTFRAMEWORK_API FProcessor final : public FQueryCallbackType
+	{
+		FProcessor(ITypedElementDataStorageInterface::EQueryTickPhase Phase, FName Group);
+		FProcessor& SetPhase(ITypedElementDataStorageInterface::EQueryTickPhase NewPhase);
+		FProcessor& SetGroup(FName GroupName);
+		FProcessor& SetBeforeGroup(FName GroupName);
+		FProcessor& SetAfterGroup(FName GroupName);
+		FProcessor& ForceToGameThread(bool bForce);
+		
+		ITypedElementDataStorageInterface::EQueryTickPhase Phase;
+		FName Group;
+		FName BeforeGroup;
+		FName AfterGroup;
+		bool bForceToGameThread{ false };
+	};
+
+	struct TYPEDELEMENTFRAMEWORK_API FObserver final : public FQueryCallbackType
+	{
+		enum class EEvent : uint8
+		{
+			Add,
+			Remove
+		};
+
+		template<typename ColumnType>
+		FObserver(EEvent MonitorForEvent);
+		FObserver(EEvent MonitorForEvent, const UScriptStruct* MonitoredColumn);
+
+		FObserver& SetEvent(EEvent MonitorForEvent);
+		FObserver& SetMonitoredColumn(const UScriptStruct* MonitoredColumn);
+		template<typename ColumnType>
+		FObserver& SetMonitoredColumn();
+		FObserver& ForceToGameThread(bool bForce);
+
+		const UScriptStruct* Monitor;
+		EEvent Event;
+		bool bForceToGameThread{ false };
+	};
+
+	struct TYPEDELEMENTFRAMEWORK_API FQueryContextForwarder : public ITypedElementDataStorageInterface::FQueryContext
+	{
+		explicit FQueryContextForwarder(ITypedElementDataStorageInterface::FQueryContext& InParentContext);
+		
+		const void* GetColumn(const UScriptStruct* ColumnType) const override;
+		void* GetMutableColumn(const UScriptStruct* ColumnType) override;
+		void GetColumns(TArrayView<char*> RetrievedAddresses, TConstArrayView<const UScriptStruct*> ColumnTypes,
+			TConstArrayView<ITypedElementDataStorageInterface::EQueryAccessType> AccessTypes) override;
+		void GetColumnsUnguarded(int32 TypeCount, char** RetrievedAddresses, const UScriptStruct* const* ColumnTypes, 
+			const ITypedElementDataStorageInterface::EQueryAccessType* AccessTypes) override;
+
+		USubsystem* GetMutableSubsystem(const TSubclassOf<USubsystem> SubsystemClass) override;
+		const USubsystem* GetSubsystem(const TSubclassOf<USubsystem> SubsystemClass) override;
+		void GetSubsystems(TArrayView<char*> RetrievedAddresses, TConstArrayView<const TSubclassOf<USubsystem>> SubsystemTypes, 
+			TConstArrayView<ITypedElementDataStorageInterface::EQueryAccessType> AccessTypes) override;
+
+		uint32 GetRowCount() const override;
+
+		ITypedElementDataStorageInterface::FQueryContext& ParentContext;
+	};
+
+	template<typename... Subsystems>
+	struct FCachedQueryContext final : public FQueryContextForwarder
+	{
+		explicit FCachedQueryContext(ITypedElementDataStorageInterface::FQueryContext& InParentContext);
+
+		void Fetch();
+		static void Register(ITypedElementDataStorageInterface::FQueryDescription& Query);
+
+		template<typename Subsystem>
+		Subsystem* GetCachedMutableSubsystem();
+		template<typename Subsystem>
+		const Subsystem* GetCachedSubsystem() const;
+
+		char* SubsystemAddresses[sizeof...(Subsystems)];
+		const TSubclassOf<USubsystem> SubsystemTypes[sizeof...(Subsystems)];
+		ITypedElementDataStorageInterface::EQueryAccessType SubsystemAccessList[sizeof...(Subsystems)];
+	};
+
+	// Explicitly not following the naming convention in order to present this as a query that can be read as such.
 	class TYPEDELEMENTFRAMEWORK_API Select final
 	{
 	public:
 		Select();
 
+		template<typename CallbackType, typename Function>
+		Select(FName Name, const CallbackType& Type, Function Callback);
+		template<typename CallbackType, typename Class, typename Function>
+		Select(FName Name, const CallbackType& Type, Class* Instance, Function Callback);
+
 		template<typename... TargetTypes>
 		Select& ReadOnly();
 		Select& ReadOnly(const UScriptStruct* Target);
@@ -163,70 +266,6 @@ namespace TypedElementQueryBuilder
 	private:
 		ITypedElementDataStorageInterface::FQueryDescription Query;
 	};
+} // namespace TypedElementQueryBuilder
 
-
-	// Implementations
-
-	//
-	// FDependecy
-	//
-
-	template<typename... TargetTypes>
-	FDependency& FDependency::ReadOnly()
-	{
-		ReadOnly({ TargetTypes::StaticClass()... });
-		return *this;
-	}
-
-	template<typename... TargetTypes>
-	FDependency& FDependency::ReadWrite()
-	{
-		ReadWrite({ TargetTypes::StaticClass()... });
-		return *this;
-	}
-
-
-	//
-	// Select
-	//
-
-	template<typename... TargetTypes>
-	Select& Select::ReadOnly()
-	{
-		ReadOnly({ TargetTypes::StaticStruct()... });
-		return *this;
-	}
-
-	template<typename... TargetTypes>
-	Select& Select::ReadWrite()
-	{
-		ReadWrite({ TargetTypes::StaticStruct()... });
-		return *this;
-	}
-
-
-	//
-	// FSimpleQuery
-	//
-
-	template<typename... TargetTypes>
-	FSimpleQuery& FSimpleQuery::All()
-	{
-		All({ TargetTypes::StaticStruct()... });
-		return *this;
-	}
-
-	template<typename... TargetTypes>
-	FSimpleQuery& FSimpleQuery::Any()
-	{
-		Any({ TargetTypes::StaticStruct()... });
-		return *this;
-	}
-
-	template<typename... TargetTypes>
-	FSimpleQuery& FSimpleQuery::None()
-	{
-		None({ TargetTypes::StaticStruct()... });
-		return *this;
-	}
-} // namespace TypedElementQueryBuilder
\ No newline at end of file
+#include "Elements/Framework/TypedElementQueryBuilder.inl"
\ No newline at end of file
diff --git a/Engine/Source/Runtime/TypedElementFramework/Public/Elements/Framework/TypedElementQueryBuilder.inl b/Engine/Source/Runtime/TypedElementFramework/Public/Elements/Framework/TypedElementQueryBuilder.inl
new file mode 100644
--- /dev/null
+++ b/Engine/Source/Runtime/TypedElementFramework/Public/Elements/Framework/TypedElementQueryBuilder.inl
@@ -0,0 +1,475 @@
+// Copyright Epic Games, Inc. All Rights Reserved.
+
+#include <type_traits>
+#include "Templates/IsConst.h"
+#include "Templates/UnrealTypeTraits.h"
+
+namespace TypedElementQueryBuilder
+{
+	namespace Internal
+	{
+		// This assumes that the types are unique, but for queries this should be true and otherwise
+		// both resuls would point to the first found index.
+		template<typename Target, typename ArgsCurrent, typename... ArgsRemainder>
+		constexpr uint32 GetVarArgIndex()
+		{
+			if constexpr (TIsSame<Target, ArgsCurrent>::Value)
+			{
+				return 0;
+			}
+			else
+			{
+				return 1 + GetVarArgIndex<Target, ArgsRemainder...>();
+			}
+		}
+
+		template<typename Type>
+		constexpr ITypedElementDataStorageInterface::EQueryAccessType GetAccessType()
+		{
+			using BaseType = typename TRemoveReference<Type>::Type;
+			if constexpr (TIsConst<BaseType>::Value)
+			{
+				return ITypedElementDataStorageInterface::EQueryAccessType::ReadOnly;
+			}
+			else
+			{
+				return ITypedElementDataStorageInterface::EQueryAccessType::ReadWrite;
+			}
+		}
+	}
+
+	//
+	// FDependecy
+	//
+
+	template<typename... TargetTypes>
+	FDependency& FDependency::ReadOnly()
+	{
+		ReadOnly({ TargetTypes::StaticClass()... });
+		return *this;
+	}
+
+	template<typename... TargetTypes>
+	FDependency& FDependency::ReadWrite()
+	{
+		ReadWrite({ TargetTypes::StaticClass()... });
+		return *this;
+	}
+
+	
+	//
+	// FObserver
+	//
+	
+	template<typename ColumnType>
+	FObserver::FObserver(EEvent MonitorForEvent)
+		: FObserver(MonitorForEvent, ColumnType::StaticStruct())
+	{}
+
+	template<typename ColumnType>
+	FObserver& FObserver::SetMonitoredColumn()
+	{
+		return SetMonitoredColumn(ColumnType::StaticStruct());
+	}
+
+
+	//
+	// FCachedQueryContext
+	//
+	
+	template<typename... Subsystems>
+	FCachedQueryContext<Subsystems...>::FCachedQueryContext(ITypedElementDataStorageInterface::FQueryContext& InParentContext)
+		: FQueryContextForwarder(InParentContext)
+		, SubsystemTypes{ Subsystems::StaticClass()... }
+		, SubsystemAccessList{ Internal::GetAccessType<Subsystems>()... }
+	{
+	}
+
+	template<typename... Subsystems>
+	void FCachedQueryContext<Subsystems...>::Fetch()
+	{
+		ParentContext.GetSubsystems(SubsystemAddresses, SubsystemTypes, SubsystemAccessList);
+#if DO_CHECK
+		char** SubsystemAddessIt = SubsystemAddresses;
+		for (int32 Index = 0; Index < sizeof...(Subsystems); ++Index)
+		{
+			checkf(*SubsystemAddessIt != nullptr, TEXT("The cached query context tried to retrieve a subsystem '%s' that didn't exist."),
+				SubsystemTypes[Index] != nullptr ? *(SubsystemTypes[Index]->GetName()) : TEXT("<no type information>"));
+			++SubsystemAddessIt;
+		}
+#endif // DO_CHECK
+	}
+	
+	template<typename... Subsystems>
+	void FCachedQueryContext<Subsystems...>::Register(ITypedElementDataStorageInterface::FQueryDescription& Query)
+	{
+		( Query.Dependencies.Emplace(Subsystems::StaticClass(), Internal::GetAccessType<Subsystems>()), ... );
+	}
+
+	template<typename... Subsystems>
+	template<typename Subsystem>
+	Subsystem* FCachedQueryContext<Subsystems...>::GetCachedMutableSubsystem()
+	{
+		// Don't allow a subsystem registered as const to be found.
+		constexpr uint32 Index = Internal::GetVarArgIndex<std::remove_const_t<Subsystem>, Subsystems...>();
+		static_assert(Index < sizeof...(Subsystems), "Requested subsystem isn't part of the query context cache.");
+		return reinterpret_cast<Subsystem*>(SubsystemAddresses[Index]);
+	}
+
+	template<typename... Subsystems>
+	template<typename Subsystem>
+	const Subsystem* FCachedQueryContext<Subsystems...>::GetCachedSubsystem() const
+	{
+		// Allow access to subsystems registered with and without const.
+		constexpr uint32 Index = Internal::GetVarArgIndex<Subsystem, Subsystems...>();
+		if constexpr (Index < sizeof...(Subsystems))
+		{
+			return reinterpret_cast<Subsystem*>(SubsystemAddresses[Index]);
+		}
+		else
+		{
+			constexpr uint32 ConstIndex = Internal::GetVarArgIndex<std::add_const_t<Subsystem>, Subsystems...>();
+			static_assert(ConstIndex < sizeof...(Subsystems), "Requested subsystem isn't part of the query context cache.");
+			return reinterpret_cast<Subsystem*>(SubsystemAddresses[ConstIndex]);
+		}
+	}
+
+	
+	//
+	// Select
+	//
+	namespace Internal
+	{
+		using QueryContext = ITypedElementDataStorageInterface::FQueryContext;
+		using QueryAccess = ITypedElementDataStorageInterface::EQueryAccessType;
+
+#		define QUERY_BINDING_BODY(Caller)																					\
+			static const UScriptStruct* ColumnTypes[] = { Args::StaticStruct()... };										\
+			char* ColumnAddresses[sizeof...(Args)];																			\
+			static constexpr QueryAccess ColumnAccess[] = { GetAccessType<Args>()... };										\
+			Context.GetColumnsUnguarded(sizeof...(Args), ColumnAddresses, ColumnTypes, ColumnAccess);						\
+			uint32 ColumnSize[] = { sizeof(Args)... };																		\
+																															\
+			using BaseContextType = typename TRemoveReference<ContextType>::Type;											\
+			if constexpr(TIsSame<BaseContextType, QueryContext>::Value)														\
+			{																												\
+				const uint32 RowCount = Context.GetRowCount();																\
+				for (uint32 Index = 0; Index < RowCount; ++Index)															\
+				{																											\
+					Caller (Context, *reinterpret_cast<Args*>(ColumnAddresses[GetVarArgIndex<Args, Args...>()])...);		\
+					for (uint32 ColumnIndex = 0; ColumnIndex < sizeof...(Args); ++ColumnIndex)								\
+					{																										\
+						ColumnAddresses[ColumnIndex] += ColumnSize[ColumnIndex];											\
+					}																										\
+				}																											\
+			}																												\
+			else																											\
+			{																												\
+				BaseContextType CachedContext(Context);																		\
+				CachedContext.Fetch();																						\
+																															\
+				const uint32 RowCount = Context.GetRowCount();																\
+				for (uint32 Index = 0; Index < RowCount; ++Index)															\
+				{																											\
+					Caller(CachedContext, *reinterpret_cast<Args*>(ColumnAddresses[GetVarArgIndex<Args, Args...>()])...);	\
+					for (uint32 ColumnIndex = 0; ColumnIndex < sizeof...(Args); ++ColumnIndex)								\
+					{																										\
+						ColumnAddresses[ColumnIndex] += ColumnSize[ColumnIndex];											\
+					}																										\
+				}																											\
+			}
+
+		template<typename ContextType, typename... Args>
+		void BindQueryFunction(ITypedElementDataStorageInterface::QueryCallback& Function, 
+			void (*Callback)(ContextType&, Args&...))
+		{
+			Function = [Callback](QueryContext& Context)
+			{
+				QUERY_BINDING_BODY(Callback)
+			};
+		}
+
+		template<typename Class, typename ContextType, typename... Args>
+		void BindQueryFunction(ITypedElementDataStorageInterface::QueryCallback& Function, Class* Target, 
+			void (Class::* Callback)(ContextType&, Args&...))
+		{
+			Function = [Target, Callback](QueryContext& Context)
+			{
+				QUERY_BINDING_BODY((Target->*Callback))
+			};
+		}
+
+		template<typename Class, typename ContextType, typename... Args>
+		void BindQueryFunction(ITypedElementDataStorageInterface::QueryCallback& Function, Class* Target, 
+			void (Class::*Callback)(ContextType&, Args&...) const)
+		{
+			Function = [Target, Callback](QueryContext& Context)
+			{
+				QUERY_BINDING_BODY((Target->*Callback))
+			};
+		}
+
+		template<typename Functor, typename Class, typename ContextType, typename... Args>
+		void BindQueryFunction_Expand(ITypedElementDataStorageInterface::QueryCallback& Function, Functor CallbackObject, 
+			void (Class::* Callback)(ContextType& , Args&...) const)
+		{
+			Function = [CallbackObject = std::forward<Functor>(CallbackObject)](QueryContext& Context)
+			{
+				QUERY_BINDING_BODY(CallbackObject)
+			};
+		}
+
+		template<typename Functor, typename Class, typename ContextType, typename... Args>
+		void BindQueryFunction_Expand(ITypedElementDataStorageInterface::QueryCallback& Function, Functor CallbackObject, 
+			void (Class::*Callback)(ContextType&, Args&...))
+		{
+			Function = [CallbackObject = std::forward<Functor>(CallbackObject)](QueryContext& Context)
+			{
+				QUERY_BINDING_BODY(CallbackObject)
+			};
+		}
+
+		template<typename Functor>
+		void BindQueryFunction(ITypedElementDataStorageInterface::QueryCallback& Function, Functor Callback)
+		{
+			BindQueryFunction_Expand(Function, std::forward<Functor>(Callback), &Functor::operator());
+		}
+
+#		undef QUERY_BINDING_BODY
+		
+		template<typename Arg>
+		void AddArgumentToSelect(Select& Target)
+		{
+			using BaseType = typename TRemoveReference<Arg>::Type;
+			if constexpr (TIsConst<BaseType>::Value)
+			{
+				Target.ReadOnly(BaseType::StaticStruct());
+			}
+			else
+			{
+				Target.ReadWrite(BaseType::StaticStruct());
+			}
+		}
+
+		template<typename ContextType> void RegisterDependencies(ITypedElementDataStorageInterface::FQueryDescription& Query)
+		{
+			using BaseType = typename TRemoveReference<ContextType>::Type;
+			if constexpr (TIsDerivedFrom<BaseType, FQueryContextForwarder>::Value)
+			{
+				BaseType::Register(Query);
+			}
+		}
+
+		template<typename ContextType, typename... Args>
+		void RegisterFunctionArguments(ITypedElementDataStorageInterface::FQueryDescription& Query, Select& Target, void (*)(ContextType&, Args...))
+		{ 
+			RegisterDependencies<ContextType>(Query);
+			(AddArgumentToSelect<Args>(Target), ...); 
+		}
+
+		template<typename Class, typename ContextType, typename... Args>
+		void RegisterFunctionArguments(ITypedElementDataStorageInterface::FQueryDescription& Query, Select& Target, void (Class::*)(ContextType&, Args...))
+		{
+			RegisterDependencies<ContextType>(Query);
+			(AddArgumentToSelect<Args>(Target), ...); 
+		}
+
+		template<typename Class, typename ContextType, typename... Args>
+		void RegisterFunctionArguments(ITypedElementDataStorageInterface::FQueryDescription& Query, Select& Target, void (Class::*)(ContextType&, Args...) const)
+		{
+			RegisterDependencies<ContextType>(Query);
+			(AddArgumentToSelect<Args>(Target), ...);
+		}
+		
+		template<typename Functor>
+		void RegisterFunctionArguments(ITypedElementDataStorageInterface::FQueryDescription& Query, Select& Target, Functor)
+		{ 
+			RegisterFunctionArguments(Query, Target, &Functor::operator());
+		}
+
+		template <typename T>
+		struct HasStaticStructMethod
+		{
+		private:
+			template <typename C> static constexpr bool Check(decltype(&C::StaticStruct)) { return std::is_same_v<decltype(&C::StaticStruct), UScriptStruct*(*)()>; }
+			template <typename C> static constexpr bool Check(...) { return false; }
+
+		public:
+			static constexpr bool Value = Check<T>(nullptr);
+		};
+
+		template<typename Context>
+		constexpr bool IsValidContextType()
+		{
+			using BaseContextType = typename std::remove_const_t<typename TRemoveReference<Context>::Type>;
+			if constexpr (TIsSame<BaseContextType, QueryContext>::Value)
+			{
+				return true;
+			}
+			else if constexpr (TIsDerivedFrom<BaseContextType, FQueryContextForwarder>::Value)
+			{
+				return true;
+			}
+			else
+			{
+				return false;
+			}
+		}
+
+		template<typename Arg>
+		constexpr bool IsValidArgument()
+		{
+			if constexpr (TIsReferenceType<Arg>::Value)
+			{
+				using BaseType = typename std::remove_const_t<typename TRemoveReference<Arg>::Type>;
+				return HasStaticStructMethod<BaseType>::Value;
+			}
+			else
+			{
+				return false;
+			}
+		}
+
+		template <typename T>
+		class IsFunctor
+		{
+			template <typename C> static constexpr bool Check(decltype(&C::operator())) { return true; }
+			template <typename C> static constexpr bool Check(...) { return false; }
+
+		public:
+			static constexpr bool Value = Check<T>(nullptr);
+		};
+
+		template<typename T> struct IsValidSelectFunctionSignatureImpl
+		{ 
+			constexpr static bool Value = false; 
+		};
+
+		template<typename ContextType, typename... Args> struct IsValidSelectFunctionSignatureImpl<void (*)(ContextType&, Args...)>
+		{ 
+			constexpr static bool Value = IsValidContextType<ContextType>() && (IsValidArgument<Args>() && ...);
+		};
+		
+		template<typename Class, typename ContextType, typename... Args> struct IsValidSelectFunctionSignatureImpl<void (Class::*)(ContextType&, Args...)>
+		{ 
+			constexpr static bool Value = IsValidContextType<ContextType>() && (IsValidArgument<Args>() && ...);
+		};
+
+		template<typename Class, typename ContextType, typename... Args> struct IsValidSelectFunctionSignatureImpl<void (Class::*)(ContextType&, Args...) const>
+		{ 
+			constexpr static bool Value = IsValidContextType<ContextType>() && (IsValidArgument<Args>() && ...);
+		};
+
+		template<typename T> constexpr bool IsValidSelectFunctionSignature()
+		{ 
+			if constexpr (IsFunctor<T>::Value)
+			{
+				return IsValidSelectFunctionSignatureImpl<decltype(&T::operator())>::Value;
+			}
+			else
+			{
+				return IsValidSelectFunctionSignatureImpl<T>::Value;
+			}
+		};
+		
+		inline void PrepareForQueryBinding(ITypedElementDataStorageInterface::FQueryDescription& Query, const FProcessor& Processor)
+		{
+			Query.Callback.Type = ITypedElementDataStorageInterface::EQueryCallbackType::Processor;
+			Query.Callback.Phase = Processor.Phase;
+			Query.Callback.Group = Processor.Group;
+			Query.Callback.BeforeGroup = Processor.BeforeGroup;
+			Query.Callback.AfterGroup = Processor.AfterGroup;
+			Query.Callback.bForceToGameThread = Processor.bForceToGameThread;
+		}
+
+		inline void PrepareForQueryBinding(ITypedElementDataStorageInterface::FQueryDescription& Query, const FObserver& Observer)
+		{
+			switch (Observer.Event)
+			{
+			case FObserver::EEvent::Add:
+				Query.Callback.Type = ITypedElementDataStorageInterface::EQueryCallbackType::ObserveAdd;
+				break;
+			case FObserver::EEvent::Remove:
+				Query.Callback.Type = ITypedElementDataStorageInterface::EQueryCallbackType::ObserveRemove;
+				break;
+			}
+			Query.Callback.MonitoredType = Observer.Monitor;
+			Query.Callback.bForceToGameThread = Observer.bForceToGameThread;
+		}
+		
+		template<typename CallbackType, typename Function>
+		void PrepareForQueryBinding(Select& Target, ITypedElementDataStorageInterface::FQueryDescription& Query, FName Name, 
+			const CallbackType& Type, Function Callback)
+		{
+			static_assert(TIsDerivedFrom<CallbackType, FQueryCallbackType>::Value, "The callback type provided isn't one of the available "
+				"classes derived from FQueryCallbackType.");
+			static_assert(IsValidSelectFunctionSignature<Function>(),
+				R"(The function provided to the Query Builder's Select call wasn't invocable or doesn't contain a supported combination of arguments.
+The following options are supported:
+- void(ITypedElementDataStorageInterface::FQueryContext&, [const]Column&...) 
+        e.g. void(ITypedElementDataStorageInterface::FQueryContext& Context, ColumnType0& ColumnA, const ColumnType1& ColumnB)
+- void(FCachedQueryContext<[const]Dependencies...>&, [const]Column&...)
+		e.g. void(FCachedQueryContext<Subsystem1, const Subsystem2>& Context, ColumnType0& ColumnA, const ColumnType1& ColumnB)
+)");
+			RegisterFunctionArguments(Query, Target, Callback);
+			PrepareForQueryBinding(Query, Type);
+			Query.Callback.Name = Name;
+		}
+	}
+
+	template<typename CallbackType, typename Function>
+	Select::Select(FName Name, const CallbackType& Type, Function Callback)
+		: Select()
+	{
+		Internal::PrepareForQueryBinding(*this, Query, Name, Type, Callback);
+		Internal::BindQueryFunction(Query.Callback.Function, std::forward<Function>(Callback));
+	}
+
+	template<typename CallbackType, typename Class, typename Function>
+	Select::Select(FName Name, const CallbackType& Type, Class* Instance, Function Callback)
+		: Select()
+	{
+		Internal::PrepareForQueryBinding(*this, Query, Name, Type, Callback);
+		Internal::BindQueryFunction(Query.Callback.Function, Instance, Callback);
+	}
+
+	template<typename... TargetTypes>
+	Select& Select::ReadOnly()
+	{
+		ReadOnly({ TargetTypes::StaticStruct()... });
+		return *this;
+	}
+
+	template<typename... TargetTypes>
+	Select& Select::ReadWrite()
+	{
+		ReadWrite({ TargetTypes::StaticStruct()... });
+		return *this;
+	}
+
+
+	//
+	// FSimpleQuery
+	//
+
+	template<typename... TargetTypes>
+	FSimpleQuery& FSimpleQuery::All()
+	{
+		All({ TargetTypes::StaticStruct()... });
+		return *this;
+	}
+
+	template<typename... TargetTypes>
+	FSimpleQuery& FSimpleQuery::Any()
+	{
+		Any({ TargetTypes::StaticStruct()... });
+		return *this;
+	}
+
+	template<typename... TargetTypes>
+	FSimpleQuery& FSimpleQuery::None()
+	{
+		None({ TargetTypes::StaticStruct()... });
+		return *this;
+	}
+} // namespace TypedElementQueryBuilder
\ No newline at end of file
diff --git a/Engine/Source/Runtime/TypedElementFramework/Public/Elements/Interfaces/TypedElementDataStorageInterface.h b/Engine/Source/Runtime/TypedElementFramework/Public/Elements/Interfaces/TypedElementDataStorageInterface.h
--- a/Engine/Source/Runtime/TypedElementFramework/Public/Elements/Interfaces/TypedElementDataStorageInterface.h
+++ b/Engine/Source/Runtime/TypedElementFramework/Public/Elements/Interfaces/TypedElementDataStorageInterface.h
@@ -7,6 +7,7 @@
 #include "Delegates/Delegate.h"
 #include "Elements/Framework/TypedElementColumnUtils.h"
 #include "Math/NumericLimits.h"
+#include "Templates/Function.h"
 #include "UObject/Interface.h"
 #include "UObject/NameTypes.h"
 #include "UObject/ObjectMacros.h"
@@ -14,6 +15,7 @@
 #include "TypedElementDataStorageInterface.generated.h"
 
 class UClass;
+class USubsystem;
 class UScriptStruct;
 
 struct ColumnDataResult
@@ -143,6 +145,73 @@ public:
 	 * can be used directly in the processor to do additional filtering. This will however impact performance and it's 
 	 * therefore recommended to try to simplify the query first before relying on extended query filtering in a processor.
 	 */
+	
+	enum class EQueryTickPhase : uint8
+	{
+		PrePhysics, //< Queries are executed before physics simulation starts.
+		DuringPhysics, //< Queries that can be run in parallel with physics simulation work.
+		PostPhysics, //< Queries that need rigid body and cloth simulation to be completed before being executed.
+		FrameEnd, //< Catchall for queries demoted to the last possible moment.
+
+		MAX
+	};
+
+	enum class EQueryTickGroups : uint8
+	{
+		/** The standard group to run work in. */
+		Default,
+		/** 
+		 * The group for processors that need to sync data from external sources such as subsystems or the world into
+		 * the Data Storage. These typically run early in a phase.
+		 */
+		SyncExternalToDataStorage,
+		/**
+		 * The group for processors that need to sync data from the Data Storage to external sources such as subsystems 
+		 * or the world into. These typically run late in a phase.
+		 */
+		SyncDataStorageToExternal
+	};
+
+	enum class EQueryCallbackType : uint8
+	{
+		/** No callback provided. */
+		None,
+		/** The query will be run every tick if at least one row matches. */
+		Processor,
+		/** The query will be run when a row is added that matches the query. The first recorded column will be actively monitored for changes. */
+		ObserveAdd,
+		/** The query will be run when a row is removed that matches the query. The first recorded column will be actively monitored for changes. */
+		ObserveRemove
+	};
+
+	enum class EQueryAccessType
+	{ 
+		ReadOnly, 
+		ReadWrite 
+	};
+
+	/** Base class to be provided to query callbacks. */
+	struct FQueryContext
+	{
+		virtual ~FQueryContext() = default;
+		
+		virtual const void* GetColumn(const UScriptStruct* ColumnType) const = 0;
+		virtual void* GetMutableColumn(const UScriptStruct* ColumnType) = 0;
+		virtual void GetColumns(TArrayView<char*> RetrievedAddresses, TConstArrayView<const UScriptStruct*> ColumnTypes,
+			TConstArrayView<ITypedElementDataStorageInterface::EQueryAccessType> AccessTypes) = 0;
+		virtual void GetColumnsUnguarded(int32 TypeCount, char** RetrievedAddresses, const UScriptStruct* const* ColumnTypes, 
+			const ITypedElementDataStorageInterface::EQueryAccessType* AccessTypes) = 0;
+
+		virtual USubsystem* GetMutableSubsystem(const TSubclassOf<USubsystem> SubsystemClass) = 0;
+		virtual const USubsystem* GetSubsystem(const TSubclassOf<USubsystem> SubsystemClass) = 0;
+		virtual void GetSubsystems(TArrayView<char*> RetrievedAddresses, TConstArrayView<const TSubclassOf<USubsystem>> SubsystemTypes, 
+			TConstArrayView<EQueryAccessType> AccessTypes) = 0;
+
+		virtual uint32 GetRowCount() const = 0;
+	};
+
+	using QueryCallback = TFunction<void(FQueryContext&)>;
+
 	struct FQueryDescription final
 	{
 		enum class EActionType : uint8
@@ -153,7 +222,7 @@ public:
 		};
 
 		using OperatorIndex = int32;
-		enum class EOperatorType
+		enum class EOperatorType : uint16
 		{
 			SimpleAll,			// Unary: Type
 			SimpleAny,			// Unary: Type
@@ -178,26 +247,38 @@ public:
 			const UScriptStruct* Type;
 		};
 
-		enum class EAccessType { ReadOnly, ReadWrite };
-
 		struct FAccessControlledStruct
 		{
 			FAccessControlledStruct() = default;
-			inline FAccessControlledStruct(const UScriptStruct* Type, EAccessType Access) : Type(Type), Access(Access) {}
+			inline FAccessControlledStruct(const UScriptStruct* Type, EQueryAccessType Access) : Type(Type), Access(Access) {}
 
 			const UScriptStruct* Type;
-			EAccessType Access;
+			EQueryAccessType Access;
 		};
 
 		struct FAccessControlledClass
 		{
 			FAccessControlledClass() = default;
-			inline FAccessControlledClass(const UClass* Type, EAccessType Access) : Type(Type), Access(Access) {}
+			inline FAccessControlledClass(const UClass* Type, EQueryAccessType Access) : Type(Type), Access(Access) {}
 
 			const UClass* Type;
-			EAccessType Access;
+			EQueryAccessType Access;
 		};
 
+		struct FCallbackData
+		{
+			QueryCallback Function;
+			FName Name;
+			FName Group;
+			FName BeforeGroup;
+			FName AfterGroup;
+			const UScriptStruct* MonitoredType{ nullptr };
+			EQueryCallbackType Type{ EQueryCallbackType::None };
+			EQueryTickPhase Phase;
+			bool bForceToGameThread{ false };
+		};
+		
+		FCallbackData Callback;
 		TArray<FAccessControlledStruct, TInlineAllocator<16>> Selection;
 		TArray<EOperatorType, TInlineAllocator<32>> ConditionTypes;
 		TArray<FOperator, TInlineAllocator<32>> ConditionOperators;
@@ -206,8 +287,14 @@ public:
 		/** If true, this query only has simple operations and is guaranteed to be executed fully and at optimal performance. */
 		bool bSimpleQuery{ false };
 	};
-	virtual TypedElementQueryHandle RegisterQuery(const FQueryDescription& Query) = 0;
+	virtual TypedElementQueryHandle RegisterQuery(FQueryDescription&& Query) = 0;
 	virtual void UnregisterQuery(TypedElementQueryHandle Query) = 0;
+	/**
+	 * Tick groups for queries can be given any name and the Data Storage will figure out the order of execution based on found
+	 * dependencies. However keeping processors within the same query group can help promote better performance through parallelization.
+	 * Therefore a collection of common tick group names is provided to help create consistent tick group names.
+	 */
+	virtual FName GetQueryTickGroupName(EQueryTickGroups Group) const = 0;
 	struct FQueryResult
 	{
 		enum class ECompletion
@@ -279,6 +366,8 @@ public:
 	SystemType* GetExternalSystem();
 };
 
+
+
 // Implementations
 template<typename TagType>
 void ITypedElementDataStorageInterface::AddTag(TypedElementRowHandle Row)
@@ -310,4 +399,4 @@ template<typename SystemType>
 SystemType* ITypedElementDataStorageInterface::GetExternalSystem()
 {
 	return reinterpret_cast<SystemType*>(GetExternalSystemAddress(SystemType::StaticClass()));
-}
\ No newline at end of file
+}
