Query callback improvements to capabilitites and performance.

The following was updated in this change list:
- Query callbacks now support ranges by using pointers instead of references for the columns.
- Support for retrieving the row handle as an argument for query callbacks was added. Ranges are also supported.
- Several new tick groups have been added to support the updated query callbacks. Most of these allow for a finer control over the order of query callbacks.
- The Query Context has been updated to support removing rows and adding/removing columns to rows.
- Instead of void pointers, std::tuple is now used to store column information. This adds additional type safety, allowed the previous macro to be fully removed and generates slightly faster code in the inner most loop.
- Error reporting on incorrect query callbacks has been improved. It will now be more specific where a problem is.
- Fixed an issue where MASS fragments couldn't be used in queries. It's still prefered to use TEDS columns, but the option is there for interoperability, especially for actors.
- Removed an unnecessary additional column for the StandardActorWithTransformTable.
- The names of the MASS processors running the query callbacks have been updated to use simpler text like [Editor Processor] instead of the object's name.
- The DCO's of the MASS processors now no longer register themselves to be run automatically. This also removed the need for a check if those processors have a valid TEDS query.
- FQueryContextForwarder has been inlined. This allows the compiler to fully optimize away the cost of the indirection.
- FCachedQueryContext now fetches subsystems on construction.

diff --git a/Engine/Plugins/Experimental/TypedElementsDataStorage/Source/TypedElementsDataStorage/Private/Processors/TypedElementProcessorAdaptors.cpp b/Engine/Plugins/Experimental/TypedElementsDataStorage/Source/TypedElementsDataStorage/Private/Processors/TypedElementProcessorAdaptors.cpp
--- a/Engine/Plugins/Experimental/TypedElementsDataStorage/Source/TypedElementsDataStorage/Private/Processors/TypedElementProcessorAdaptors.cpp
+++ b/Engine/Plugins/Experimental/TypedElementsDataStorage/Source/TypedElementsDataStorage/Private/Processors/TypedElementProcessorAdaptors.cpp
@@ -2,6 +2,7 @@
 
 #include "Processors/TypedElementProcessorAdaptors.h"
 
+#include <utility>
 #include "MassCommonTypes.h"
 #include "MassExecutionContext.h"
 #include "TypedElementDatabase.h"
@@ -23,7 +24,7 @@ struct FMassContextForwarder final : public ITypedElementDataStorageInterface::F
 		return Context.GetMutableFragmentView(ColumnType).GetData();
 	}
 
-	void GetColumns(TArrayView<char*> RetrievedAddresses, TConstArrayView<const UScriptStruct*> ColumnTypes,
+	void GetColumns(TArrayView<char*> RetrievedAddresses, TConstArrayView<TWeakObjectPtr<const UScriptStruct>> ColumnTypes,
 		TConstArrayView<ITypedElementDataStorageInterface::EQueryAccessType> AccessTypes) override
 	{
 		checkf(RetrievedAddresses.Num() == ColumnTypes.Num(), TEXT("Unable to retrieve a batch of columns as the number of addresses "
@@ -34,14 +35,15 @@ struct FMassContextForwarder final : public ITypedElementDataStorageInterface::F
 		GetColumnsUnguarded(ColumnTypes.Num(), RetrievedAddresses.GetData(), ColumnTypes.GetData(), AccessTypes.GetData());
 	}
 
-	void GetColumnsUnguarded(int32 TypeCount, char** RetrievedAddresses, const UScriptStruct* const* ColumnTypes,
+	void GetColumnsUnguarded(int32 TypeCount, char** RetrievedAddresses, const TWeakObjectPtr<const UScriptStruct>* ColumnTypes,
 		const ITypedElementDataStorageInterface::EQueryAccessType* AccessTypes) override
 	{
 		for (int32 Index = 0; Index < TypeCount; ++Index)
 		{
+			checkf(ColumnTypes->IsValid(), TEXT("Attempting to retrieve a column that is not available."));
 			*RetrievedAddresses = *AccessTypes == ITypedElementDataStorageInterface::EQueryAccessType::ReadOnly
-				? const_cast<char*>(reinterpret_cast<const char*>(Context.GetFragmentView(*ColumnTypes).GetData()))
-				: reinterpret_cast<char*>(Context.GetMutableFragmentView(*ColumnTypes).GetData());
+				? const_cast<char*>(reinterpret_cast<const char*>(Context.GetFragmentView(ColumnTypes->Get()).GetData()))
+				: reinterpret_cast<char*>(Context.GetMutableFragmentView(ColumnTypes->Get()).GetData());
 
 			++RetrievedAddresses;
 			++ColumnTypes;
@@ -49,36 +51,37 @@ struct FMassContextForwarder final : public ITypedElementDataStorageInterface::F
 		}
 	}
 
-	USubsystem* GetMutableSubsystem(const TSubclassOf<USubsystem> SubsystemClass) override
+	UObject* GetMutableDependency(const UClass* DependencyClass) override
 	{
-		return Context.GetMutableSubsystem<USubsystem>(SubsystemClass);
+		return Context.GetMutableSubsystem<USubsystem>(const_cast<UClass*>(DependencyClass));
 	}
 
-	const USubsystem* GetSubsystem(const TSubclassOf<USubsystem> SubsystemClass) override
+	const UObject* GetDependency(const UClass* DependencyClass) override
 	{
-		return Context.GetSubsystem<USubsystem>(SubsystemClass);
+		return Context.GetSubsystem<USubsystem>(const_cast<UClass*>(DependencyClass));
 	}
 
-	void GetSubsystems(TArrayView<char*> RetrievedAddresses, TConstArrayView<const TSubclassOf<USubsystem>> SubsystemTypes, 
+	void GetDependencies(TArrayView<UObject*> RetrievedAddresses, TConstArrayView<TWeakObjectPtr<const UClass>> SubsystemTypes,
 		TConstArrayView<ITypedElementDataStorageInterface::EQueryAccessType> AccessTypes) override
 	{
 		checkf(RetrievedAddresses.Num() == SubsystemTypes.Num(), TEXT("Unable to retrieve a batch of subsystem as the number of addresses "
 			"doesn't match the number of requested subsystem types."));
 
-		GetSubsystemsUnguarded(RetrievedAddresses.Num(), RetrievedAddresses.GetData(), SubsystemTypes.GetData(), AccessTypes.GetData());
+		GetDependenciesUnguarded(RetrievedAddresses.Num(), RetrievedAddresses.GetData(), SubsystemTypes.GetData(), AccessTypes.GetData());
 	}
 
-	void GetSubsystemsUnguarded(int32 SubsystemCount, char** RetrievedAddresses, const TSubclassOf<USubsystem>* SubsystemTypes,
+	void GetDependenciesUnguarded(int32 SubsystemCount, UObject** RetrievedAddresses, const TWeakObjectPtr<const UClass>* DependencyTypes,
 		const ITypedElementDataStorageInterface::EQueryAccessType* AccessTypes)
 	{
 		for (int32 Index = 0; Index < SubsystemCount; ++Index)
 		{
+			checkf(DependencyTypes->IsValid(), TEXT("Attempting to retrieve a subsystem that's no longer valid."));
 			*RetrievedAddresses = *AccessTypes == ITypedElementDataStorageInterface::EQueryAccessType::ReadOnly
-				? const_cast<char*>(reinterpret_cast<const char*>(Context.GetSubsystem<USubsystem>(*SubsystemTypes)))
-				: reinterpret_cast<char*>(Context.GetMutableSubsystem<USubsystem>(*SubsystemTypes));
+				? const_cast<USubsystem*>(Context.GetSubsystem<USubsystem>(const_cast<UClass*>(DependencyTypes->Get())))
+				: Context.GetMutableSubsystem<USubsystem>(const_cast<UClass*>(DependencyTypes->Get()));
 
 			++RetrievedAddresses;
-			++SubsystemTypes;
+			++DependencyTypes;
 			++AccessTypes;
 		}
 	}
@@ -88,6 +91,100 @@ struct FMassContextForwarder final : public ITypedElementDataStorageInterface::F
 		return RowCount;
 	}
 
+	TConstArrayView<TypedElementRowHandle> GetRowHandles() const
+	{
+		static_assert(
+			sizeof(TypedElementRowHandle) == sizeof(FMassEntityHandle) && alignof(TypedElementRowHandle) == alignof(FMassEntityHandle),
+			"TypedElementRowHandle and FMassEntityHandle need to by layout compatible to support Typed Elements Data Storage.");
+		TConstArrayView<FMassEntityHandle> Entities = Context.GetEntities();
+		return TConstArrayView<TypedElementRowHandle>(reinterpret_cast<const TypedElementRowHandle*>(Entities.GetData()), Entities.Num());
+	}
+
+	void RemoveRow(TypedElementRowHandle Row) override
+	{
+		Context.Defer().DestroyEntity(FMassEntityHandle::FromNumber(Row));
+	}
+
+	void RemoveRows(TConstArrayView<TypedElementRowHandle> Rows) override
+	{
+		for (TypedElementRowHandle Row : Rows)
+		{
+			Context.Defer().DestroyEntity(FMassEntityHandle::FromNumber(Row));
+		}
+	}
+
+	void AddColumns(TypedElementRowHandle Row, TConstArrayView<const UScriptStruct*> ColumnTypes)
+	{
+		for (const UScriptStruct* ColumnType : ColumnTypes)
+		{
+			bool bIsTag = ColumnType->IsChildOf(FMassTag::StaticStruct());
+
+			checkf(ColumnType->IsChildOf(FMassFragment::StaticStruct()) || bIsTag,
+				TEXT("Given struct type to add is not a valid fragment or tag type."));
+
+			if (bIsTag)
+			{
+				Context.Defer().PushCommand<FMassDeferredAddCommand>(
+					[EntityHandle = FMassEntityHandle::FromNumber(Row), ColumnType](FMassEntityManager& System)
+					{
+						System.AddTagToEntity(EntityHandle, ColumnType);
+					});
+			}
+			else
+			{
+				Context.Defer().PushCommand<FMassDeferredAddCommand>(
+					[EntityHandle = FMassEntityHandle::FromNumber(Row), ColumnType](FMassEntityManager& System)
+					{
+						System.AddFragmentToEntity(EntityHandle, ColumnType);
+					});
+			}
+		}
+	}
+
+	void AddColumns(TConstArrayView<TypedElementRowHandle> Rows, TConstArrayView<const UScriptStruct*> ColumnTypes)
+	{
+		for (TypedElementRowHandle Row : Rows)
+		{
+			AddColumns(Row, ColumnTypes);
+		}
+	}
+
+	void RemoveColumns(TypedElementRowHandle Row, TConstArrayView<const UScriptStruct*> ColumnTypes)
+	{
+		for (const UScriptStruct* ColumnType : ColumnTypes)
+		{
+			bool bIsTag = ColumnType->IsChildOf(FMassTag::StaticStruct());
+
+			checkf(ColumnType->IsChildOf(FMassFragment::StaticStruct()) || bIsTag,
+				TEXT("Given struct type to remove is not a valid fragment or tag type."));
+
+			if (bIsTag)
+			{
+				Context.Defer().PushCommand<FMassDeferredAddCommand>(
+					[EntityHandle = FMassEntityHandle::FromNumber(Row), ColumnType](FMassEntityManager& System)
+					{
+						System.RemoveTagFromEntity(EntityHandle, ColumnType);
+					});
+			}
+			else
+			{
+				Context.Defer().PushCommand<FMassDeferredAddCommand>(
+					[EntityHandle = FMassEntityHandle::FromNumber(Row), ColumnType](FMassEntityManager& System)
+					{
+						System.RemoveFragmentFromEntity(EntityHandle, ColumnType);
+					});
+			}
+		}
+	}
+
+	void RemoveColumns(TConstArrayView<TypedElementRowHandle> Rows, TConstArrayView<const UScriptStruct*> ColumnTypes)
+	{
+		for (TypedElementRowHandle Row : Rows)
+		{
+			RemoveColumns(Row, ColumnTypes);
+		}
+	}
+
 	FMassExecutionContext& Context;
 	uint32 RowCount{ 0 };
 };
@@ -99,7 +196,6 @@ struct FMassContextForwarder final : public ITypedElementDataStorageInterface::F
 FTypedElementQueryProcessorData::FTypedElementQueryProcessorData(UMassProcessor& Owner)
 	: Query(Owner)
 {
-
 }
 
 EMassProcessingPhase FTypedElementQueryProcessorData::MapToMassProcessingPhase(ITypedElementDataStorageInterface::EQueryTickPhase Phase) const
@@ -120,33 +216,22 @@ EMassProcessingPhase FTypedElementQueryProcessorData::MapToMassProcessingPhase(I
 	};
 }
 
-FString FTypedElementQueryProcessorData::GetProcessorName(FString RootProcessorName) const
+FString FTypedElementQueryProcessorData::GetProcessorName() const
 {
-	if (ParentQuery)
-	{
-		FString Result = ParentQuery->Callback.Name.ToString();
-		Result += TEXT(" [");
-		Result += RootProcessorName;
-		Result += TEXT("]");
-		return Result;
-	}
-	else
-	{
-		return MoveTemp(RootProcessorName);
-	}
+	return ParentQuery ? ParentQuery->Callback.Name.ToString() : FString(TEXT("<unnamed>"));
 }
 
 void FTypedElementQueryProcessorData::Execute(FMassEntityManager& EntityManager, FMassExecutionContext& Context)
 {
-	if (ParentQuery)
-	{
-		Query.ForEachEntityChunk(EntityManager, Context, [this](FMassExecutionContext& Context)
-			{
-				FMassContextForwarder QueryContext(Context);
-				ParentQuery->Callback.Function(QueryContext);
-			}
-		);
-	}
+	checkf(ParentQuery, TEXT("A query callback was registered for execution without an associated query."));
+	
+	Query.ForEachEntityChunk(EntityManager, Context, 
+		[this](FMassExecutionContext& Context)
+		{
+			FMassContextForwarder QueryContext(Context);
+			ParentQuery->Callback.Function(QueryContext);
+		}
+	);
 }
 
 
@@ -159,6 +244,7 @@ UTypedElementQueryProcessorCallbackAdapterProcessor::UTypedElementQueryProcessor
 	: Data(*this)
 {
 	bAllowMultipleInstances = true;
+	bAutoRegisterWithProcessingPhases = false;
 }
 
 FMassEntityQuery& UTypedElementQueryProcessorCallbackAdapterProcessor::GetQuery()
@@ -198,7 +284,9 @@ void UTypedElementQueryProcessorCallbackAdapterProcessor::PostInitProperties()
 
 FString UTypedElementQueryProcessorCallbackAdapterProcessor::GetProcessorName() const
 {
-	return Data.GetProcessorName(Super::GetProcessorName());
+	FString Name = Data.GetProcessorName();
+	Name += TEXT(" [Editor Processor]");
+	return Name;
 }
 
 void UTypedElementQueryProcessorCallbackAdapterProcessor::Execute(FMassEntityManager& EntityManager, FMassExecutionContext& Context)
@@ -216,6 +304,7 @@ UTypedElementQueryObserverCallbackAdapterProcessor::UTypedElementQueryObserverCa
 	: Data(*this)
 {
 	bAllowMultipleInstances = true;
+	bAutoRegisterWithProcessingPhases = false;
 }
 
 FMassEntityQuery& UTypedElementQueryObserverCallbackAdapterProcessor::GetQuery()
@@ -239,8 +328,7 @@ void UTypedElementQueryObserverCallbackAdapterProcessor::ConfigureQueryCallback(
 
 	bRequiresGameThreadExecution = TargetParentQuery.Callback.bForceToGameThread;
 	ExecutionFlags = static_cast<int32>(EProcessorExecutionFlags::Editor);
-	ProcessingPhase = Data.MapToMassProcessingPhase(TargetParentQuery.Callback.Phase);
-
+	
 	ObservedType = const_cast<UScriptStruct*>(TargetParentQuery.Callback.MonitoredType);
 	
 	switch (TargetParentQuery.Callback.Type)
@@ -277,7 +365,21 @@ void UTypedElementQueryObserverCallbackAdapterProcessor::Register()
 
 FString UTypedElementQueryObserverCallbackAdapterProcessor::GetProcessorName() const
 {
-	return Data.GetProcessorName(Super::GetProcessorName());
+	FString Name = Data.GetProcessorName();
+	EMassObservedOperation ObservationType = GetObservedOperation();
+	if (ObservationType == EMassObservedOperation::Add)
+	{
+		Name += TEXT(" [Editor Add Observer]");
+	}
+	else if (ObservationType == EMassObservedOperation::Remove)
+	{
+		Name += TEXT(" [Editor Remove Observer]");
+	}
+	else
+	{
+		Name += TEXT(" [Editor <Unknown> Observer]");
+	}
+	return Name;
 }
 
 void UTypedElementQueryObserverCallbackAdapterProcessor::Execute(FMassEntityManager& EntityManager, FMassExecutionContext& Context)
diff --git a/Engine/Plugins/Experimental/TypedElementsDataStorage/Source/TypedElementsDataStorage/Private/Processors/TypedElementProcessorAdaptors.h b/Engine/Plugins/Experimental/TypedElementsDataStorage/Source/TypedElementsDataStorage/Private/Processors/TypedElementProcessorAdaptors.h
--- a/Engine/Plugins/Experimental/TypedElementsDataStorage/Source/TypedElementsDataStorage/Private/Processors/TypedElementProcessorAdaptors.h
+++ b/Engine/Plugins/Experimental/TypedElementsDataStorage/Source/TypedElementsDataStorage/Private/Processors/TypedElementProcessorAdaptors.h
@@ -20,7 +20,7 @@ struct FTypedElementQueryProcessorData
 	explicit FTypedElementQueryProcessorData(UMassProcessor& Owner);
 
 	EMassProcessingPhase MapToMassProcessingPhase(ITypedElementDataStorageInterface::EQueryTickPhase Phase) const;
-	FString GetProcessorName(FString RootProcessorName) const;
+	FString GetProcessorName() const;
 
 	void Execute(FMassEntityManager& EntityManager, FMassExecutionContext& Context);
 
diff --git a/Engine/Plugins/Experimental/TypedElementsDataStorage/Source/TypedElementsDataStorage/Private/TypedElementDatabase.cpp b/Engine/Plugins/Experimental/TypedElementsDataStorage/Source/TypedElementsDataStorage/Private/TypedElementDatabase.cpp
--- a/Engine/Plugins/Experimental/TypedElementsDataStorage/Source/TypedElementsDataStorage/Private/TypedElementDatabase.cpp
+++ b/Engine/Plugins/Experimental/TypedElementsDataStorage/Source/TypedElementsDataStorage/Private/TypedElementDatabase.cpp
@@ -12,6 +12,12 @@
 #include "Stats/Stats2.h"
 #include "TickTaskManagerInterface.h"
 
+const FName UTypedElementDatabase::TickGroupName_PrepareSyncWorldToMass(TEXT("PrepareSyncWorldToMass"));
+const FName UTypedElementDatabase::TickGroupName_FinalizeSyncWorldToMass(TEXT("FinalizeSyncWorldToMass"));
+const FName UTypedElementDatabase::TickGroupName_PrepareSyncMassToExternal(TEXT("PrepareSyncMassToExternal"));
+const FName UTypedElementDatabase::TickGroupName_FinalizeSyncMassToExternal(TEXT("FinalizeSyncMassToExternal"));
+const FName UTypedElementDatabase::TickGroupName_SyncWidget(TEXT("SyncWidgets"));
+
 void UTypedElementDatabase::Initialize()
 {
 	check(GEditor);
@@ -317,8 +323,11 @@ TypedElementQueryHandle UTypedElementDatabase::RegisterQuery(FQueryDescription&&
 		for (const FQueryDescription::FAccessControlledStruct& SelectEntry : Query.Selection)
 		{
 			checkf(SelectEntry.Type, TEXT("Provided query selection type can not be null."));
-			checkf(SelectEntry.Type->IsChildOf(FTypedElementDataStorageColumn::StaticStruct()),
-				TEXT("Provided query selection type '%s' is not based on FTypedElementDataStorageColumn."), 
+			checkf(
+				SelectEntry.Type->IsChildOf(FTypedElementDataStorageColumn::StaticStruct()) ||
+				SelectEntry.Type->IsChildOf(FMassFragment::StaticStruct()) || 
+				SelectEntry.Type->IsChildOf(FMassTag::StaticStruct()),
+				TEXT("Provided query selection type '%s' is not based on FTypedElementDataStorageColumn or another supported base type."), 
 				*SelectEntry.Type->GetStructPathName().ToString());
 			NativeQuery.AddRequirement(SelectEntry.Type, LocalToNativeAccess(SelectEntry.Access));
 		}
@@ -444,15 +453,23 @@ FName UTypedElementDatabase::GetQueryTickGroupName(EQueryTickGroups Group) const
 	{
 		case EQueryTickGroups::Default:
 			return NAME_None;
+		case EQueryTickGroups::PrepareSyncExternalToDataStorage:
+			return TickGroupName_PrepareSyncWorldToMass;
 		case EQueryTickGroups::SyncExternalToDataStorage:
 			return UE::Mass::ProcessorGroupNames::SyncWorldToMass;
+		case EQueryTickGroups::FinalizeSyncExternalToDataStorage:
+			return TickGroupName_FinalizeSyncWorldToMass;
+		
+		case EQueryTickGroups::PrepareSyncDataStorageToExternal:
+			return TickGroupName_PrepareSyncMassToExternal;
 		case EQueryTickGroups::SyncDataStorageToExternal:
 			return UE::Mass::ProcessorGroupNames::UpdateWorldFromMass;
+		case EQueryTickGroups::FinalizeSyncDataStorageToExternal:
+			return TickGroupName_FinalizeSyncMassToExternal;
+		
 		case EQueryTickGroups::SyncWidgets:
-		{
-			static const FName SyncWidgetGroupName(TEXT("SyncWidgets"));
-			return SyncWidgetGroupName;
-		}
+			return TickGroupName_SyncWidget;
+
 		default:
 			checkf(false, TEXT("EQueryTickGroups value %i can't be translated to a group name by this Data Storage backend."), static_cast<int>(Group));
 			return NAME_None;
diff --git a/Engine/Plugins/Experimental/TypedElementsDataStorage/Source/TypedElementsDataStorage/Private/TypedElementDatabase.h b/Engine/Plugins/Experimental/TypedElementsDataStorage/Source/TypedElementsDataStorage/Private/TypedElementDatabase.h
--- a/Engine/Plugins/Experimental/TypedElementsDataStorage/Source/TypedElementsDataStorage/Private/TypedElementDatabase.h
+++ b/Engine/Plugins/Experimental/TypedElementsDataStorage/Source/TypedElementsDataStorage/Private/TypedElementDatabase.h
@@ -80,6 +80,12 @@ private:
 
 	void Reset();
 	
+	static const FName TickGroupName_PrepareSyncWorldToMass;
+	static const FName TickGroupName_FinalizeSyncWorldToMass;
+	static const FName TickGroupName_PrepareSyncMassToExternal;
+	static const FName TickGroupName_FinalizeSyncMassToExternal;
+	static const FName TickGroupName_SyncWidget;
+	
 	TArray<FMassArchetypeHandle> Tables;
 	TMap<FName, TypedElementTableHandle> TableNameLookup;
 
diff --git a/Engine/Plugins/Experimental/TypedElementsDataStorage/Source/TypedElementsDataStorage/Private/TypedElementDatabaseCompatibility.cpp b/Engine/Plugins/Experimental/TypedElementsDataStorage/Source/TypedElementsDataStorage/Private/TypedElementDatabaseCompatibility.cpp
--- a/Engine/Plugins/Experimental/TypedElementsDataStorage/Source/TypedElementsDataStorage/Private/TypedElementDatabaseCompatibility.cpp
+++ b/Engine/Plugins/Experimental/TypedElementsDataStorage/Source/TypedElementsDataStorage/Private/TypedElementDatabaseCompatibility.cpp
@@ -86,7 +86,6 @@ void UTypedElementDatabaseCompatibility::CreateStandardArchetypes()
 	StandardActorWithTransformTable = Storage->RegisterTable(StandardActorTable,
 		MakeArrayView(
 		{
-			FTypedElementLabelColumn::StaticStruct(),
 			FTypedElementLocalTransformColumn::StaticStruct()
 		}), FName("Editor_StandardActorWithTransformTable"));
 }
diff --git a/Engine/Plugins/Experimental/TypedElementsDataStorage/Source/TypedElementsDataStorageUI/Private/TypedElementCounterWidget.cpp b/Engine/Plugins/Experimental/TypedElementsDataStorage/Source/TypedElementsDataStorageUI/Private/TypedElementCounterWidget.cpp
--- a/Engine/Plugins/Experimental/TypedElementsDataStorage/Source/TypedElementsDataStorageUI/Private/TypedElementCounterWidget.cpp
+++ b/Engine/Plugins/Experimental/TypedElementsDataStorage/Source/TypedElementsDataStorageUI/Private/TypedElementCounterWidget.cpp
@@ -32,8 +32,8 @@ void UTypedElementCounterWidgetFactory::RegisterQueries(ITypedElementDataStorage
 			const FTypedElementCounterWidgetColumn& Counter
 		)
 		{
-			UTypedElementDataStorageSubsystem* Subsystem = Context.GetCachedMutableSubsystem<UTypedElementDataStorageSubsystem>();
-			DSI* DataInterface = Subsystem->Get();
+			UTypedElementDataStorageSubsystem& Subsystem = Context.GetCachedMutableDependency<UTypedElementDataStorageSubsystem>();
+			DSI* DataInterface = Subsystem.Get();
 			checkf(DataInterface, TEXT("FTypedElementsDataStorageUiModule tried to process widgets before the "
 				"Typed Elements Data Storage interface is available."));
 
diff --git a/Engine/Source/Runtime/TypedElementFramework/Private/Elements/Framework/TypedElementQueryBuilder.cpp b/Engine/Source/Runtime/TypedElementFramework/Private/Elements/Framework/TypedElementQueryBuilder.cpp
--- a/Engine/Source/Runtime/TypedElementFramework/Private/Elements/Framework/TypedElementQueryBuilder.cpp
+++ b/Engine/Source/Runtime/TypedElementFramework/Private/Elements/Framework/TypedElementQueryBuilder.cpp
@@ -219,57 +219,6 @@ namespace TypedElementQueryBuilder
 	}
 
 
-	/**
-	 * FQueryContextForwarder
-	 */
-
-	FQueryContextForwarder::FQueryContextForwarder(ITypedElementDataStorageInterface::FQueryContext& InParentContext)
-		: ParentContext(InParentContext)
-	{}
-
-	const void* FQueryContextForwarder::GetColumn(const UScriptStruct* ColumnType) const
-	{
-		return ParentContext.GetColumn(ColumnType);
-	}
-
-	void* FQueryContextForwarder::GetMutableColumn(const UScriptStruct* ColumnType)
-	{
-		return ParentContext.GetMutableColumn(ColumnType);
-	}
-
-	void FQueryContextForwarder::GetColumns(TArrayView<char*> RetrievedAddresses, TConstArrayView<const UScriptStruct*> ColumnTypes,
-		TConstArrayView<ITypedElementDataStorageInterface::EQueryAccessType> AccessTypes)
-	{
-		ParentContext.GetColumns(RetrievedAddresses, ColumnTypes, AccessTypes);
-	}
-
-	void FQueryContextForwarder::GetColumnsUnguarded(int32 TypeCount, char** RetrievedAddresses, const UScriptStruct* const* ColumnTypes,
-		const ITypedElementDataStorageInterface::EQueryAccessType* AccessTypes)
-	{
-		ParentContext.GetColumnsUnguarded(TypeCount, RetrievedAddresses, ColumnTypes, AccessTypes);
-	}
-
-	USubsystem* FQueryContextForwarder::GetMutableSubsystem(const TSubclassOf<USubsystem> SubsystemClass)
-	{
-		return ParentContext.GetMutableSubsystem(SubsystemClass);
-	}
-
-	const USubsystem* FQueryContextForwarder::GetSubsystem(const TSubclassOf<USubsystem> SubsystemClass)
-	{
-		return ParentContext.GetSubsystem(SubsystemClass);
-	}
-
-	void FQueryContextForwarder::GetSubsystems(TArrayView<char*> RetrievedAddresses, TConstArrayView<const TSubclassOf<USubsystem>> SubsystemTypes,
-		TConstArrayView<ITypedElementDataStorageInterface::EQueryAccessType> AccessTypes)
-	{
-		ParentContext.GetSubsystems(RetrievedAddresses, SubsystemTypes, AccessTypes);
-	}
-
-	uint32 FQueryContextForwarder::GetRowCount() const
-	{
-		return ParentContext.GetRowCount();
-	}
-
 
 	/**
 	 * Select
diff --git a/Engine/Source/Runtime/TypedElementFramework/Public/Elements/Framework/TypedElementQueryBuilder.h b/Engine/Source/Runtime/TypedElementFramework/Public/Elements/Framework/TypedElementQueryBuilder.h
--- a/Engine/Source/Runtime/TypedElementFramework/Public/Elements/Framework/TypedElementQueryBuilder.h
+++ b/Engine/Source/Runtime/TypedElementFramework/Public/Elements/Framework/TypedElementQueryBuilder.h
@@ -35,16 +35,30 @@ class USubsystem;
  *
  * Some functions allow binding to a callback. In these cases the arguments to the provided callback are analyzed and
  * added to the query automatically. Const arguments are added as ReadOnly, while non-const arguments are added as 
- * ReadWrite. Callbacks can be periodically called if constructed as a processor, in which case the callback is
- * repeatedly, usually once per frame, called for all row (ranges) that match the query. If constructed as an observer
+ * ReadWrite. Callbacks can be periodically called if constructed as a processor, in which case the callback is triggered
+ * repeatedly, usually once per frame and called for all row (ranges) that match the query. If constructed as an observer
  * the provided target type is monitored for actions like addition or deletion into/from any table and will trigger the
- * callback once if the query matches. The following function signatures are excepted by "Select":
- *		- void(ITypedElementDataStorageInterface::FQueryContext&, [const]Column&...) 
- *			- The provided callback will be called for each row that matches the provided columns.
- *		- void(FCachedQueryContext<[const]Dependencies...>&, [const]Column&...)
- *			- Same as the previous but each listed dependency is registered and cached to improve performance. Note that
- *				dependencies marked as const will only be available for reading and otherwise are can be accessed for read/write.
- * Usage example:
+ * callback once if the query matches. The following function signatures are accepted by "Select":
+ *	- void([const]Column&...) 
+ *	- void([const]Column*...) 
+ *	- void(TypedElementRowHandle, [const]Column&...) 
+ *	- void(<Context>&, [const]Column&...) 
+ *	- void(<Context>&, TypedElementRowHandle, [const]Column&...) 
+ *	- void(<Context>&, [const]Column*...) 
+ *	- void(<Context>&, const TypedElementRowHandle*, [const]Column*...) 
+ *	Where <Context> is ITypedElementDataStorageInterface::FQueryContext or FCachedQueryContext<...>	e.g.:
+ *		void(
+ *			FCachedQueryContext<Subsystem1, const Subsystem2>& Context, 
+ *			TypedElementRowHandle Row, 
+ *			ColumnType0& ColumnA, 
+ *			const ColumnType1& ColumnB) 
+ *			{...}
+ *
+ * FCachedQueryContext can be used to store cached pointers to dependencies to reduce the overhead of retrieving these. The same
+ * const principle as for other arguments applies so dependencies marked as const can only be accessed as read-only or otherwise
+ * can be accessed as readwrite.
+ *
+ * The following is a simplified example of these options combined together:
  *		FProcessor Info(
  *			ITypedElementDataStorageInterface::EQueryTickPhase::FrameEnd, 
  *			DataStorage->GetQueryTickGroupName(ITypedElementDataStorageInterface::EQueryTickGroups::SyncExternalToDataStorage);
@@ -187,43 +201,51 @@ namespace TypedElementQueryBuilder
 		bool bForceToGameThread{ false };
 	};
 
-	struct TYPEDELEMENTFRAMEWORK_API FQueryContextForwarder : public ITypedElementDataStorageInterface::FQueryContext
+	// Because this is a thin wrapper called from within a query callback, it's better to inline fully so all
+	// function pre/postambles can be optimized away.
+	struct FQueryContextForwarder : public ITypedElementDataStorageInterface::FQueryContext
 	{
-		explicit FQueryContextForwarder(ITypedElementDataStorageInterface::FQueryContext& InParentContext);
-		
-		const void* GetColumn(const UScriptStruct* ColumnType) const override;
-		void* GetMutableColumn(const UScriptStruct* ColumnType) override;
-		void GetColumns(TArrayView<char*> RetrievedAddresses, TConstArrayView<const UScriptStruct*> ColumnTypes,
+		inline explicit FQueryContextForwarder(ITypedElementDataStorageInterface::FQueryContext& InParentContext);
+		inline ~FQueryContextForwarder() = default;
+
+		inline const void* GetColumn(const UScriptStruct* ColumnType) const override;
+		inline void* GetMutableColumn(const UScriptStruct* ColumnType) override;
+		inline void GetColumns(TArrayView<char*> RetrievedAddresses, TConstArrayView<TWeakObjectPtr<const UScriptStruct>> ColumnTypes,
 			TConstArrayView<ITypedElementDataStorageInterface::EQueryAccessType> AccessTypes) override;
-		void GetColumnsUnguarded(int32 TypeCount, char** RetrievedAddresses, const UScriptStruct* const* ColumnTypes, 
+		inline void GetColumnsUnguarded(int32 TypeCount, char** RetrievedAddresses, const TWeakObjectPtr<const UScriptStruct>* ColumnTypes,
 			const ITypedElementDataStorageInterface::EQueryAccessType* AccessTypes) override;
 
-		USubsystem* GetMutableSubsystem(const TSubclassOf<USubsystem> SubsystemClass) override;
-		const USubsystem* GetSubsystem(const TSubclassOf<USubsystem> SubsystemClass) override;
-		void GetSubsystems(TArrayView<char*> RetrievedAddresses, TConstArrayView<const TSubclassOf<USubsystem>> SubsystemTypes, 
+		inline UObject* GetMutableDependency(const UClass* DependencyClass) override;
+		inline const UObject* GetDependency(const UClass* DependencyClass) override;
+		inline void GetDependencies(TArrayView<UObject*> RetrievedAddresses, TConstArrayView<TWeakObjectPtr<const UClass>> SubsystemTypes,
 			TConstArrayView<ITypedElementDataStorageInterface::EQueryAccessType> AccessTypes) override;
 
-		uint32 GetRowCount() const override;
+		inline uint32 GetRowCount() const override;
+		inline TConstArrayView<TypedElementRowHandle> GetRowHandles() const override;
+		inline void RemoveRow(TypedElementRowHandle Row) override;
+		inline void RemoveRows(TConstArrayView<TypedElementRowHandle> Rows) override;
+
+		inline void AddColumns(TypedElementRowHandle Row, TConstArrayView<const UScriptStruct*> ColumnTypes) override;
+		inline void AddColumns(TConstArrayView<TypedElementRowHandle> Rows, TConstArrayView<const UScriptStruct*> ColumnTypes) override;
+		inline void RemoveColumns(TypedElementRowHandle Row, TConstArrayView<const UScriptStruct*> ColumnTypes) override;
+		inline void RemoveColumns(TConstArrayView<TypedElementRowHandle> Rows, TConstArrayView<const UScriptStruct*> ColumnTypes) override;
 
 		ITypedElementDataStorageInterface::FQueryContext& ParentContext;
 	};
 
-	template<typename... Subsystems>
+	template<typename... Dependencies>
 	struct FCachedQueryContext final : public FQueryContextForwarder
 	{
 		explicit FCachedQueryContext(ITypedElementDataStorageInterface::FQueryContext& InParentContext);
-
-		void Fetch();
+		
 		static void Register(ITypedElementDataStorageInterface::FQueryDescription& Query);
 
-		template<typename Subsystem>
-		Subsystem* GetCachedMutableSubsystem();
-		template<typename Subsystem>
-		const Subsystem* GetCachedSubsystem() const;
+		template<typename Dependency>
+		Dependency& GetCachedMutableDependency();
+		template<typename Dependency>
+		const Dependency& GetCachedDependency() const;
 
-		char* SubsystemAddresses[sizeof...(Subsystems)];
-		const TSubclassOf<USubsystem> SubsystemTypes[sizeof...(Subsystems)];
-		ITypedElementDataStorageInterface::EQueryAccessType SubsystemAccessList[sizeof...(Subsystems)];
+		UObject* DependencyAddresses[sizeof...(Dependencies)];
 	};
 
 	// Explicitly not following the naming convention in order to present this as a query that can be read as such.
@@ -233,9 +255,9 @@ namespace TypedElementQueryBuilder
 		Select();
 
 		template<typename CallbackType, typename Function>
-		Select(FName Name, const CallbackType& Type, Function Callback);
+		Select(FName Name, const CallbackType& Type, Function&& Callback);
 		template<typename CallbackType, typename Class, typename Function>
-		Select(FName Name, const CallbackType& Type, Class* Instance, Function Callback);
+		Select(FName Name, const CallbackType& Type, Class* Instance, Function&& Callback);
 
 		template<typename... TargetTypes>
 		Select& ReadOnly();
diff --git a/Engine/Source/Runtime/TypedElementFramework/Public/Elements/Framework/TypedElementQueryBuilder.inl b/Engine/Source/Runtime/TypedElementFramework/Public/Elements/Framework/TypedElementQueryBuilder.inl
--- a/Engine/Source/Runtime/TypedElementFramework/Public/Elements/Framework/TypedElementQueryBuilder.inl
+++ b/Engine/Source/Runtime/TypedElementFramework/Public/Elements/Framework/TypedElementQueryBuilder.inl
@@ -1,6 +1,8 @@
 // Copyright Epic Games, Inc. All Rights Reserved.
 
+#include <tuple>
 #include <type_traits>
+#include <utility>
 #include "Templates/IsConst.h"
 #include "Templates/UnrealTypeTraits.h"
 
@@ -26,7 +28,7 @@ namespace TypedElementQueryBuilder
 		template<typename Type>
 		constexpr ITypedElementDataStorageInterface::EQueryAccessType GetAccessType()
 		{
-			using BaseType = typename TRemoveReference<Type>::Type;
+			using BaseType = typename std::remove_reference_t<Type>;
 			if constexpr (TIsConst<BaseType>::Value)
 			{
 				return ITypedElementDataStorageInterface::EQueryAccessType::ReadOnly;
@@ -36,7 +38,7 @@ namespace TypedElementQueryBuilder
 				return ITypedElementDataStorageInterface::EQueryAccessType::ReadWrite;
 			}
 		}
-	}
+	} // namespace Internal
 
 	//
 	// FDependecy
@@ -73,67 +75,148 @@ namespace TypedElementQueryBuilder
 	}
 
 
+	//
+	// FQueryContextForwarder
+	//
+
+	FQueryContextForwarder::FQueryContextForwarder(ITypedElementDataStorageInterface::FQueryContext& InParentContext)
+		: ParentContext(InParentContext)
+	{}
+
+	const void* FQueryContextForwarder::GetColumn(const UScriptStruct* ColumnType) const
+	{
+		return ParentContext.GetColumn(ColumnType);
+	}
+
+	void* FQueryContextForwarder::GetMutableColumn(const UScriptStruct* ColumnType)
+	{
+		return ParentContext.GetMutableColumn(ColumnType);
+	}
+
+	void FQueryContextForwarder::GetColumns(TArrayView<char*> RetrievedAddresses, TConstArrayView<TWeakObjectPtr<const UScriptStruct>> ColumnTypes,
+		TConstArrayView<ITypedElementDataStorageInterface::EQueryAccessType> AccessTypes)
+	{
+		ParentContext.GetColumns(RetrievedAddresses, ColumnTypes, AccessTypes);
+	}
+
+	void FQueryContextForwarder::GetColumnsUnguarded(int32 TypeCount, char** RetrievedAddresses, const TWeakObjectPtr<const UScriptStruct>* ColumnTypes,
+		const ITypedElementDataStorageInterface::EQueryAccessType* AccessTypes)
+	{
+		ParentContext.GetColumnsUnguarded(TypeCount, RetrievedAddresses, ColumnTypes, AccessTypes);
+	}
+
+	UObject* FQueryContextForwarder::GetMutableDependency(const UClass* DependencyClass)
+	{
+		return ParentContext.GetMutableDependency(DependencyClass);
+	}
+
+	const UObject* FQueryContextForwarder::GetDependency(const UClass* DependencyClass)
+	{
+		return ParentContext.GetDependency(DependencyClass);
+	}
+
+	void FQueryContextForwarder::GetDependencies(TArrayView<UObject*> RetrievedAddresses, TConstArrayView<TWeakObjectPtr<const UClass>> DependencyTypes,
+		TConstArrayView<ITypedElementDataStorageInterface::EQueryAccessType> AccessTypes)
+	{
+		ParentContext.GetDependencies(RetrievedAddresses, DependencyTypes, AccessTypes);
+	}
+
+	uint32 FQueryContextForwarder::GetRowCount() const
+	{
+		return ParentContext.GetRowCount();
+	}
+
+	TConstArrayView<TypedElementRowHandle> FQueryContextForwarder::GetRowHandles() const
+	{
+		return ParentContext.GetRowHandles();
+	}
+
+	void FQueryContextForwarder::RemoveRow(TypedElementRowHandle Row)
+	{
+		ParentContext.RemoveRow(Row);
+	}
+
+	void FQueryContextForwarder::RemoveRows(TConstArrayView<TypedElementRowHandle> Rows)
+	{
+		ParentContext.RemoveRows(Rows);
+	}
+
+	void FQueryContextForwarder::AddColumns(TypedElementRowHandle Row, TConstArrayView<const UScriptStruct*> ColumnTypes)
+	{
+		ParentContext.AddColumns(Row, ColumnTypes);
+	}
+
+	void FQueryContextForwarder::AddColumns(TConstArrayView<TypedElementRowHandle> Rows, TConstArrayView<const UScriptStruct*> ColumnTypes)
+	{
+		ParentContext.AddColumns(Rows, ColumnTypes);
+	}
+
+	void FQueryContextForwarder::RemoveColumns(TypedElementRowHandle Row, TConstArrayView<const UScriptStruct*> ColumnTypes)
+	{
+		ParentContext.RemoveColumns(Row, ColumnTypes);
+	}
+
+	void FQueryContextForwarder::RemoveColumns(TConstArrayView<TypedElementRowHandle> Rows, TConstArrayView<const UScriptStruct*> ColumnTypes)
+	{
+		ParentContext.RemoveColumns(Rows, ColumnTypes);
+	}
+
 	//
 	// FCachedQueryContext
 	//
 	
-	template<typename... Subsystems>
-	FCachedQueryContext<Subsystems...>::FCachedQueryContext(ITypedElementDataStorageInterface::FQueryContext& InParentContext)
+	template<typename... Dependencies>
+	FCachedQueryContext<Dependencies...>::FCachedQueryContext(ITypedElementDataStorageInterface::FQueryContext& InParentContext)
 		: FQueryContextForwarder(InParentContext)
-		, SubsystemTypes{ Subsystems::StaticClass()... }
-		, SubsystemAccessList{ Internal::GetAccessType<Subsystems>()... }
 	{
-	}
+		static constexpr ITypedElementDataStorageInterface::EQueryAccessType DependencyAccessList[] = { Internal::GetAccessType<Dependencies>()... };
+		static TWeakObjectPtr<const UClass> DependencyTypes[] = { Dependencies::StaticClass()... };
 
-	template<typename... Subsystems>
-	void FCachedQueryContext<Subsystems...>::Fetch()
-	{
-		ParentContext.GetSubsystems(SubsystemAddresses, SubsystemTypes, SubsystemAccessList);
+		ParentContext.GetDependencies(DependencyAddresses, DependencyTypes, DependencyAccessList);
 #if DO_CHECK
-		char** SubsystemAddessIt = SubsystemAddresses;
-		for (int32 Index = 0; Index < sizeof...(Subsystems); ++Index)
+		UObject** DependencyAddessIt = DependencyAddresses;
+		for (int32 Index = 0; Index < sizeof...(Dependencies); ++Index)
 		{
-			checkf(*SubsystemAddessIt != nullptr, TEXT("The cached query context tried to retrieve a subsystem '%s' that didn't exist."),
-				SubsystemTypes[Index] != nullptr ? *(SubsystemTypes[Index]->GetName()) : TEXT("<no type information>"));
-			++SubsystemAddessIt;
+			checkf(*DependencyAddessIt != nullptr, TEXT("The cached query context tried to retrieve a dependency '%s' that didn't exist."),
+				DependencyTypes[Index] != nullptr ? *(DependencyTypes[Index]->GetName()) : TEXT("<no type information>"));
+			++DependencyAddessIt;
 		}
 #endif // DO_CHECK
 	}
-	
-	template<typename... Subsystems>
-	void FCachedQueryContext<Subsystems...>::Register(ITypedElementDataStorageInterface::FQueryDescription& Query)
+
+	template<typename... Dependencies>
+	void FCachedQueryContext<Dependencies...>::Register(ITypedElementDataStorageInterface::FQueryDescription& Query)
 	{
-		( Query.Dependencies.Emplace(Subsystems::StaticClass(), Internal::GetAccessType<Subsystems>()), ... );
+		( Query.Dependencies.Emplace(Dependencies::StaticClass(), Internal::GetAccessType<Dependencies>()), ... );
 	}
 
-	template<typename... Subsystems>
-	template<typename Subsystem>
-	Subsystem* FCachedQueryContext<Subsystems...>::GetCachedMutableSubsystem()
+	template<typename... Dependencies>
+	template<typename Dependency>
+	Dependency& FCachedQueryContext<Dependencies...>::GetCachedMutableDependency()
 	{
 		// Don't allow a subsystem registered as const to be found.
-		constexpr uint32 Index = Internal::GetVarArgIndex<std::remove_const_t<Subsystem>, Subsystems...>();
-		static_assert(Index < sizeof...(Subsystems), "Requested subsystem isn't part of the query context cache.");
-		return reinterpret_cast<Subsystem*>(SubsystemAddresses[Index]);
+		constexpr uint32 Index = Internal::GetVarArgIndex<std::remove_const_t<Dependency>, Dependencies...>();
+		static_assert(Index < sizeof...(Dependencies), "Requested dependency isn't part of the query context cache.");
+		return *static_cast<Dependency*>(DependencyAddresses[Index]);
 	}
 
-	template<typename... Subsystems>
-	template<typename Subsystem>
-	const Subsystem* FCachedQueryContext<Subsystems...>::GetCachedSubsystem() const
+	template<typename... Dependencies>
+	template<typename Dependency>
+	const Dependency& FCachedQueryContext<Dependencies...>::GetCachedDependency() const
 	{
 		// Allow access to subsystems registered with and without const.
-		constexpr uint32 Index = Internal::GetVarArgIndex<Subsystem, Subsystems...>();
-		if constexpr (Index < sizeof...(Subsystems))
+		constexpr uint32 Index = Internal::GetVarArgIndex<Dependency, Dependencies...>();
+		if constexpr (Index < sizeof...(Dependencies))
 		{
-			return reinterpret_cast<Subsystem*>(SubsystemAddresses[Index]);
+			return *static_cast<Dependency*>(DependencyAddresses[Index]);
 		}
 		else
 		{
-			constexpr uint32 ConstIndex = Internal::GetVarArgIndex<std::add_const_t<Subsystem>, Subsystems...>();
-			static_assert(ConstIndex < sizeof...(Subsystems), "Requested subsystem isn't part of the query context cache.");
-			return reinterpret_cast<Subsystem*>(SubsystemAddresses[ConstIndex]);
+			constexpr uint32 ConstIndex = Internal::GetVarArgIndex<std::add_const_t<Dependency>, Dependencies...>();
+			static_assert(ConstIndex < sizeof...(Dependencies), "Requested dependency isn't part of the query context cache.");
+			return *static_cast<Dependency*>(DependencyAddresses[ConstIndex]);
 		}
 	}
-
 	
 	//
 	// Select
@@ -143,111 +226,478 @@ namespace TypedElementQueryBuilder
 		using QueryContext = ITypedElementDataStorageInterface::FQueryContext;
 		using QueryAccess = ITypedElementDataStorageInterface::EQueryAccessType;
 
-#		define QUERY_BINDING_BODY(Caller)																					\
-			static const UScriptStruct* ColumnTypes[] = { Args::StaticStruct()... };										\
-			char* ColumnAddresses[sizeof...(Args)];																			\
-			static constexpr QueryAccess ColumnAccess[] = { GetAccessType<Args>()... };										\
-			Context.GetColumnsUnguarded(sizeof...(Args), ColumnAddresses, ColumnTypes, ColumnAccess);						\
-			uint32 ColumnSize[] = { sizeof(Args)... };																		\
-																															\
-			using BaseContextType = typename TRemoveReference<ContextType>::Type;											\
-			if constexpr(std::is_same_v<BaseContextType, QueryContext>)														\
-			{																												\
-				const uint32 RowCount = Context.GetRowCount();																\
-				for (uint32 Index = 0; Index < RowCount; ++Index)															\
-				{																											\
-					Caller (Context, *reinterpret_cast<Args*>(ColumnAddresses[GetVarArgIndex<Args, Args...>()])...);		\
-					for (uint32 ColumnIndex = 0; ColumnIndex < sizeof...(Args); ++ColumnIndex)								\
-					{																										\
-						ColumnAddresses[ColumnIndex] += ColumnSize[ColumnIndex];											\
-					}																										\
-				}																											\
-			}																												\
-			else																											\
-			{																												\
-				BaseContextType CachedContext(Context);																		\
-				CachedContext.Fetch();																						\
-																															\
-				const uint32 RowCount = Context.GetRowCount();																\
-				for (uint32 Index = 0; Index < RowCount; ++Index)															\
-				{																											\
-					Caller(CachedContext, *reinterpret_cast<Args*>(ColumnAddresses[GetVarArgIndex<Args, Args...>()])...);	\
-					for (uint32 ColumnIndex = 0; ColumnIndex < sizeof...(Args); ++ColumnIndex)								\
-					{																										\
-						ColumnAddresses[ColumnIndex] += ColumnSize[ColumnIndex];											\
-					}																										\
-				}																											\
+		template <typename T>
+		struct HasStaticStructMethod
+		{
+		private:
+			template <typename C> static constexpr bool Check(decltype(&C::StaticStruct)) { return std::is_same_v<decltype(&C::StaticStruct), UScriptStruct* (*)()>; }
+			template <typename C> static constexpr bool Check(...) { return false; }
+
+		public:
+			static constexpr bool Value = Check<T>(nullptr);
+		};
+
+		template<typename Context>
+		constexpr bool IsValidContextType()
+		{
+			using BaseContextType = typename std::remove_const_t<typename TRemoveReference<Context>::Type>;
+			if constexpr (std::is_same_v<BaseContextType, QueryContext>)
+			{
+				return true;
 			}
+			else if constexpr (TIsDerivedFrom<BaseContextType, FQueryContextForwarder>::Value)
+			{
+				return true;
+			}
+			else
+			{
+				return false;
+			}
+		}
+
+		template<typename... Columns>
+		constexpr bool AreAllColumnsReferences()
+		{
+			if constexpr (sizeof...(Columns) > 0)
+			{
+				return (std::is_reference_v<Columns> && ...);
+			}
+			else
+			{
+				return true;
+			}
+		}
+
+		template<typename... Columns>
+		constexpr bool AreAllColumnsPointers()
+		{
+			if constexpr (sizeof...(Columns) > 0)
+			{
+				return (std::is_pointer_v<Columns> && ...);
+			}
+			else
+			{
+				return true;
+			}
+		}
+
+		template<typename Column>
+		using BaseColumnType = std::remove_pointer_t<std::remove_reference_t<Column>>;
+
+		template<typename Column>
+		using UndecoratedColumnType = std::remove_cv_t<BaseColumnType<Column>>;
+
+		template<typename Column>
+		constexpr bool IsValidColumnType()
+		{
+			// Do not combine these two into a single statement as that would also allow invalid arguments like "MyColumn*&"
+			if constexpr (std::is_reference_v<Column>)
+			{
+				using BaseType = typename std::remove_const_t<std::remove_reference_t<Column>>;
+				return HasStaticStructMethod<BaseType>::Value;
+			}
+			else if constexpr (std::is_pointer_v<Column>)
+			{
+				using BaseType = typename std::remove_const_t<std::remove_pointer_t<Column>>;
+				return HasStaticStructMethod<BaseType>::Value;
+			}
+			else
+			{
+				return false;
+			}
+		}
+
+		template<typename Column>
+		const UScriptStruct* GetColumnType()
+		{
+			return UndecoratedColumnType<Column>::StaticStruct();
+		}
+
+		template<typename RowType>
+		constexpr bool IsRowHandleType()
+		{
+			return std::is_same_v<std::remove_cv_t<std::remove_pointer_t<RowType>>, TypedElementRowHandle>;
+		}
+
+		template<typename RowType, typename... Columns>
+		constexpr bool IsRowTypeCompatibleWithColumns()
+		{
+			if constexpr (std::is_pointer_v<RowType> && std::is_const_v<std::remove_pointer_t<RowType>>)
+			{
+				return AreAllColumnsPointers<Columns...>();
+			}
+			else
+			{
+				return AreAllColumnsReferences<Columns...>();
+			}
+		}
+
+		template <typename T>
+		class IsFunctor
+		{
+			template <typename C> static constexpr bool Check(decltype(&C::operator())) { return true; }
+			template <typename C> static constexpr bool Check(...) { return false; }
+
+		public:
+			static constexpr bool Value = Check<T>(nullptr);
+		};
+
+		template<typename... ColumnTypes>
+		struct FFunctionColumnInfo
+		{
+			std::tuple<BaseColumnType<ColumnTypes>*...> Columns;
+			static constexpr bool bArePointerColumns = AreAllColumnsPointers<ColumnTypes...>();
+
+			explicit FFunctionColumnInfo(QueryContext& Context)
+			{
+				static constexpr QueryAccess ColumnAccess[] = { GetAccessType<ColumnTypes>()... };
+				static TWeakObjectPtr<const UScriptStruct> ColumnDescriptions[] = { GetColumnType<ColumnTypes>()...};
+
+				char* ColumnAddresses[sizeof...(ColumnTypes)];
+				Context.GetColumnsUnguarded(sizeof...(ColumnTypes), ColumnAddresses, ColumnDescriptions, ColumnAccess);
+				std::apply([ColumnAddresses](auto&&... Column)
+					{
+						int Index = 0;
+						((Column = reinterpret_cast<BaseColumnType<ColumnTypes>*>(ColumnAddresses[Index++])), ...);
+					}, Columns);
+			}
+		};
+
+		template<>
+		struct FFunctionColumnInfo<>
+		{
+			explicit FFunctionColumnInfo(QueryContext&) {}
+			static constexpr bool bArePointerColumns = false;
+		};
+
+		template<typename ContextType>
+		struct FContextInfo
+		{
+			using BaseContextType = std::remove_reference_t<ContextType>;
+			static constexpr bool bIsUndecoratedContext = std::is_same_v<BaseContextType, QueryContext>;
+			using ContextWrapperType = std::conditional_t<bIsUndecoratedContext, QueryContext&, BaseContextType>;
+			ContextWrapperType ContextWrapper;
+
+			explicit FContextInfo(QueryContext& Context)
+				: ContextWrapper(Context)
+			{}
+		};
+
+		template<typename ContextType, typename RowHandleType, typename... Args>
+		struct FContextRowHandleColumnsFunction : public FFunctionColumnInfo<Args...>, public FContextInfo<ContextType>
+		{
+			using SuperColumn = FFunctionColumnInfo<Args...>;
+			using SuperContext = FContextInfo<ContextType>;
+
+			explicit FContextRowHandleColumnsFunction(QueryContext& Context)
+				: SuperColumn(Context)
+				, SuperContext(Context)
+			{}
+
+			template<typename CallerType>
+			void Call(QueryContext& Context, CallerType&& Caller)
+			{
+				TConstArrayView<TypedElementRowHandle> Rows = Context.GetRowHandles();
+				if constexpr (SuperColumn::bArePointerColumns || std::is_pointer_v<RowHandleType>)
+				{
+					if constexpr (sizeof...(Args) > 0)
+					{
+						std::apply([this, Rows = Rows.GetData(), &Caller](auto&&... Column)
+							{
+								Caller(this->ContextWrapper, Rows, Column...);
+							}, this->Columns);
+					}
+					else
+					{
+						Caller(this->ContextWrapper, Rows.GetData());
+					}
+				}
+				else
+				{
+					if constexpr (sizeof...(Args) > 0)
+					{
+						for (TypedElementRowHandle Row : Rows)
+						{
+							if constexpr (sizeof...(Args) > 0)
+							{
+								std::apply([this, Row, &Caller](auto&&... Column)
+									{
+										Caller(this->ContextWrapper, Row, *Column...);
+										((++Column), ...);
+									}, this->Columns);
+							}
+							else
+							{
+								Caller(this->ContextWrapper, Row);
+							}
+						}
+					}
+				}
+			}
+		};
 
 		template<typename ContextType, typename... Args>
+		struct FContextColumnsFunction : public FFunctionColumnInfo<Args...>, public FContextInfo<ContextType>
+		{
+			using SuperColumn = FFunctionColumnInfo<Args...>;
+			using SuperContext = FContextInfo<ContextType>;
+
+			explicit FContextColumnsFunction(QueryContext& Context)
+				: SuperColumn(Context)
+				, SuperContext(Context)
+			{}
+
+			template<typename CallerType>
+			void Call(QueryContext& Context, CallerType&& Caller)
+			{
+				if constexpr (sizeof...(Args) > 0)
+				{
+					if constexpr (SuperColumn::bArePointerColumns)
+					{
+						std::apply(
+							[this, Caller = std::forward<CallerType>(Caller)](auto&&... Column)
+							{
+								Caller(this->ContextWrapper, Column...);
+							}, this->Columns);
+					}
+					else
+					{
+						auto CallForwarder = [this, Caller = std::forward<CallerType>(Caller)](auto&&... Column)
+						{
+							Caller(this->ContextWrapper, *Column...);
+							((++Column), ...);
+						};
+						const uint32 RowCount = Context.GetRowCount();
+						for (uint32 Index = 0; Index < RowCount; ++Index)
+						{
+							std::apply(CallForwarder, this->Columns);
+						}
+					}
+				}
+				else
+				{
+					const uint32 RowCount = Context.GetRowCount();
+					for (uint32 Index = 0; Index < RowCount; ++Index)
+					{
+						Caller(this->ContextWrapper);
+					}
+				}
+			}
+		};
+
+		template<typename RowHandleType, typename... Columns>
+		struct FRowHandleColumnsFunction : public FFunctionColumnInfo<Columns...>
+		{
+			using Super = FFunctionColumnInfo<Columns...>;
+
+			explicit FRowHandleColumnsFunction(QueryContext& Context)
+				: Super(Context)
+			{}
+
+			template<typename CallerType>
+			void Call(QueryContext& Context, CallerType&& Caller)
+			{
+				TConstArrayView<TypedElementRowHandle> Rows = Context.GetRowHandles();
+				if constexpr (Super::bArePointerColumns || std::is_pointer_v<RowHandleType>)
+				{
+					if constexpr (sizeof...(Columns) > 0)
+					{
+						std::apply([Rows = Rows.GetData(), &Caller](auto&&... Column)
+							{
+								Caller(Rows, Column...);
+							}, this->Columns);
+					}
+					else
+					{
+						Caller(Rows.GetData());
+					}
+				}
+				else
+				{
+					for (TypedElementRowHandle Row : Rows)
+					{
+						if constexpr (sizeof...(Columns) > 0)
+						{
+							std::apply([Row, &Caller](auto&&... Column)
+								{
+									Caller(Row, *Column...);
+									((++Column), ...);
+								}, this->Columns);
+						}
+						else
+						{
+							Caller(Row);
+						}
+					}
+				}
+			}
+		};
+
+		template<typename... Args>
+		struct FColumnsFunction : public FFunctionColumnInfo<Args...>
+		{
+			using Super = FFunctionColumnInfo<Args...>;
+			
+			explicit FColumnsFunction(QueryContext& Context)
+				: Super(Context)
+			{}
+
+			template<typename CallerType>
+			void Call(QueryContext& Context, CallerType&& Caller)
+			{
+				if constexpr (sizeof...(Args) > 0)
+				{
+					if constexpr (Super::bArePointerColumns)
+					{
+						std::apply(
+							[Caller = std::forward<CallerType>(Caller)](auto&&... Column)
+							{
+								Caller(Column...);
+							}, this->Columns);
+					}
+					else
+					{
+						auto CallForwarder = [Caller = std::forward<CallerType>(Caller)](auto&&... Column)
+						{
+							Caller(*Column...);
+							((++Column), ...);
+						};
+						const uint32 RowCount = Context.GetRowCount();
+						for (uint32 Index = 0; Index < RowCount; ++Index)
+						{
+							std::apply(CallForwarder, this->Columns);
+						}
+					}
+				}
+				else
+				{
+					const uint32 RowCount = Context.GetRowCount();
+					for (uint32 Index = 0; Index < RowCount; ++Index)
+					{
+						Caller();
+					}
+				}
+			}
+		};
+
+		template<typename... Args>
+		struct FFunctionInfoHelper
+		{
+			using BaseClass = FColumnsFunction<Args...>;
+		};
+
+		template<typename Arg0>
+		struct FFunctionInfoHelper<Arg0>
+		{
+			static constexpr bool bArg0IsContext = IsValidContextType<Arg0>();
+			static constexpr bool bArg0IsRowHandle = IsRowHandleType<Arg0>();
+			using BaseClass = 
+				std::conditional_t<bArg0IsContext,
+					FContextColumnsFunction<Arg0>,
+					std::conditional_t<bArg0IsRowHandle,
+						FRowHandleColumnsFunction<Arg0>,
+						FColumnsFunction<Arg0>
+					>
+				>;
+		};
+
+		template<typename Arg0, typename Arg1, typename... Args>
+		struct FFunctionInfoHelper<Arg0, Arg1, Args...>
+		{
+			static constexpr bool bArg0IsContext = IsValidContextType<Arg0>();
+			static constexpr bool bArg0IsRowHandle = IsRowHandleType<Arg0>();
+			
+			static constexpr bool bArg1IsRowHandle = IsRowHandleType<Arg1>();
+
+			using BaseClass =
+				std::conditional_t<bArg0IsContext,
+					std::conditional_t<bArg1IsRowHandle,
+						FContextRowHandleColumnsFunction<Arg0, Arg1, Args...>,
+						FContextColumnsFunction<Arg0, Arg1, Args...>
+					>,
+					std::conditional_t<bArg0IsRowHandle,
+						FRowHandleColumnsFunction<Arg0, Arg1, Args...>,
+						FColumnsFunction<Arg0, Arg1, Args...>
+					>
+				>;
+		};
+
+		template<typename... Args>
+		struct FFunctionInfo final : public FFunctionInfoHelper<Args...>::BaseClass
+		{
+			using Super = typename FFunctionInfoHelper<Args...>::BaseClass;
+
+			explicit FFunctionInfo(QueryContext& Context)
+				: Super(Context)
+			{}
+		};
+
+		template<typename... Args>
 		void BindQueryFunction(ITypedElementDataStorageInterface::QueryCallback& Function, 
-			void (*Callback)(ContextType&, Args&...))
+			void (*Callback)(Args...))
 		{
 			Function = [Callback](QueryContext& Context)
 			{
-				QUERY_BINDING_BODY(Callback)
+				FFunctionInfo<Args...>(Context).Call(Context, Callback);
 			};
 		}
 
-		template<typename Class, typename ContextType, typename... Args>
+		template<typename Class, typename... Args>
 		void BindQueryFunction(ITypedElementDataStorageInterface::QueryCallback& Function, Class* Target, 
-			void (Class::* Callback)(ContextType&, Args&...))
+			void (Class::*Callback)(Args...))
 		{
 			Function = [Target, Callback](QueryContext& Context)
 			{
-				QUERY_BINDING_BODY((Target->*Callback))
+				FFunctionInfo<Args...>(Context).Call(Context, (Target->*Callback));
 			};
 		}
 
-		template<typename Class, typename ContextType, typename... Args>
+		template<typename Class, typename... Args>
 		void BindQueryFunction(ITypedElementDataStorageInterface::QueryCallback& Function, Class* Target, 
-			void (Class::*Callback)(ContextType&, Args&...) const)
+			void (Class::*Callback)(Args...) const)
 		{
 			Function = [Target, Callback](QueryContext& Context)
 			{
-				QUERY_BINDING_BODY((Target->*Callback))
+				FFunctionInfo<Args...>(Context).Call(Context, (Target->*Callback));
 			};
 		}
 
-		template<typename Functor, typename Class, typename ContextType, typename... Args>
-		void BindQueryFunction_Expand(ITypedElementDataStorageInterface::QueryCallback& Function, Functor CallbackObject, 
-			void (Class::* Callback)(ContextType& , Args&...) const)
+		template<typename Functor, typename Class, typename... Args>
+		void BindQueryFunction_Expand(ITypedElementDataStorageInterface::QueryCallback& Function, Functor&& CallbackObject, 
+			void (Class::*Callback)(Args...) const)
 		{
 			Function = [CallbackObject = std::forward<Functor>(CallbackObject)](QueryContext& Context)
 			{
-				QUERY_BINDING_BODY(CallbackObject)
+				FFunctionInfo<Args...>(Context).Call(Context, CallbackObject);
 			};
 		}
 
-		template<typename Functor, typename Class, typename ContextType, typename... Args>
-		void BindQueryFunction_Expand(ITypedElementDataStorageInterface::QueryCallback& Function, Functor CallbackObject, 
-			void (Class::*Callback)(ContextType&, Args&...))
+		template<typename Functor, typename Class, typename... Args>
+		void BindQueryFunction_Expand(ITypedElementDataStorageInterface::QueryCallback& Function, Functor&& CallbackObject, 
+			void (Class::*Callback)(Args...))
 		{
 			Function = [CallbackObject = std::forward<Functor>(CallbackObject)](QueryContext& Context)
 			{
-				QUERY_BINDING_BODY(CallbackObject)
+				FFunctionInfo<Args...>(Context).Call(Context, CallbackObject);
 			};
 		}
 
 		template<typename Functor>
-		void BindQueryFunction(ITypedElementDataStorageInterface::QueryCallback& Function, Functor Callback)
+		void BindQueryFunction(ITypedElementDataStorageInterface::QueryCallback& Function, Functor&& Callback)
 		{
 			BindQueryFunction_Expand(Function, std::forward<Functor>(Callback), &Functor::operator());
 		}
-
-#		undef QUERY_BINDING_BODY
 		
 		template<typename Arg>
-		void AddArgumentToSelect(Select& Target)
+		void AddColumnToSelect(Select& Target)
 		{
-			using BaseType = typename TRemoveReference<Arg>::Type;
-			if constexpr (TIsConst<BaseType>::Value)
+			if constexpr (std::is_const_v<BaseColumnType<Arg>>)
 			{
-				Target.ReadOnly(BaseType::StaticStruct());
+				Target.ReadOnly(UndecoratedColumnType<Arg>::StaticStruct());
 			}
 			else
 			{
-				Target.ReadWrite(BaseType::StaticStruct());
+				Target.ReadWrite(UndecoratedColumnType<Arg>::StaticStruct());
 			}
 		}
 
@@ -260,104 +710,161 @@ namespace TypedElementQueryBuilder
 			}
 		}
 
-		template<typename ContextType, typename... Args>
-		void RegisterFunctionArguments(ITypedElementDataStorageInterface::FQueryDescription& Query, Select& Target, void (*)(ContextType&, Args...))
-		{ 
-			RegisterDependencies<ContextType>(Query);
-			(AddArgumentToSelect<Args>(Target), ...); 
+		template<typename... Args>
+		struct RegisterFunctionArgumentsHelper
+		{
+			void Register(ITypedElementDataStorageInterface::FQueryDescription&, Select& Target)
+			{
+				(AddColumnToSelect<Args>(Target), ...);
+			}
+		};
+
+		template<>
+		struct RegisterFunctionArgumentsHelper<>
+		{
+			static void Register(ITypedElementDataStorageInterface::FQueryDescription&, Select&) {}
+		};
+
+		template<typename Arg0>
+		struct RegisterFunctionArgumentsHelper<Arg0>
+		{
+			static void Register(ITypedElementDataStorageInterface::FQueryDescription& Query, Select& Target)
+			{
+				if constexpr (IsValidContextType<Arg0>())
+				{
+					RegisterDependencies<Arg0>(Query);
+				}
+				else if constexpr (IsValidColumnType<Arg0>())
+				{
+					AddColumnToSelect<Arg0>(Target);
+				}
+				// The third option is a row handle, which doesn't need to be registered.
+			}
+		};
+
+		template<typename Arg0, typename Arg1, typename... Args>
+		struct RegisterFunctionArgumentsHelper<Arg0, Arg1, Args...>
+		{
+			static void Register(ITypedElementDataStorageInterface::FQueryDescription& Query, Select& Target)
+			{
+				if constexpr (IsValidContextType<Arg0>())
+				{
+					RegisterDependencies<Arg0>(Query);
+				}
+				else if constexpr (IsValidColumnType<Arg0>())
+				{
+					AddColumnToSelect<Arg0>(Target);
+				}
+				// The third option is a row handle, which doesn't need to be registered.
+
+				if constexpr (IsValidColumnType<Arg1>())
+				{
+					AddColumnToSelect<Arg1>(Target);
+				}
+				// The second option is a row handle, which doesn't need to be registered.
+
+				(AddColumnToSelect<Args>(Target), ...);
+			}
+		};
+
+		template<typename... Args>
+		void RegisterFunctionArguments(ITypedElementDataStorageInterface::FQueryDescription& Query, Select& Target, void (*)(Args...))
+		{
+			RegisterFunctionArgumentsHelper<Args...>::Register(Query, Target);
 		}
 
-		template<typename Class, typename ContextType, typename... Args>
-		void RegisterFunctionArguments(ITypedElementDataStorageInterface::FQueryDescription& Query, Select& Target, void (Class::*)(ContextType&, Args...))
+		template<typename Class, typename... Args>
+		void RegisterFunctionArguments(ITypedElementDataStorageInterface::FQueryDescription& Query, Select& Target, void (Class::*)(Args...))
 		{
-			RegisterDependencies<ContextType>(Query);
-			(AddArgumentToSelect<Args>(Target), ...); 
+			RegisterFunctionArgumentsHelper<Args...>::Register(Query, Target);
 		}
 
-		template<typename Class, typename ContextType, typename... Args>
-		void RegisterFunctionArguments(ITypedElementDataStorageInterface::FQueryDescription& Query, Select& Target, void (Class::*)(ContextType&, Args...) const)
+		template<typename Class, typename... Args>
+		void RegisterFunctionArguments(ITypedElementDataStorageInterface::FQueryDescription& Query, Select& Target, void (Class::*)(Args...) const)
 		{
-			RegisterDependencies<ContextType>(Query);
-			(AddArgumentToSelect<Args>(Target), ...);
+			RegisterFunctionArgumentsHelper<Args...>::Register(Query, Target);
 		}
-		
+
 		template<typename Functor>
 		void RegisterFunctionArguments(ITypedElementDataStorageInterface::FQueryDescription& Query, Select& Target, Functor)
-		{ 
+		{
 			RegisterFunctionArguments(Query, Target, &Functor::operator());
 		}
 
-		template <typename T>
-		struct HasStaticStructMethod
+		template<typename... Args>
+		struct IsValidSelectFunctionSignatureImpl2
 		{
-		private:
-			template <typename C> static constexpr bool Check(decltype(&C::StaticStruct)) { return std::is_same_v<decltype(&C::StaticStruct), UScriptStruct*(*)()>; }
-			template <typename C> static constexpr bool Check(...) { return false; }
-
-		public:
-			static constexpr bool Value = Check<T>(nullptr);
+			constexpr static bool Value = (IsValidColumnType<Args>() && ...);
 		};
 
-		template<typename Context>
-		constexpr bool IsValidContextType()
+		template<typename Arg0, typename... Args>
+		struct IsValidSelectFunctionSignatureImpl2<Arg0, Args...>
 		{
-			using BaseContextType = typename std::remove_const_t<typename TRemoveReference<Context>::Type>;
-			if constexpr (std::is_same_v<BaseContextType, QueryContext>)
-			{
-				return true;
-			}
-			else if constexpr (TIsDerivedFrom<BaseContextType, FQueryContextForwarder>::Value)
-			{
-				return true;
-			}
-			else
-			{
-				return false;
-			}
-		}
+			constexpr static bool Arg0IsContext = IsValidContextType<Arg0>();
+			constexpr static bool Arg0IsRowHandle = IsRowHandleType<Arg0>();
+			constexpr static bool Arg0IsArgument = IsValidColumnType<Arg0>();
+			constexpr static bool Arg0IsValidRowHandle = Arg0IsRowHandle && IsRowTypeCompatibleWithColumns<Arg0, Args...>();
+			constexpr static bool Arg0IsValid = (Arg0IsContext || Arg0IsValidRowHandle || Arg0IsArgument);
+			
+			static_assert(!Arg0IsRowHandle || Arg0IsValidRowHandle, "Row handles need to taken by value when the columns are requested by "
+				"reference or by const pointer if the columns are taken by reference.");
+			static_assert(Arg0IsValid, "The first provided argument has to be a reference to a compatible query reference, a row handle"
+				" or a (const) reference/pointer to a column.");
+			
+			constexpr static bool Value = Arg0IsValid && (IsValidColumnType<Args>() && ...);
+		};
 
-		template<typename Arg>
-		constexpr bool IsValidArgument()
+		template<typename Arg0, typename Arg1, typename... Args>
+		struct IsValidSelectFunctionSignatureImpl2<Arg0, Arg1, Args...>
 		{
-			if constexpr (TIsReferenceType<Arg>::Value)
-			{
-				using BaseType = typename std::remove_const_t<typename TRemoveReference<Arg>::Type>;
-				return HasStaticStructMethod<BaseType>::Value;
-			}
-			else
-			{
-				return false;
-			}
-		}
+			constexpr static bool Arg0IsContext = IsValidContextType<Arg0>();
+			constexpr static bool Arg0IsRowHandle = IsRowHandleType<Arg0>();
+			constexpr static bool Arg0IsArgument = IsValidColumnType<Arg0>();
+			constexpr static bool Arg0IsValidRowHandle = Arg0IsRowHandle && IsRowTypeCompatibleWithColumns<Arg0, Args...>();
+			constexpr static bool Arg0IsValid = (Arg0IsContext || Arg0IsValidRowHandle || Arg0IsArgument);
+			
+			static_assert(!Arg0IsRowHandle || Arg0IsValidRowHandle, "Row handles need to taken by value when the columns are requested by "
+				"reference or by const pointer if the columns are taken by reference.");
+			static_assert(Arg0IsValid, "The first argument to a query callback has to be a reference to a compatible query reference, a "
+				"row handle or a (const) reference/pointer to a column.");
+			
+			constexpr static bool Arg1IsRowHandle = IsRowHandleType<Arg1>();
+			constexpr static bool Arg1IsArgument = IsValidColumnType<Arg1>();
+			constexpr static bool Arg1IsValidRowHandle = Arg1IsRowHandle && IsRowTypeCompatibleWithColumns<Arg1, Args...>();
+			constexpr static bool Arg1IsValid = (Arg1IsValidRowHandle || Arg1IsArgument);
+
+			static_assert(!Arg1IsRowHandle || Arg1IsValidRowHandle, "Row handles need to taken by value when the columns are requested by "
+				"reference or by const pointer if the columns are taken by reference.");
+			static_assert(Arg1IsValid, "The second argument to a query callback has to be a row handle or a (const) reference/pointer "
+				"to a column.");
+			
+			constexpr static bool Value = Arg0IsValid && Arg1IsValid && (IsValidColumnType<Args>() && ...);
+		};
 
-		template <typename T>
-		class IsFunctor
+		template<>
+		struct IsValidSelectFunctionSignatureImpl2<>
 		{
-			template <typename C> static constexpr bool Check(decltype(&C::operator())) { return true; }
-			template <typename C> static constexpr bool Check(...) { return false; }
-
-		public:
-			static constexpr bool Value = Check<T>(nullptr);
+			constexpr static bool Value = false;
 		};
-
+				
 		template<typename T> struct IsValidSelectFunctionSignatureImpl
 		{ 
 			constexpr static bool Value = false; 
 		};
 
-		template<typename ContextType, typename... Args> struct IsValidSelectFunctionSignatureImpl<void (*)(ContextType&, Args...)>
+		template<typename... Args> struct IsValidSelectFunctionSignatureImpl<void (*)(Args...)>
 		{ 
-			constexpr static bool Value = IsValidContextType<ContextType>() && (IsValidArgument<Args>() && ...);
+			constexpr static bool Value = IsValidSelectFunctionSignatureImpl2<Args...>::Value;
 		};
 		
-		template<typename Class, typename ContextType, typename... Args> struct IsValidSelectFunctionSignatureImpl<void (Class::*)(ContextType&, Args...)>
+		template<typename Class, typename... Args> struct IsValidSelectFunctionSignatureImpl<void (Class::*)(Args...)>
 		{ 
-			constexpr static bool Value = IsValidContextType<ContextType>() && (IsValidArgument<Args>() && ...);
+			constexpr static bool Value = IsValidSelectFunctionSignatureImpl2<Args...>::Value;
 		};
 
-		template<typename Class, typename ContextType, typename... Args> struct IsValidSelectFunctionSignatureImpl<void (Class::*)(ContextType&, Args...) const>
+		template<typename Class, typename... Args> struct IsValidSelectFunctionSignatureImpl<void (Class::*)(Args...) const>
 		{ 
-			constexpr static bool Value = IsValidContextType<ContextType>() && (IsValidArgument<Args>() && ...);
+			constexpr static bool Value = IsValidSelectFunctionSignatureImpl2<Args...>::Value;
 		};
 
 		template<typename T> constexpr bool IsValidSelectFunctionSignature()
@@ -406,10 +913,14 @@ namespace TypedElementQueryBuilder
 			static_assert(IsValidSelectFunctionSignature<Function>(),
 				R"(The function provided to the Query Builder's Select call wasn't invocable or doesn't contain a supported combination of arguments.
 The following options are supported:
-- void(ITypedElementDataStorageInterface::FQueryContext&, [const]Column&...) 
-        e.g. void(ITypedElementDataStorageInterface::FQueryContext& Context, ColumnType0& ColumnA, const ColumnType1& ColumnB)
-- void(FCachedQueryContext<[const]Dependencies...>&, [const]Column&...)
-		e.g. void(FCachedQueryContext<Subsystem1, const Subsystem2>& Context, ColumnType0& ColumnA, const ColumnType1& ColumnB)
+- void([const]Column&...) 
+- void(TypedElementRowHandle, [const]Column&...) 
+- void(<Context>&, [const]Column&...) 
+- void(<Context>&, TypedElementRowHandle, [const]Column&...) 
+- void(<Context>&, [const]Column*...) 
+- void(<Context>&, const TypedElementRowHandle*, [const]Column*...) 
+Where <Context> is ITypedElementDataStorageInterface::FQueryContext or FCachedQueryContext<...>
+e.g. void(FCachedQueryContext<Subsystem1, const Subsystem2>& Context, TypedElementRowHandle Row, ColumnType0& ColumnA, const ColumnType1& ColumnB) {...}
 )");
 			RegisterFunctionArguments(Query, Target, Callback);
 			PrepareForQueryBinding(Query, Type);
@@ -418,7 +929,7 @@ The following options are supported:
 	}
 
 	template<typename CallbackType, typename Function>
-	Select::Select(FName Name, const CallbackType& Type, Function Callback)
+	Select::Select(FName Name, const CallbackType& Type, Function&& Callback)
 		: Select()
 	{
 		Internal::PrepareForQueryBinding(*this, Query, Name, Type, Callback);
@@ -426,7 +937,7 @@ The following options are supported:
 	}
 
 	template<typename CallbackType, typename Class, typename Function>
-	Select::Select(FName Name, const CallbackType& Type, Class* Instance, Function Callback)
+	Select::Select(FName Name, const CallbackType& Type, Class* Instance, Function&& Callback)
 		: Select()
 	{
 		Internal::PrepareForQueryBinding(*this, Query, Name, Type, Callback);
diff --git a/Engine/Source/Runtime/TypedElementFramework/Public/Elements/Interfaces/TypedElementDataStorageInterface.h b/Engine/Source/Runtime/TypedElementFramework/Public/Elements/Interfaces/TypedElementDataStorageInterface.h
--- a/Engine/Source/Runtime/TypedElementFramework/Public/Elements/Interfaces/TypedElementDataStorageInterface.h
+++ b/Engine/Source/Runtime/TypedElementFramework/Public/Elements/Interfaces/TypedElementDataStorageInterface.h
@@ -160,16 +160,34 @@ public:
 	{
 		/** The standard group to run work in. */
 		Default,
+		
+		/**
+		 * Prepares to sync data from an external source to the Data Storage. This is typically used for preparing columns with
+		 * data that is needed to sync from external sources.
+		 */
+		PrepareSyncExternalToDataStorage,
 		/** 
 		 * The group for processors that need to sync data from external sources such as subsystems or the world into
 		 * the Data Storage. These typically run early in a phase.
 		 */
 		SyncExternalToDataStorage,
 		/**
-		 * The group for processors that need to sync data from the Data Storage to external sources such as subsystems 
+		 * Finalizes the steps to sync data from an external source to the Data Storage. This is typically needed for processing 
+		 * data with a dependency on external updates being completed.
+		 */
+		FinalizeSyncExternalToDataStorage,
+
+		
+		/** Prepares for data from the Data Storage to be synced to external sources such as subsystem or the world. */
+		PrepareSyncDataStorageToExternal,
+		/**
+		 * The group for processors that need to sync data from the Data Storage to external sources such as subsystems
 		 * or the world into. These typically run late in a phase.
 		 */
 		SyncDataStorageToExternal,
+		/** Finalizes the sync of data from the Data Storage to external sources such as subsystem or the world. */
+		FinalizeSyncDataStorageToExternal,
+
 		/**
 		 * Processors grouped under this name will sync data to/from widgets.
 		 */
@@ -194,24 +212,92 @@ public:
 		ReadWrite 
 	};
 
-	/** Base class to be provided to query callbacks. */
+	/** 
+	 * Base class to be provided to query callbacks.
+	 * Note that at the time of writing only subclasses of Subsystem are supported as dependencies.
+	 */
 	struct FQueryContext
 	{
 		virtual ~FQueryContext() = default;
 		
+		/** Return the address of a immutable column matching the requested type or a nullptr if not found. */
 		virtual const void* GetColumn(const UScriptStruct* ColumnType) const = 0;
+		/** Return the address of a mutable column matching the requested type or a nullptr if not found. */
 		virtual void* GetMutableColumn(const UScriptStruct* ColumnType) = 0;
-		virtual void GetColumns(TArrayView<char*> RetrievedAddresses, TConstArrayView<const UScriptStruct*> ColumnTypes,
+		/** 
+		 * Get a list of columns or nullptrs if the columntype wasn't found. Mutable addresses are returned and it's up to 
+		 * the caller to not change immutable addresses.
+		 */
+		virtual void GetColumns(TArrayView<char*> RetrievedAddresses, TConstArrayView<TWeakObjectPtr<const UScriptStruct>> ColumnTypes,
 			TConstArrayView<ITypedElementDataStorageInterface::EQueryAccessType> AccessTypes) = 0;
-		virtual void GetColumnsUnguarded(int32 TypeCount, char** RetrievedAddresses, const UScriptStruct* const* ColumnTypes, 
+		/**
+		 * Get a list of columns or nullptrs if the columntype wasn't found. Mutable addresses are returned and it's up to
+		 * the caller to not change immutable addresses. This version doesn't verify that the enough space is provided and
+		 * it's up to the caller to guarantee the target addresses have enough space.
+		 */
+		virtual void GetColumnsUnguarded(int32 TypeCount, char** RetrievedAddresses, const TWeakObjectPtr<const UScriptStruct>* ColumnTypes,
 			const ITypedElementDataStorageInterface::EQueryAccessType* AccessTypes) = 0;
 
-		virtual USubsystem* GetMutableSubsystem(const TSubclassOf<USubsystem> SubsystemClass) = 0;
-		virtual const USubsystem* GetSubsystem(const TSubclassOf<USubsystem> SubsystemClass) = 0;
-		virtual void GetSubsystems(TArrayView<char*> RetrievedAddresses, TConstArrayView<const TSubclassOf<USubsystem>> SubsystemTypes, 
+		/** Returns an immutable instance of the requested dependency or a nullptr if not found. */
+		virtual const UObject* GetDependency(const UClass* DependencyClass) = 0;
+		/** Returns a mutable instance of the requested dependency or a nullptr if not found. */
+		virtual UObject* GetMutableDependency(const UClass* DependencyClass) = 0;
+		/** 
+		 * Returns a list of dependencies or nullptrs if a dependency wasn't found. Mutable versions are return and it's up to the 
+		 * caller to not change immutable dependencies.
+		 */
+		virtual void GetDependencies(TArrayView<UObject*> RetrievedAddresses, TConstArrayView<TWeakObjectPtr<const UClass>> DependencyTypes,
 			TConstArrayView<EQueryAccessType> AccessTypes) = 0;
 
+		/** Returns the number rows in the batch. */
 		virtual uint32 GetRowCount() const = 0;
+		/** 
+		 * Returns an immutable view that contains the row handles for all returned results. The returned size will be the same  as the
+		 * value returned by GetRowCount().
+		 */
+		virtual TConstArrayView<TypedElementRowHandle> GetRowHandles() const = 0;
+		/**
+		 * Removes the row with the provided row handle. The removal will not be immediately done but delayed until the end of the tick 
+		 * group.
+		 */
+		virtual void RemoveRow(TypedElementRowHandle Row) = 0;
+		/**
+		 * Removes rows with the provided row handles. The removal will not be immediately done but delayed until the end of the tick
+		 * group.
+		 */
+		virtual void RemoveRows(TConstArrayView<TypedElementRowHandle> Rows) = 0;
+
+		/**
+		 * Adds new empty columns to a row of the provided type. The addition will not be immediately done but delayed until the end of the
+		 * tick group.
+		 */
+		virtual void AddColumns(TypedElementRowHandle Row, TConstArrayView<const UScriptStruct*> ColumnTypes) = 0;
+		/**
+		 * Adds new empty columns to the listed rows of the provided type. The addition will not be immediately done but delayed until the end of the
+		 * tick group.
+		 */
+		virtual void AddColumns(TConstArrayView<TypedElementRowHandle> Rows, TConstArrayView<const UScriptStruct*> ColumnTypes) = 0;
+		/**
+		 * Removes columns of the provided types from a row. The removal will not be immediately done but delayed until the end of the
+		 * tick group.
+		 */
+		virtual void RemoveColumns(TypedElementRowHandle Row, TConstArrayView<const UScriptStruct*> ColumnTypes) = 0;
+		/**
+		 * Removes columns of the provided types from the listed rows. The removal will not be immediately done but delayed until the end of the
+		 * tick group.
+		 */
+		virtual void RemoveColumns(TConstArrayView<TypedElementRowHandle> Rows, TConstArrayView<const UScriptStruct*> ColumnTypes) = 0;
+
+		// Utility functions
+
+		template<typename... Columns>
+		void AddColumns(TypedElementRowHandle Row);
+		template<typename... Columns>
+		void AddColumns(TConstArrayView<TypedElementRowHandle> Rows);
+		template<typename... Columns>
+		void RemoveColumns(TypedElementRowHandle Row);
+		template<typename... Columns>
+		void RemoveColumns(TConstArrayView<TypedElementRowHandle> Rows);
 	};
 
 	using QueryCallback = TFunction<void(FQueryContext&)>;
@@ -404,3 +490,28 @@ SystemType* ITypedElementDataStorageInterface::GetExternalSystem()
 {
 	return reinterpret_cast<SystemType*>(GetExternalSystemAddress(SystemType::StaticClass()));
 }
+
+
+template<typename... Columns>
+void ITypedElementDataStorageInterface::FQueryContext::AddColumns(TypedElementRowHandle Row)
+{
+	AddColumns(Row, { Columns::StaticStruct()... });
+}
+
+template<typename... Columns>
+void ITypedElementDataStorageInterface::FQueryContext::AddColumns(TConstArrayView<TypedElementRowHandle> Rows)
+{
+	AddColumns(Rows, { Columns::StaticStruct()... });
+}
+
+template<typename... Columns>
+void ITypedElementDataStorageInterface::FQueryContext::RemoveColumns(TypedElementRowHandle Row)
+{
+	RemoveColumns(Row, { Columns::StaticStruct()... });
+}
+
+template<typename... Columns>
+void ITypedElementDataStorageInterface::FQueryContext::RemoveColumns(TConstArrayView<TypedElementRowHandle> Rows)
+{
+	RemoveColumns(Rows, { Columns::StaticStruct()... });
+}
\ No newline at end of file
