Expanded utility of TEDS' queries outside the internal tick.

The following changes were made:
- The query description is now stored inside the stored information for future reference.
- It's now possible to retrieve the description of a registered query.
- All type information for queries is now stored using weak pointers for safety. This will also allow detecting queries that need to be removed when plugins are dynamically unloaded.
- The type information used by the query builder is now retrieved from the associated query description instead of being cached locally in a static array. This avoids stale data getting reloaded in the future and also improved the compilers ability to inline callbacks.
- RunQuery has been overloaded with a version that takes a callback. As a result it's no possible to run Select queries directly.
- Renamed "FQueryContext" to "IQueryContext" to highlight that this is an interface.
- Added "IDirectQueryContext" as the context for callbacks that are called through RunQuery. Not all functionality between a direct and tick query callback universally supported, so using two different contexts allows for better control over what is and isn't available.
- Added a new CVar command "TEDS.PrintActorLabels" that will print the names of all actors stored in the data storage.
- Replaced the query dependency read-only/read-write state with EQueryDependencyFlags to allow storing of how to deal with the refreshing cached dependency addresses and whether or not it runs on the main thread. This has not been integrated into the query function binding yet, but can otherwise be used.
- Reduced the number of selections, conditions and dependencies can be stored inline a query. It was much larger than needed for the current set of queries. If any of those fields grow to large memory will be dynamically allocated but otherwise has no impact. The numbers can easily be adjusted if new queries frequently go over the inlined storage.
- Added several template functions to make it more convenient to directly do column operations with a query context.

diff --git a/Engine/Plugins/Experimental/TypedElementsDataStorage/Source/TypedElementsDataStorage/Private/Compatibility/TypedElementActorLabelProcessors.cpp b/Engine/Plugins/Experimental/TypedElementsDataStorage/Source/TypedElementsDataStorage/Private/Compatibility/TypedElementActorLabelProcessors.cpp
--- a/Engine/Plugins/Experimental/TypedElementsDataStorage/Source/TypedElementsDataStorage/Private/Compatibility/TypedElementActorLabelProcessors.cpp
+++ b/Engine/Plugins/Experimental/TypedElementsDataStorage/Source/TypedElementsDataStorage/Private/Compatibility/TypedElementActorLabelProcessors.cpp
@@ -5,9 +5,69 @@
 #include "Elements/Columns/TypedElementMiscColumns.h"
 #include "Elements/Columns/TypedElementLabelColumns.h"
 #include "Elements/Framework/TypedElementQueryBuilder.h"
+#include "Elements/Framework/TypedElementRegistry.h"
+#include "HAL/IConsoleManager.h"
 #include "Hash/CityHash.h"
 #include "MassActorSubsystem.h"
+#include "ProfilingDebugging/CpuProfilerTrace.h"
 
+FAutoConsoleCommandWithOutputDevice PrintActorLabelsConsoleCommand(
+	TEXT("TEDS.PrintActorLabels"),
+	TEXT("Prints out the labels for all actors found in the Typed Elements Data Storage."),
+	FConsoleCommandWithOutputDeviceDelegate::CreateLambda([](FOutputDevice& Output)
+		{
+			using namespace TypedElementQueryBuilder;
+			using DSI = ITypedElementDataStorageInterface;
+
+			TRACE_CPUPROFILER_EVENT_SCOPE(TEDS.PrintActorLabelsCommand);
+
+			if (ITypedElementDataStorageInterface* DataStorage = UTypedElementRegistry::GetInstance()->GetMutableDataStorage())
+			{
+				static TypedElementQueryHandle LabelQuery = TypedElementInvalidQueryHandle;
+				if (LabelQuery == TypedElementInvalidQueryHandle)
+				{
+					LabelQuery = DataStorage->RegisterQuery(
+						Select()
+							.ReadOnly<FTypedElementLabelColumn>()
+						.Where()
+							.All<FMassActorFragment>()
+						.Compile());
+				}
+				
+				if (LabelQuery != TypedElementInvalidQueryHandle)
+				{
+					FString Message;
+					Output.Log(TEXT("The Typed Elements Data Storage has the following actors:"));
+					DataStorage->RunQuery(LabelQuery, [&Output, &Message](const DSI::FQueryDescription&, DSI::IDirectQueryContext& Context)
+						{
+							const FTypedElementLabelColumn* Labels = Context.GetColumn<FTypedElementLabelColumn>();
+							const uint32 Count = Context.GetRowCount();
+
+							const FTypedElementLabelColumn* LabelsIt = Labels;
+							int32 CharacterCount = 0;
+							// Reserve memory first to avoid repeated memory allocations.
+							for (uint32 Index = 0; Index < Count; ++Index)
+							{
+								CharacterCount += 12 /*Prefixed text size*/ + LabelsIt->Label.Len() + 1 /*Trailing new line*/;
+								++LabelsIt;
+							}
+							Message.Reset(CharacterCount);
+
+							LabelsIt = Labels;
+							for (uint32 Index = 0; Index < Count; ++Index)
+							{
+								Message += TEXT("    Actor: ");
+								Message += LabelsIt->Label;
+								Message += TEXT('\n');
+								++LabelsIt;
+							}
+
+							Output.Log(Message);
+						});
+					Output.Log(TEXT("End of Typed Elements Data Storage actors list."));
+				}
+			}
+		}));
 
 void UTypedElementActorLabelFactory::RegisterQueries(ITypedElementDataStorageInterface& DataStorage) const
 {
@@ -53,7 +113,7 @@ void UTypedElementActorLabelFactory::RegisterLabelColumnToActorQuery(ITypedEleme
 			TEXT("Sync label column to actor"),
 			FProcessor(DSI::EQueryTickPhase::FrameEnd, DataStorage.GetQueryTickGroupName(DSI::EQueryTickGroups::SyncDataStorageToExternal))
 				.ForceToGameThread(true),
-			[](DSI::FQueryContext&, FMassActorFragment& Actor, const FTypedElementLabelColumn& Label, const FTypedElementLabelHashColumn& LabelHash)
+			[](FMassActorFragment& Actor, const FTypedElementLabelColumn& Label, const FTypedElementLabelHashColumn& LabelHash)
 			{
 				if (AActor* ActorInstance = Actor.GetMutable(); ActorInstance != nullptr)
 				{
diff --git a/Engine/Plugins/Experimental/TypedElementsDataStorage/Source/TypedElementsDataStorage/Private/Compatibility/TypedElementActorTransformProcessors.cpp b/Engine/Plugins/Experimental/TypedElementsDataStorage/Source/TypedElementsDataStorage/Private/Compatibility/TypedElementActorTransformProcessors.cpp
--- a/Engine/Plugins/Experimental/TypedElementsDataStorage/Source/TypedElementsDataStorage/Private/Compatibility/TypedElementActorTransformProcessors.cpp
+++ b/Engine/Plugins/Experimental/TypedElementsDataStorage/Source/TypedElementsDataStorage/Private/Compatibility/TypedElementActorTransformProcessors.cpp
@@ -27,7 +27,7 @@ void UTypedElementActorTransformFactory::RegisterActorAddTransformColumn(ITypedE
 				DataStorage.GetQueryTickGroupName(DSI::EQueryTickGroups::PrepareSyncExternalToDataStorage))
 				.SetBeforeGroup(DataStorage.GetQueryTickGroupName(DSI::EQueryTickGroups::SyncExternalToDataStorage))
 				.ForceToGameThread(true),
-			[](DSI::FQueryContext& Context, TypedElementRowHandle Row, const FMassActorFragment& Actor)
+			[](DSI::IQueryContext& Context, TypedElementRowHandle Row, const FMassActorFragment& Actor)
 			{
 				const AActor* ActorInstance = Actor.Get();
 				if (ActorInstance != nullptr && ActorInstance->GetRootComponent())
@@ -52,7 +52,7 @@ void UTypedElementActorTransformFactory::RegisterActorLocalTransformToColumn(ITy
 			TEXT("Sync actor transform to column"),
 			FProcessor(DSI::EQueryTickPhase::PrePhysics, DataStorage.GetQueryTickGroupName(DSI::EQueryTickGroups::SyncExternalToDataStorage))
 				.ForceToGameThread(true),
-			[](DSI::FQueryContext& Context, TypedElementRowHandle Row, const FMassActorFragment& Actor, FTypedElementLocalTransformColumn& Transform)
+			[](DSI::IQueryContext& Context, TypedElementRowHandle Row, const FMassActorFragment& Actor, FTypedElementLocalTransformColumn& Transform)
 			{
 				const AActor* ActorInstance = Actor.Get();
 				if (ActorInstance != nullptr && ActorInstance->GetRootComponent() != nullptr)
diff --git a/Engine/Plugins/Experimental/TypedElementsDataStorage/Source/TypedElementsDataStorage/Private/Processors/TypedElementMiscProcessors.cpp b/Engine/Plugins/Experimental/TypedElementsDataStorage/Source/TypedElementsDataStorage/Private/Processors/TypedElementMiscProcessors.cpp
--- a/Engine/Plugins/Experimental/TypedElementsDataStorage/Source/TypedElementsDataStorage/Private/Processors/TypedElementMiscProcessors.cpp
+++ b/Engine/Plugins/Experimental/TypedElementsDataStorage/Source/TypedElementsDataStorage/Private/Processors/TypedElementMiscProcessors.cpp
@@ -16,7 +16,7 @@ void UTypedElementRemoveSyncToWorldTagFactory::RegisterQueries(ITypedElementData
 			FProcessor(DSI::EQueryTickPhase::FrameEnd, 
 				DataStorage.GetQueryTickGroupName(DSI::EQueryTickGroups::FinalizeSyncDataStorageToExternal))
 				.SetAfterGroup(DataStorage.GetQueryTickGroupName(DSI::EQueryTickGroups::SyncDataStorageToExternal)),
-			[](DSI::FQueryContext& Context, const TypedElementRowHandle* Rows)
+			[](DSI::IQueryContext& Context, const TypedElementRowHandle* Rows)
 			{
 				Context.RemoveColumns<FTypedElementSyncBackToWorldTag>(TConstArrayView<TypedElementRowHandle>(Rows, Context.GetRowCount()));
 			}
diff --git a/Engine/Plugins/Experimental/TypedElementsDataStorage/Source/TypedElementsDataStorage/Private/Processors/TypedElementProcessorAdaptors.cpp b/Engine/Plugins/Experimental/TypedElementsDataStorage/Source/TypedElementsDataStorage/Private/Processors/TypedElementProcessorAdaptors.cpp
--- a/Engine/Plugins/Experimental/TypedElementsDataStorage/Source/TypedElementsDataStorage/Private/Processors/TypedElementProcessorAdaptors.cpp
+++ b/Engine/Plugins/Experimental/TypedElementsDataStorage/Source/TypedElementsDataStorage/Private/Processors/TypedElementProcessorAdaptors.cpp
@@ -7,36 +7,20 @@
 #include "MassExecutionContext.h"
 #include "TypedElementDatabase.h"
 
-struct FMassContextForwarder final : public ITypedElementDataStorageInterface::FQueryContext
+struct FMassContextForwarderShared
 {
-	explicit FMassContextForwarder(FMassExecutionContext& Context)
-		: Context(Context)
-		, RowCount(Context.GetNumEntities())
-	{}
-
-	const void* GetColumn(const UScriptStruct* ColumnType) const override
+	static const void* GetColumn(const FMassExecutionContext& Context, const UScriptStruct* ColumnType)
 	{
 		return Context.GetFragmentView(ColumnType).GetData();
 	}
 
-	void* GetMutableColumn(const UScriptStruct* ColumnType) override
+	static void* GetMutableColumn(FMassExecutionContext& Context, const UScriptStruct* ColumnType)
 	{
 		return Context.GetMutableFragmentView(ColumnType).GetData();
 	}
 
-	void GetColumns(TArrayView<char*> RetrievedAddresses, TConstArrayView<TWeakObjectPtr<const UScriptStruct>> ColumnTypes,
-		TConstArrayView<ITypedElementDataStorageInterface::EQueryAccessType> AccessTypes) override
-	{
-		checkf(RetrievedAddresses.Num() == ColumnTypes.Num(), TEXT("Unable to retrieve a batch of columns as the number of addresses "
-			"doesn't match the number of requested column."));
-		checkf(RetrievedAddresses.Num() == AccessTypes.Num(), TEXT("Unable to retrieve a batch of columns as the number of addresses "
-			"doesn't match the number of access types."));
-		
-		GetColumnsUnguarded(ColumnTypes.Num(), RetrievedAddresses.GetData(), ColumnTypes.GetData(), AccessTypes.GetData());
-	}
-
-	void GetColumnsUnguarded(int32 TypeCount, char** RetrievedAddresses, const TWeakObjectPtr<const UScriptStruct>* ColumnTypes,
-		const ITypedElementDataStorageInterface::EQueryAccessType* AccessTypes) override
+	static void GetColumnsUnguarded(FMassExecutionContext& Context, int32 TypeCount, char** RetrievedAddresses,
+		const TWeakObjectPtr<const UScriptStruct>* ColumnTypes, const ITypedElementDataStorageInterface::EQueryAccessType* AccessTypes)
 	{
 		for (int32 Index = 0; Index < TypeCount; ++Index)
 		{
@@ -51,6 +35,63 @@ struct FMassContextForwarder final : public ITypedElementDataStorageInterface::F
 		}
 	}
 
+	static void GetColumns(FMassExecutionContext& Context, TArrayView<char*> RetrievedAddresses, 
+		TConstArrayView<TWeakObjectPtr<const UScriptStruct>> ColumnTypes,
+		TConstArrayView<ITypedElementDataStorageInterface::EQueryAccessType> AccessTypes)
+	{
+		checkf(RetrievedAddresses.Num() == ColumnTypes.Num(), TEXT("Unable to retrieve a batch of columns as the number of addresses "
+			"doesn't match the number of requested column."));
+		checkf(RetrievedAddresses.Num() == AccessTypes.Num(), TEXT("Unable to retrieve a batch of columns as the number of addresses "
+			"doesn't match the number of access types."));
+
+		GetColumnsUnguarded(Context, ColumnTypes.Num(), RetrievedAddresses.GetData(), ColumnTypes.GetData(), AccessTypes.GetData());
+	}
+
+	static uint32 GetRowCount(const FMassExecutionContext& Context)
+	{
+		return Context.GetNumEntities();
+	}
+
+	static TConstArrayView<TypedElementRowHandle> GetRowHandles(const FMassExecutionContext& Context)
+	{
+		static_assert(
+			sizeof(TypedElementRowHandle) == sizeof(FMassEntityHandle) && alignof(TypedElementRowHandle) == alignof(FMassEntityHandle),
+			"TypedElementRowHandle and FMassEntityHandle need to by layout compatible to support Typed Elements Data Storage.");
+		TConstArrayView<FMassEntityHandle> Entities = Context.GetEntities();
+		return TConstArrayView<TypedElementRowHandle>(reinterpret_cast<const TypedElementRowHandle*>(Entities.GetData()), Entities.Num());
+	}
+};
+
+struct FMassContextForwarder final : public ITypedElementDataStorageInterface::IQueryContext
+{
+	explicit FMassContextForwarder(FMassExecutionContext& Context)
+		: Context(Context)
+	{}
+
+	~FMassContextForwarder() override = default;
+
+	const void* GetColumn(const UScriptStruct* ColumnType) const override
+	{
+		return FMassContextForwarderShared::GetColumn(Context, ColumnType);
+	}
+
+	void* GetMutableColumn(const UScriptStruct* ColumnType) override
+	{
+		return FMassContextForwarderShared::GetMutableColumn(Context, ColumnType);
+	}
+
+	void GetColumns(TArrayView<char*> RetrievedAddresses, TConstArrayView<TWeakObjectPtr<const UScriptStruct>> ColumnTypes,
+		TConstArrayView<ITypedElementDataStorageInterface::EQueryAccessType> AccessTypes) override
+	{
+		FMassContextForwarderShared::GetColumns(Context, RetrievedAddresses, ColumnTypes, AccessTypes);
+	}
+
+	void GetColumnsUnguarded(int32 TypeCount, char** RetrievedAddresses, const TWeakObjectPtr<const UScriptStruct>* ColumnTypes,
+		const ITypedElementDataStorageInterface::EQueryAccessType* AccessTypes) override
+	{
+		FMassContextForwarderShared::GetColumnsUnguarded(Context, TypeCount, RetrievedAddresses, ColumnTypes, AccessTypes);
+	}
+
 	UObject* GetMutableDependency(const UClass* DependencyClass) override
 	{
 		return Context.GetMutableSubsystem<USubsystem>(const_cast<UClass*>(DependencyClass));
@@ -88,16 +129,12 @@ struct FMassContextForwarder final : public ITypedElementDataStorageInterface::F
 
 	uint32 GetRowCount() const override
 	{
-		return RowCount;
+		return FMassContextForwarderShared::GetRowCount(Context);
 	}
 
 	TConstArrayView<TypedElementRowHandle> GetRowHandles() const
 	{
-		static_assert(
-			sizeof(TypedElementRowHandle) == sizeof(FMassEntityHandle) && alignof(TypedElementRowHandle) == alignof(FMassEntityHandle),
-			"TypedElementRowHandle and FMassEntityHandle need to by layout compatible to support Typed Elements Data Storage.");
-		TConstArrayView<FMassEntityHandle> Entities = Context.GetEntities();
-		return TConstArrayView<TypedElementRowHandle>(reinterpret_cast<const TypedElementRowHandle*>(Entities.GetData()), Entities.Num());
+		return FMassContextForwarderShared::GetRowHandles(Context);
 	}
 
 	void RemoveRow(TypedElementRowHandle Row) override
@@ -186,7 +223,49 @@ struct FMassContextForwarder final : public ITypedElementDataStorageInterface::F
 	}
 
 	FMassExecutionContext& Context;
-	uint32 RowCount{ 0 };
+};
+
+struct FMassDirectContextForwarder final : public ITypedElementDataStorageInterface::IDirectQueryContext
+{
+	explicit FMassDirectContextForwarder(FMassExecutionContext& Context)
+		: Context(Context)
+	{}
+
+	~FMassDirectContextForwarder() override = default;
+
+	FMassExecutionContext& Context;
+
+	const void* GetColumn(const UScriptStruct* ColumnType) const override
+	{
+		return FMassContextForwarderShared::GetColumn(Context, ColumnType);
+	}
+
+	void* GetMutableColumn(const UScriptStruct* ColumnType) override
+	{
+		return FMassContextForwarderShared::GetMutableColumn(Context, ColumnType);
+	}
+
+	void GetColumns(TArrayView<char*> RetrievedAddresses, TConstArrayView<TWeakObjectPtr<const UScriptStruct>> ColumnTypes,
+		TConstArrayView<ITypedElementDataStorageInterface::EQueryAccessType> AccessTypes) override
+	{
+		FMassContextForwarderShared::GetColumns(Context, RetrievedAddresses, ColumnTypes, AccessTypes);
+	}
+
+	void GetColumnsUnguarded(int32 TypeCount, char** RetrievedAddresses, const TWeakObjectPtr<const UScriptStruct>* ColumnTypes,
+		const ITypedElementDataStorageInterface::EQueryAccessType* AccessTypes) override
+	{
+		FMassContextForwarderShared::GetColumnsUnguarded(Context, TypeCount, RetrievedAddresses, ColumnTypes, AccessTypes);
+	}
+
+	uint32 GetRowCount() const override
+	{
+		return FMassContextForwarderShared::GetRowCount(Context);
+	}
+
+	TConstArrayView<TypedElementRowHandle> GetRowHandles() const
+	{
+		return FMassContextForwarderShared::GetRowHandles(Context);
+	}
 };
 
 
@@ -218,18 +297,82 @@ EMassProcessingPhase FTypedElementQueryProcessorData::MapToMassProcessingPhase(I
 
 FString FTypedElementQueryProcessorData::GetProcessorName() const
 {
-	return ParentQuery ? ParentQuery->Callback.Name.ToString() : FString(TEXT("<unnamed>"));
+	return ParentQuery ? ParentQuery->Description.Callback.Name.ToString() : FString(TEXT("<unnamed>"));
+}
+
+bool FTypedElementQueryProcessorData::PrepareCachedDependenciesOnParentQuery(
+	ITypedElementDataStorageInterface::FQueryDescription& Description, FMassExecutionContext& Context)
+{
+	const int32 DependencyCount = Description.DependencyTypes.Num();
+	TWeakObjectPtr<const UClass>* Types = Description.DependencyTypes.GetData();
+	ITypedElementDataStorageInterface::EQueryDependencyFlags* Flags = Description.DependencyFlags.GetData();
+	TWeakObjectPtr<UObject>* Caches = Description.CachedDependencies.GetData();
+
+	for (int32 Index = 0; Index < DependencyCount; ++Index)
+	{
+		checkf(Types->IsValid(), TEXT("Attempting to retrieve a dependency type that's no longer available."));
+		
+		if (EnumHasAnyFlags(*Flags, ITypedElementDataStorageInterface::EQueryDependencyFlags::AlwaysRefresh) || !Caches->IsValid())
+		*Caches = EnumHasAnyFlags(*Flags, ITypedElementDataStorageInterface::EQueryDependencyFlags::ReadOnly)
+			? const_cast<USubsystem*>(Context.GetSubsystem<USubsystem>(const_cast<UClass*>(Types->Get())))
+			: Context.GetMutableSubsystem<USubsystem>(const_cast<UClass*>(Types->Get()));
+
+		if (*Caches != nullptr)
+		{
+			++Types;
+			++Flags;
+			++Caches;
+		}
+		else
+		{
+			checkf(false, TEXT("Unable to retrieve instance of depencendy '%s'."), *((*Types)->GetName()));
+			return false;
+		}
+	}
+	return true;
+}
+
+ITypedElementDataStorageInterface::FQueryResult FTypedElementQueryProcessorData::Execute(
+	ITypedElementDataStorageInterface::DirectQueryCallbackRef& Callback,
+	ITypedElementDataStorageInterface::FQueryDescription& Description, 
+	FMassEntityQuery& NativeQuery, 
+	FMassEntityManager& EntityManager)
+{
+	FMassExecutionContext Context(EntityManager);
+	ITypedElementDataStorageInterface::FQueryResult Result;
+	Result.Completed = ITypedElementDataStorageInterface::FQueryResult::ECompletion::Fully;
+	
+	NativeQuery.ForEachEntityChunk(EntityManager, Context,
+		[&Result, &Callback, &Description](FMassExecutionContext& Context)
+		{
+			if (PrepareCachedDependenciesOnParentQuery(Description, Context))
+			{
+				FMassDirectContextForwarder QueryContext(Context);
+				Callback(Description, QueryContext);
+				Result.Count += Context.GetNumEntities();
+			}
+			else
+			{
+				Result.Completed = ITypedElementDataStorageInterface::FQueryResult::ECompletion::MissingDependency;
+			}
+		}
+	);
+	return Result;
 }
 
 void FTypedElementQueryProcessorData::Execute(FMassEntityManager& EntityManager, FMassExecutionContext& Context)
 {
 	checkf(ParentQuery, TEXT("A query callback was registered for execution without an associated query."));
 	
-	Query.ForEachEntityChunk(EntityManager, Context, 
-		[this](FMassExecutionContext& Context)
+	ITypedElementDataStorageInterface::FQueryDescription& Description = ParentQuery->Description;
+	Query.ForEachEntityChunk(EntityManager, Context,
+		[&Description](FMassExecutionContext& Context)
 		{
-			FMassContextForwarder QueryContext(Context);
-			ParentQuery->Callback.Function(QueryContext);
+			if (PrepareCachedDependenciesOnParentQuery(Description, Context))
+			{
+				FMassContextForwarder QueryContext(Context);
+				Description.Callback.Function(Description, QueryContext);
+			}
 		}
 	);
 }
@@ -256,18 +399,18 @@ void UTypedElementQueryProcessorCallbackAdapterProcessor::ConfigureQueryCallback
 {
 	Data.ParentQuery = &TargetParentQuery;
 
-	bRequiresGameThreadExecution = TargetParentQuery.Callback.bForceToGameThread;
+	bRequiresGameThreadExecution = TargetParentQuery.Description.Callback.bForceToGameThread;
 	ExecutionFlags = static_cast<int32>(EProcessorExecutionFlags::Editor); 
-	ExecutionOrder.ExecuteInGroup = TargetParentQuery.Callback.Group;
-	if (!TargetParentQuery.Callback.BeforeGroup.IsNone())
+	ExecutionOrder.ExecuteInGroup = TargetParentQuery.Description.Callback.Group;
+	if (!TargetParentQuery.Description.Callback.BeforeGroup.IsNone())
 	{
-		ExecutionOrder.ExecuteBefore.Add(TargetParentQuery.Callback.BeforeGroup);
+		ExecutionOrder.ExecuteBefore.Add(TargetParentQuery.Description.Callback.BeforeGroup);
 	}
-	if (!TargetParentQuery.Callback.AfterGroup.IsNone())
+	if (!TargetParentQuery.Description.Callback.AfterGroup.IsNone())
 	{
-		ExecutionOrder.ExecuteAfter.Add(TargetParentQuery.Callback.AfterGroup);
+		ExecutionOrder.ExecuteAfter.Add(TargetParentQuery.Description.Callback.AfterGroup);
 	}
-	ProcessingPhase = Data.MapToMassProcessingPhase(TargetParentQuery.Callback.Phase);
+	ProcessingPhase = Data.MapToMassProcessingPhase(TargetParentQuery.Description.Callback.Phase);
 
 	Super::PostInitProperties();
 }
@@ -326,12 +469,12 @@ void UTypedElementQueryObserverCallbackAdapterProcessor::ConfigureQueryCallback(
 {
 	Data.ParentQuery = &TargetParentQuery;
 
-	bRequiresGameThreadExecution = TargetParentQuery.Callback.bForceToGameThread;
+	bRequiresGameThreadExecution = TargetParentQuery.Description.Callback.bForceToGameThread;
 	ExecutionFlags = static_cast<int32>(EProcessorExecutionFlags::Editor);
 	
-	ObservedType = const_cast<UScriptStruct*>(TargetParentQuery.Callback.MonitoredType);
+	ObservedType = const_cast<UScriptStruct*>(TargetParentQuery.Description.Callback.MonitoredType);
 	
-	switch (TargetParentQuery.Callback.Type)
+	switch (TargetParentQuery.Description.Callback.Type)
 	{
 	case ITypedElementDataStorageInterface::EQueryCallbackType::ObserveAdd:
 		Operation = EMassObservedOperation::Add;
@@ -341,7 +484,7 @@ void UTypedElementQueryObserverCallbackAdapterProcessor::ConfigureQueryCallback(
 		break;
 	default:
 		checkf(false, TEXT("Query type %i is not supported from the observer processor adapter."),
-			static_cast<int>(TargetParentQuery.Callback.Type));
+			static_cast<int>(TargetParentQuery.Description.Callback.Type));
 		break;
 	}
 
diff --git a/Engine/Plugins/Experimental/TypedElementsDataStorage/Source/TypedElementsDataStorage/Private/Processors/TypedElementProcessorAdaptors.h b/Engine/Plugins/Experimental/TypedElementsDataStorage/Source/TypedElementsDataStorage/Private/Processors/TypedElementProcessorAdaptors.h
--- a/Engine/Plugins/Experimental/TypedElementsDataStorage/Source/TypedElementsDataStorage/Private/Processors/TypedElementProcessorAdaptors.h
+++ b/Engine/Plugins/Experimental/TypedElementsDataStorage/Source/TypedElementsDataStorage/Private/Processors/TypedElementProcessorAdaptors.h
@@ -22,10 +22,19 @@ struct FTypedElementQueryProcessorData
 	EMassProcessingPhase MapToMassProcessingPhase(ITypedElementDataStorageInterface::EQueryTickPhase Phase) const;
 	FString GetProcessorName() const;
 
+	static ITypedElementDataStorageInterface::FQueryResult Execute(
+		ITypedElementDataStorageInterface::DirectQueryCallbackRef& Callback,
+		ITypedElementDataStorageInterface::FQueryDescription& Description, 
+		FMassEntityQuery& NativeQuery, 
+		FMassEntityManager& EntityManager);
 	void Execute(FMassEntityManager& EntityManager, FMassExecutionContext& Context);
 
 	FTypedElementDatabaseExtendedQuery* ParentQuery{ nullptr };
 	FMassEntityQuery Query;
+
+private:
+	static bool PrepareCachedDependenciesOnParentQuery(
+		ITypedElementDataStorageInterface::FQueryDescription& Description, FMassExecutionContext& Context);
 };
 
 /**
diff --git a/Engine/Plugins/Experimental/TypedElementsDataStorage/Source/TypedElementsDataStorage/Private/TypedElementDatabase.cpp b/Engine/Plugins/Experimental/TypedElementsDataStorage/Source/TypedElementsDataStorage/Private/TypedElementDatabase.cpp
--- a/Engine/Plugins/Experimental/TypedElementsDataStorage/Source/TypedElementsDataStorage/Private/TypedElementDatabase.cpp
+++ b/Engine/Plugins/Experimental/TypedElementsDataStorage/Source/TypedElementsDataStorage/Private/TypedElementDatabase.cpp
@@ -9,6 +9,7 @@
 #include "MassSimulationSubsystem.h"
 #include "MassSubsystemAccess.h"
 #include "Processors/TypedElementProcessorAdaptors.h"
+#include "ProfilingDebugging/CpuProfilerTrace.h"
 #include "Stats/Stats2.h"
 #include "TickTaskManagerInterface.h"
 
@@ -16,7 +17,9 @@ const FName UTypedElementDatabase::TickGroupName_PrepareSyncWorldToMass(TEXT("Pr
 const FName UTypedElementDatabase::TickGroupName_FinalizeSyncWorldToMass(TEXT("FinalizeSyncWorldToMass"));
 const FName UTypedElementDatabase::TickGroupName_PrepareSyncMassToExternal(TEXT("PrepareSyncMassToExternal"));
 const FName UTypedElementDatabase::TickGroupName_FinalizeSyncMassToExternal(TEXT("FinalizeSyncMassToExternal"));
+const FName UTypedElementDatabase::TickGroupName_PrepareSyncWidget(TEXT("PrepareSyncWidgets"));
 const FName UTypedElementDatabase::TickGroupName_SyncWidget(TEXT("SyncWidgets"));
+const FName UTypedElementDatabase::TickGroupName_FinalizeSyncWidget(TEXT("FinalizeSyncWidgets"));
 
 void UTypedElementDatabase::Initialize()
 {
@@ -308,28 +311,31 @@ TypedElementQueryHandle UTypedElementDatabase::RegisterQuery(FQueryDescription&&
 
 	QueryStore::Handle Result = Queries.Emplace();
 	FTypedElementDatabaseExtendedQuery& StoredQuery = Queries.GetMutable(Result);
-	StoredQuery.Action = Query.Action;
-	StoredQuery.bSimpleQuery = Query.bSimpleQuery;
-	StoredQuery.Callback = MoveTemp(Query.Callback);
 	
 	FMassEntityQuery& NativeQuery = SetupNativeQuery(Query, StoredQuery);
 
 	if (Query.Action == FQueryDescription::EActionType::Count)
 	{
-		checkf(Query.Selection.IsEmpty(), TEXT("Count queries for the Typed Elements Data Storage can't have entries for selection."));
+		checkf(Query.SelectionTypes.IsEmpty(), TEXT("Count queries for the Typed Elements Data Storage can't have entries for selection."));
+		checkf(Query.SelectionAccessTypes.IsEmpty(), TEXT("Count queries for the Typed Elements Data Storage can't have entries for selection."));
 	}
 	else
 	{
-		for (const FQueryDescription::FAccessControlledStruct& SelectEntry : Query.Selection)
+		const int32 SelectionCount = Query.SelectionTypes.Num();
+		checkf(SelectionCount == Query.SelectionAccessTypes.Num(),
+			TEXT("The number of query selection types (%i) doesn't match the number of selection access types (%i)."),
+			SelectionCount, Query.SelectionAccessTypes.Num());
+		for (int SelectionIndex = 0; SelectionIndex < SelectionCount; ++SelectionIndex)
 		{
-			checkf(SelectEntry.Type, TEXT("Provided query selection type can not be null."));
+			TWeakObjectPtr<const UScriptStruct>& Type = Query.SelectionTypes[SelectionIndex];
+			checkf(Type.IsValid(), TEXT("Provided query selection type can not be null."));
 			checkf(
-				SelectEntry.Type->IsChildOf(FTypedElementDataStorageColumn::StaticStruct()) ||
-				SelectEntry.Type->IsChildOf(FMassFragment::StaticStruct()) || 
-				SelectEntry.Type->IsChildOf(FMassTag::StaticStruct()),
+				Type->IsChildOf(FTypedElementDataStorageColumn::StaticStruct()) ||
+				Type->IsChildOf(FMassFragment::StaticStruct()) || 
+				Type->IsChildOf(FMassTag::StaticStruct()),
 				TEXT("Provided query selection type '%s' is not based on FTypedElementDataStorageColumn or another supported base type."), 
-				*SelectEntry.Type->GetStructPathName().ToString());
-			NativeQuery.AddRequirement(SelectEntry.Type, LocalToNativeAccess(SelectEntry.Access));
+				*Type->GetStructPathName().ToString());
+			NativeQuery.AddRequirement(Type.Get(), LocalToNativeAccess(Query.SelectionAccessTypes[SelectionIndex]));
 		}
 	}
 
@@ -363,25 +369,33 @@ TypedElementQueryHandle UTypedElementDatabase::RegisterQuery(FQueryDescription&&
 			}
 			else if (Operand->Type->IsChildOf(FMassFragment::StaticStruct()))
 			{
-				NativeQuery.AddRequirement(Operand->Type, EMassFragmentAccess::None, Presence);
+				NativeQuery.AddRequirement(Operand->Type.Get(), EMassFragmentAccess::None, Presence);
 			}
 
 			++Operand;
 		}
 	}
 
-	for (const FQueryDescription::FAccessControlledClass& DependencyEntry : Query.Dependencies)
+	const int32 DependencyCount = Query.DependencyTypes.Num();
+	checkf(DependencyCount == Query.DependencyFlags.Num() && DependencyCount == Query.CachedDependencies.Num(),
+		TEXT("The number of query depedencies (%i) doesn't match the number of dependency access types (%i) and/or cached dependencies count (%i)."),
+		DependencyCount, Query.DependencyFlags.Num(), Query.CachedDependencies.Num());
+	for (int32 DependencyIndex = 0; DependencyIndex < DependencyCount; ++DependencyIndex)
 	{
-		checkf(DependencyEntry.Type, TEXT("Provided query dependcy type can not be null."));
-		checkf(DependencyEntry.Type->IsChildOf<USubsystem>(),
-			TEXT("Provided query dependency type '%s' is not based on USubSystem."), 
-			*DependencyEntry.Type->GetStructPathName().ToString());
+		TWeakObjectPtr<const UClass>& Type = Query.DependencyTypes[DependencyIndex];
+		checkf(Type.IsValid(), TEXT("Provided query dependcy type can not be null."));
+		checkf(Type->IsChildOf<USubsystem>(), TEXT("Provided query dependency type '%s' is not based on USubSystem."), 
+			*Type->GetStructPathName().ToString());
 		
-		constexpr bool bGameThreadOnly = true;
+		EQueryDependencyFlags Flags = Query.DependencyFlags[DependencyIndex];
 		NativeQuery.AddSubsystemRequirement(
-			const_cast<UClass*>(DependencyEntry.Type), LocalToNativeAccess(DependencyEntry.Access), bGameThreadOnly);
+			const_cast<UClass*>(Type.Get()), 
+			EnumHasAllFlags(Flags, EQueryDependencyFlags::ReadOnly) ? EMassFragmentAccess::ReadOnly : EMassFragmentAccess::ReadWrite,
+			EnumHasAllFlags(Flags, EQueryDependencyFlags::GameThreadBound));
 	}
 
+	StoredQuery.Description = MoveTemp(Query);
+
 	if (StoredQuery.Processor)
 	{
 		if (StoredQuery.Processor->IsA<UTypedElementQueryProcessorCallbackAdapterProcessor>())
@@ -404,7 +418,7 @@ TypedElementQueryHandle UTypedElementDatabase::RegisterQuery(FQueryDescription&&
 		else
 		{
 			checkf(false, TEXT("Query processor %s is of unsupported type %s."), 
-				*StoredQuery.Callback.Name.ToString(), *StoredQuery.Processor->GetSparseClassDataStruct()->GetName());
+				*Query.Callback.Name.ToString(), *StoredQuery.Processor->GetSparseClassDataStruct()->GetName());
 		}
 	}
 
@@ -435,7 +449,7 @@ void UTypedElementDatabase::UnregisterQuery(TypedElementQueryHandle Query)
 			else
 			{
 				checkf(false, TEXT("Query processor %s is of unsupported type %s."),
-					*QueryData.Callback.Name.ToString(), *QueryData.Processor->GetSparseClassDataStruct()->GetName());
+					*QueryData.Description.Callback.Name.ToString(), *QueryData.Processor->GetSparseClassDataStruct()->GetName());
 			}
 		}
 		else
@@ -447,6 +461,23 @@ void UTypedElementDatabase::UnregisterQuery(TypedElementQueryHandle Query)
 	Queries.Remove(Handle);
 }
 
+const ITypedElementDataStorageInterface::FQueryDescription& UTypedElementDatabase::GetQueryDescription(TypedElementQueryHandle Query) const
+{
+	QueryStore::Handle Handle;
+	Handle.Handle = Query;
+
+	if (Queries.IsAlive(Handle))
+	{
+		const FTypedElementDatabaseExtendedQuery& QueryData = Queries.Get(Handle);
+		return QueryData.Description;
+	}
+	else
+	{
+		static ITypedElementDataStorageInterface::FQueryDescription EmptyDescription;
+		return EmptyDescription;
+	}
+}
+
 FName UTypedElementDatabase::GetQueryTickGroupName(EQueryTickGroups Group) const
 {
 	switch (Group)
@@ -467,8 +498,12 @@ FName UTypedElementDatabase::GetQueryTickGroupName(EQueryTickGroups Group) const
 		case EQueryTickGroups::FinalizeSyncDataStorageToExternal:
 			return TickGroupName_FinalizeSyncMassToExternal;
 		
+		case EQueryTickGroups::PrepareSyncWidgets:
+			return TickGroupName_PrepareSyncWidget;
 		case EQueryTickGroups::SyncWidgets:
 			return TickGroupName_SyncWidget;
+		case EQueryTickGroups::FinalizeSyncWidgets:
+			return TickGroupName_FinalizeSyncWidget;
 
 		default:
 			checkf(false, TEXT("EQueryTickGroups value %i can't be translated to a group name by this Data Storage backend."), static_cast<int>(Group));
@@ -478,6 +513,8 @@ FName UTypedElementDatabase::GetQueryTickGroupName(EQueryTickGroups Group) const
 
 ITypedElementDataStorageInterface::FQueryResult UTypedElementDatabase::RunQuery(TypedElementQueryHandle Query)
 {
+	TRACE_CPUPROFILER_EVENT_SCOPE(TEDS.RunQuery);
+	
 	FQueryResult Result;
 
 	QueryStore::Handle Handle;
@@ -486,17 +523,15 @@ ITypedElementDataStorageInterface::FQueryResult UTypedElementDatabase::RunQuery(
 	if (Queries.IsAlive(Handle))
 	{
 		FTypedElementDatabaseExtendedQuery& QueryData = Queries.Get(Handle);
-		if (QueryData.bSimpleQuery)
+		if (QueryData.Description.bSimpleQuery)
 		{
-			switch (QueryData.Action)
+			switch (QueryData.Description.Action)
 			{
 				case FQueryDescription::EActionType::None:
 					Result.Completed = FQueryResult::ECompletion::Fully;
 					break;
 				case FQueryDescription::EActionType::Select:
-					checkf(false, TEXT("Support for this option will be coming in a future update."));
-					Result.Completed = FQueryResult::ECompletion::Unsupported;
-					break;
+					// Fallthrough: There's nothing to callback to, so only return the total count.
 				case FQueryDescription::EActionType::Count:
 					checkf(ActiveEditorEntityManager, 
 						TEXT("Unable to run Typed Element Data Storage query before a MASS Entity Manager has been assigned."));
@@ -522,6 +557,58 @@ ITypedElementDataStorageInterface::FQueryResult UTypedElementDatabase::RunQuery(
 	return Result;
 }
 
+ITypedElementDataStorageInterface::FQueryResult UTypedElementDatabase::RunQuery(
+	TypedElementQueryHandle Query, ITypedElementDataStorageInterface::DirectQueryCallbackRef Callback)
+{
+	TRACE_CPUPROFILER_EVENT_SCOPE(TEDS.RunQuery);
+
+	FQueryResult Result;
+
+	QueryStore::Handle Handle;
+	Handle.Handle = Query;
+
+	if (Queries.IsAlive(Handle))
+	{
+		FTypedElementDatabaseExtendedQuery& QueryData = Queries.Get(Handle);
+		if (QueryData.Description.bSimpleQuery)
+		{
+			switch (QueryData.Description.Action)
+			{
+			case FQueryDescription::EActionType::None:
+				Result.Completed = FQueryResult::ECompletion::Fully;
+				break;
+			case FQueryDescription::EActionType::Select:
+				checkf(ActiveEditorEntityManager,
+					TEXT("Unable to run Typed Element Data Storage query before a MASS Entity Manager has been assigned."));
+				Result = FTypedElementQueryProcessorData::Execute(
+					Callback, QueryData.Description, QueryData.NativeQuery, *ActiveEditorEntityManager);
+				break;
+			case FQueryDescription::EActionType::Count:
+				// Only the count is requested so no need to trigger the callback.
+				checkf(ActiveEditorEntityManager,
+					TEXT("Unable to run Typed Element Data Storage query before a MASS Entity Manager has been assigned."));
+				Result.Count = QueryData.NativeQuery.GetNumMatchingEntities(*ActiveEditorEntityManager);
+				Result.Completed = FQueryResult::ECompletion::Fully;
+				break;
+			default:
+				Result.Completed = FQueryResult::ECompletion::Unsupported;
+				break;
+			}
+		}
+		else
+		{
+			checkf(false, TEXT("Support for this option will be coming in a future update."));
+			Result.Completed = FQueryResult::ECompletion::Unsupported;
+		}
+	}
+	else
+	{
+		Result.Completed = FQueryResult::ECompletion::Unavailable;
+	}
+
+	return Result;
+}
+
 FTypedElementOnDataStorageUpdate& UTypedElementDatabase::OnUpdate()
 {
 	return OnUpdateDelegate;
diff --git a/Engine/Plugins/Experimental/TypedElementsDataStorage/Source/TypedElementsDataStorage/Private/TypedElementDatabase.h b/Engine/Plugins/Experimental/TypedElementsDataStorage/Source/TypedElementsDataStorage/Private/TypedElementDatabase.h
--- a/Engine/Plugins/Experimental/TypedElementsDataStorage/Source/TypedElementsDataStorage/Private/TypedElementDatabase.h
+++ b/Engine/Plugins/Experimental/TypedElementsDataStorage/Source/TypedElementsDataStorage/Private/TypedElementDatabase.h
@@ -18,11 +18,9 @@ class UWorld;
 
 struct FTypedElementDatabaseExtendedQuery
 {
-	FMassEntityQuery NativeQuery;
-	ITypedElementDataStorageInterface::FQueryDescription::EActionType Action{ ITypedElementDataStorageInterface::FQueryDescription::EActionType::None };
-	ITypedElementDataStorageInterface::FQueryDescription::FCallbackData Callback;
+	FMassEntityQuery NativeQuery; // Used if there's no processor bound.
+	ITypedElementDataStorageInterface::FQueryDescription Description;
 	TStrongObjectPtr<UMassProcessor> Processor;
-	bool bSimpleQuery{ false };
 };
 
 UCLASS()
@@ -68,8 +66,10 @@ public:
 
 	TypedElementQueryHandle RegisterQuery(FQueryDescription&& Query) override;
 	void UnregisterQuery(TypedElementQueryHandle Query) override;
+	const FQueryDescription& GetQueryDescription(TypedElementQueryHandle Query) const override;
 	FName GetQueryTickGroupName(EQueryTickGroups Group) const override;
 	FQueryResult RunQuery(TypedElementQueryHandle Query) override;
+	FQueryResult RunQuery(TypedElementQueryHandle Query, DirectQueryCallbackRef Callback) override;
 
 	FTypedElementOnDataStorageUpdate& OnUpdate() override;
 	bool IsAvailable() const override;
@@ -84,7 +84,9 @@ private:
 	static const FName TickGroupName_FinalizeSyncWorldToMass;
 	static const FName TickGroupName_PrepareSyncMassToExternal;
 	static const FName TickGroupName_FinalizeSyncMassToExternal;
+	static const FName TickGroupName_PrepareSyncWidget;
 	static const FName TickGroupName_SyncWidget;
+	static const FName TickGroupName_FinalizeSyncWidget;
 	
 	TArray<FMassArchetypeHandle> Tables;
 	TMap<FName, TypedElementTableHandle> TableNameLookup;
diff --git a/Engine/Plugins/Experimental/TypedElementsDataStorage/Source/TypedElementsDataStorage/Private/UI/TypedElementSlateWidgetReferenceColumnUpdateProcessor.cpp b/Engine/Plugins/Experimental/TypedElementsDataStorage/Source/TypedElementsDataStorage/Private/UI/TypedElementSlateWidgetReferenceColumnUpdateProcessor.cpp
--- a/Engine/Plugins/Experimental/TypedElementsDataStorage/Source/TypedElementsDataStorage/Private/UI/TypedElementSlateWidgetReferenceColumnUpdateProcessor.cpp
+++ b/Engine/Plugins/Experimental/TypedElementsDataStorage/Source/TypedElementsDataStorage/Private/UI/TypedElementSlateWidgetReferenceColumnUpdateProcessor.cpp
@@ -11,6 +11,7 @@ void UTypedElementSlateWidgetReferenceColumnUpdateFactory::RegisterQueries(IType
 	RegisterDeleteColumnOnWidgetDeleteQuery(DataStorage);
 }
 
+#pragma optimize("", off)
 void UTypedElementSlateWidgetReferenceColumnUpdateFactory::RegisterDeleteRowOnWidgetDeleteQuery(ITypedElementDataStorageInterface& DataStorage) const
 {
 	using namespace TypedElementQueryBuilder;
@@ -19,10 +20,8 @@ void UTypedElementSlateWidgetReferenceColumnUpdateFactory::RegisterDeleteRowOnWi
 	DataStorage.RegisterQuery(
 		Select(
 			TEXT("Delete row with deleted widget"),
-			FProcessor(DSI::EQueryTickPhase::PrePhysics, 
-				DataStorage.GetQueryTickGroupName(DSI::EQueryTickGroups::PrepareSyncExternalToDataStorage))
-				.SetBeforeGroup(DataStorage.GetQueryTickGroupName(DSI::EQueryTickGroups::SyncExternalToDataStorage)),
-			[](DSI::FQueryContext& Context, TypedElementRowHandle Row, const FTypedElementSlateWidgetReferenceColumn& WidgetReference)
+			FProcessor(DSI::EQueryTickPhase::FrameEnd, DataStorage.GetQueryTickGroupName(DSI::EQueryTickGroups::PrepareSyncWidgets)).ForceToGameThread(true),
+			[](DSI::IQueryContext& Context, TypedElementRowHandle Row, const FTypedElementSlateWidgetReferenceColumn& WidgetReference)
 			{
 				if (!WidgetReference.Widget.IsValid())
 				{
@@ -35,6 +34,7 @@ void UTypedElementSlateWidgetReferenceColumnUpdateFactory::RegisterDeleteRowOnWi
 		.Compile()
 	);
 }
+#pragma optimize("", on)
 
 void UTypedElementSlateWidgetReferenceColumnUpdateFactory::RegisterDeleteColumnOnWidgetDeleteQuery(ITypedElementDataStorageInterface& DataStorage) const
 {
@@ -44,10 +44,8 @@ void UTypedElementSlateWidgetReferenceColumnUpdateFactory::RegisterDeleteColumnO
 	DataStorage.RegisterQuery(
 		Select(
 			TEXT("Delete widget column for deleted widget"),
-			FProcessor(DSI::EQueryTickPhase::PrePhysics,
-				DataStorage.GetQueryTickGroupName(DSI::EQueryTickGroups::PrepareSyncExternalToDataStorage))
-				.SetBeforeGroup(DataStorage.GetQueryTickGroupName(DSI::EQueryTickGroups::SyncExternalToDataStorage)),
-			[](DSI::FQueryContext& Context, TypedElementRowHandle Row, const FTypedElementSlateWidgetReferenceColumn& WidgetReference)
+			FProcessor(DSI::EQueryTickPhase::FrameEnd, DataStorage.GetQueryTickGroupName(DSI::EQueryTickGroups::PrepareSyncWidgets)).ForceToGameThread(true),
+			[](DSI::IQueryContext& Context, TypedElementRowHandle Row, const FTypedElementSlateWidgetReferenceColumn& WidgetReference)
 			{
 				if (!WidgetReference.Widget.IsValid())
 				{
@@ -56,7 +54,7 @@ void UTypedElementSlateWidgetReferenceColumnUpdateFactory::RegisterDeleteColumnO
 			}
 		)
 		.Where()
-			.All<FTypedElementSlateWidgetReferenceDeletesRowTag>()
+			.None<FTypedElementSlateWidgetReferenceDeletesRowTag>()
 		.Compile()
 	);
 }
diff --git a/Engine/Source/Runtime/TypedElementFramework/Private/Elements/Framework/TypedElementQueryBuilder.cpp b/Engine/Source/Runtime/TypedElementFramework/Private/Elements/Framework/TypedElementQueryBuilder.cpp
--- a/Engine/Source/Runtime/TypedElementFramework/Private/Elements/Framework/TypedElementQueryBuilder.cpp
+++ b/Engine/Source/Runtime/TypedElementFramework/Private/Elements/Framework/TypedElementQueryBuilder.cpp
@@ -39,7 +39,9 @@ namespace TypedElementQueryBuilder
 	FDependency& FDependency::ReadOnly(const UClass* Target)
 	{
 		checkf(Target, TEXT("The Dependency section in the Typed Elements query builder doesn't support nullptrs as Read-Only input."));
-		Query->Dependencies.Emplace(Target, ITypedElementDataStorageInterface::EQueryAccessType::ReadOnly);
+		Query->DependencyTypes.Emplace(Target);
+		Query->DependencyFlags.Emplace(ITypedElementDataStorageInterface::EQueryDependencyFlags::ReadOnly);
+		Query->CachedDependencies.AddDefaulted();
 		return *this;
 	}
 
@@ -55,7 +57,9 @@ namespace TypedElementQueryBuilder
 	FDependency& FDependency::ReadWrite(const UClass* Target)
 	{
 		checkf(Target, TEXT("The Dependency section in the Typed Elements query builder doesn't support nullptrs as Read/Write input."));
-		Query->Dependencies.Emplace(Target, ITypedElementDataStorageInterface::EQueryAccessType::ReadWrite);
+		Query->DependencyTypes.Emplace(Target);
+		Query->DependencyFlags.Emplace(ITypedElementDataStorageInterface::EQueryDependencyFlags::None);
+		Query->CachedDependencies.AddDefaulted();
 		return *this;
 	}
 
@@ -232,7 +236,8 @@ namespace TypedElementQueryBuilder
 	Select& Select::ReadOnly(const UScriptStruct* Target)
 	{
 		checkf(Target, TEXT("The Select section in the Typed Elements query builder doesn't support nullptrs as Read-Only input."));
-		Query.Selection.Emplace(Target, ITypedElementDataStorageInterface::EQueryAccessType::ReadOnly);
+		Query.SelectionTypes.Emplace(Target);
+		Query.SelectionAccessTypes.Emplace(ITypedElementDataStorageInterface::EQueryAccessType::ReadOnly);
 		return *this;
 	}
 
@@ -248,7 +253,8 @@ namespace TypedElementQueryBuilder
 	Select& Select::ReadWrite(const UScriptStruct* Target)
 	{
 		checkf(Target, TEXT("The Select section in the Typed Elements query builder doesn't support nullptrs as Read/Write input."));
-		Query.Selection.Emplace(Target, ITypedElementDataStorageInterface::EQueryAccessType::ReadWrite);
+		Query.SelectionTypes.Emplace(Target);
+		Query.SelectionAccessTypes.Emplace(ITypedElementDataStorageInterface::EQueryAccessType::ReadWrite);
 		return *this;
 	}
 
diff --git a/Engine/Source/Runtime/TypedElementFramework/Public/Elements/Framework/TypedElementQueryBuilder.h b/Engine/Source/Runtime/TypedElementFramework/Public/Elements/Framework/TypedElementQueryBuilder.h
--- a/Engine/Source/Runtime/TypedElementFramework/Public/Elements/Framework/TypedElementQueryBuilder.h
+++ b/Engine/Source/Runtime/TypedElementFramework/Public/Elements/Framework/TypedElementQueryBuilder.h
@@ -46,7 +46,7 @@ class USubsystem;
  *	- void(<Context>&, TypedElementRowHandle, [const]Column&...) 
  *	- void(<Context>&, [const]Column*...) 
  *	- void(<Context>&, const TypedElementRowHandle*, [const]Column*...) 
- *	Where <Context> is ITypedElementDataStorageInterface::FQueryContext or FCachedQueryContext<...>	e.g.:
+ *	Where <Context> is ITypedElementDataStorageInterface::IQueryContext or FCachedQueryContext<...>	e.g.:
  *		void(
  *			FCachedQueryContext<Subsystem1, const Subsystem2>& Context, 
  *			TypedElementRowHandle Row, 
@@ -203,9 +203,11 @@ namespace TypedElementQueryBuilder
 
 	// Because this is a thin wrapper called from within a query callback, it's better to inline fully so all
 	// function pre/postambles can be optimized away.
-	struct FQueryContextForwarder : public ITypedElementDataStorageInterface::FQueryContext
+	struct FQueryContextForwarder : public ITypedElementDataStorageInterface::IQueryContext
 	{
-		inline explicit FQueryContextForwarder(ITypedElementDataStorageInterface::FQueryContext& InParentContext);
+		inline FQueryContextForwarder(
+			const ITypedElementDataStorageInterface::FQueryDescription& InDescription, 
+			ITypedElementDataStorageInterface::IQueryContext& InParentContext);
 		inline ~FQueryContextForwarder() = default;
 
 		inline const void* GetColumn(const UScriptStruct* ColumnType) const override;
@@ -230,13 +232,16 @@ namespace TypedElementQueryBuilder
 		inline void RemoveColumns(TypedElementRowHandle Row, TConstArrayView<const UScriptStruct*> ColumnTypes) override;
 		inline void RemoveColumns(TConstArrayView<TypedElementRowHandle> Rows, TConstArrayView<const UScriptStruct*> ColumnTypes) override;
 
-		ITypedElementDataStorageInterface::FQueryContext& ParentContext;
+		ITypedElementDataStorageInterface::IQueryContext& ParentContext;
+		const ITypedElementDataStorageInterface::FQueryDescription& Description;
 	};
 
 	template<typename... Dependencies>
 	struct FCachedQueryContext final : public FQueryContextForwarder
 	{
-		explicit FCachedQueryContext(ITypedElementDataStorageInterface::FQueryContext& InParentContext);
+		explicit FCachedQueryContext(
+			const ITypedElementDataStorageInterface::FQueryDescription& InDescription, 
+			ITypedElementDataStorageInterface::IQueryContext& InParentContext);
 		
 		static void Register(ITypedElementDataStorageInterface::FQueryDescription& Query);
 
@@ -244,8 +249,6 @@ namespace TypedElementQueryBuilder
 		Dependency& GetCachedMutableDependency();
 		template<typename Dependency>
 		const Dependency& GetCachedDependency() const;
-
-		UObject* DependencyAddresses[sizeof...(Dependencies)];
 	};
 
 	// Explicitly not following the naming convention in order to present this as a query that can be read as such.
diff --git a/Engine/Source/Runtime/TypedElementFramework/Public/Elements/Framework/TypedElementQueryBuilder.inl b/Engine/Source/Runtime/TypedElementFramework/Public/Elements/Framework/TypedElementQueryBuilder.inl
--- a/Engine/Source/Runtime/TypedElementFramework/Public/Elements/Framework/TypedElementQueryBuilder.inl
+++ b/Engine/Source/Runtime/TypedElementFramework/Public/Elements/Framework/TypedElementQueryBuilder.inl
@@ -38,6 +38,24 @@ namespace TypedElementQueryBuilder
 				return ITypedElementDataStorageInterface::EQueryAccessType::ReadWrite;
 			}
 		}
+
+		template<typename Type>
+		constexpr ITypedElementDataStorageInterface::EQueryDependencyFlags GetDependencyFlags()
+		{
+			ITypedElementDataStorageInterface::EQueryDependencyFlags Result = ITypedElementDataStorageInterface::EQueryDependencyFlags::None;
+			// Until there's a way to pass in whether or not a dependency is tied to the main thread and whether
+			// it's safe to not update in between updates, default to using the game thread and always update.
+			EnumAddFlags(Result, ITypedElementDataStorageInterface::EQueryDependencyFlags::GameThreadBound);
+			EnumAddFlags(Result, ITypedElementDataStorageInterface::EQueryDependencyFlags::AlwaysRefresh);
+
+			using BaseType = typename std::remove_reference_t<Type>;
+			if constexpr (TIsConst<BaseType>::Value)
+			{
+				EnumAddFlags(Result, ITypedElementDataStorageInterface::EQueryDependencyFlags::ReadOnly);
+			}
+			
+			return Result;
+		}
 	} // namespace Internal
 
 	//
@@ -79,8 +97,11 @@ namespace TypedElementQueryBuilder
 	// FQueryContextForwarder
 	//
 
-	FQueryContextForwarder::FQueryContextForwarder(ITypedElementDataStorageInterface::FQueryContext& InParentContext)
+	FQueryContextForwarder::FQueryContextForwarder(
+		const ITypedElementDataStorageInterface::FQueryDescription& InDescription, 
+		ITypedElementDataStorageInterface::IQueryContext& InParentContext)
 		: ParentContext(InParentContext)
+		, Description(InDescription)
 	{}
 
 	const void* FQueryContextForwarder::GetColumn(const UScriptStruct* ColumnType) const
@@ -166,55 +187,49 @@ namespace TypedElementQueryBuilder
 	//
 	
 	template<typename... Dependencies>
-	FCachedQueryContext<Dependencies...>::FCachedQueryContext(ITypedElementDataStorageInterface::FQueryContext& InParentContext)
-		: FQueryContextForwarder(InParentContext)
-	{
-		static constexpr ITypedElementDataStorageInterface::EQueryAccessType DependencyAccessList[] = { Internal::GetAccessType<Dependencies>()... };
-		static TWeakObjectPtr<const UClass> DependencyTypes[] = { Dependencies::StaticClass()... };
-
-		ParentContext.GetDependencies(DependencyAddresses, DependencyTypes, DependencyAccessList);
-#if DO_CHECK
-		UObject** DependencyAddessIt = DependencyAddresses;
-		for (int32 Index = 0; Index < sizeof...(Dependencies); ++Index)
-		{
-			checkf(*DependencyAddessIt != nullptr, TEXT("The cached query context tried to retrieve a dependency '%s' that didn't exist."),
-				DependencyTypes[Index] != nullptr ? *(DependencyTypes[Index]->GetName()) : TEXT("<no type information>"));
-			++DependencyAddessIt;
-		}
-#endif // DO_CHECK
-	}
+	FCachedQueryContext<Dependencies...>::FCachedQueryContext(
+		const ITypedElementDataStorageInterface::FQueryDescription& InDescription, 
+		ITypedElementDataStorageInterface::IQueryContext& InParentContext)
+		: FQueryContextForwarder(InDescription, InParentContext)
+	{}
 
 	template<typename... Dependencies>
 	void FCachedQueryContext<Dependencies...>::Register(ITypedElementDataStorageInterface::FQueryDescription& Query)
 	{
-		( Query.Dependencies.Emplace(Dependencies::StaticClass(), Internal::GetAccessType<Dependencies>()), ... );
+		Query.DependencyTypes.Reserve(sizeof...(Dependencies));
+		Query.DependencyFlags.Reserve(sizeof...(Dependencies));
+		
+		( Query.DependencyTypes.Emplace(Dependencies::StaticClass()), ... );
+		( Query.DependencyFlags.Emplace(Internal::GetDependencyFlags<Dependencies>()), ... );
+		
+		Query.CachedDependencies.AddDefaulted(sizeof...(Dependencies));
 	}
 
 	template<typename... Dependencies>
 	template<typename Dependency>
 	Dependency& FCachedQueryContext<Dependencies...>::GetCachedMutableDependency()
 	{
-		// Don't allow a subsystem registered as const to be found.
+		// Don't allow a dependency registered as const to be found.
 		constexpr uint32 Index = Internal::GetVarArgIndex<std::remove_const_t<Dependency>, Dependencies...>();
 		static_assert(Index < sizeof...(Dependencies), "Requested dependency isn't part of the query context cache.");
-		return *static_cast<Dependency*>(DependencyAddresses[Index]);
+		return *static_cast<Dependency*>(Description.CachedDependencies[Index].Get());
 	}
 
 	template<typename... Dependencies>
 	template<typename Dependency>
 	const Dependency& FCachedQueryContext<Dependencies...>::GetCachedDependency() const
 	{
-		// Allow access to subsystems registered with and without const.
+		// Allow access to dependencies registered with and without const.
 		constexpr uint32 Index = Internal::GetVarArgIndex<Dependency, Dependencies...>();
 		if constexpr (Index < sizeof...(Dependencies))
 		{
-			return *static_cast<Dependency*>(DependencyAddresses[Index]);
+			return *static_cast<Dependency*>(Description.CachedDependencies[Index].Get());
 		}
 		else
 		{
 			constexpr uint32 ConstIndex = Internal::GetVarArgIndex<std::add_const_t<Dependency>, Dependencies...>();
 			static_assert(ConstIndex < sizeof...(Dependencies), "Requested dependency isn't part of the query context cache.");
-			return *static_cast<Dependency*>(DependencyAddresses[ConstIndex]);
+			return *static_cast<Dependency*>(Description.CachedDependencies[ConstIndex].Get());
 		}
 	}
 	
@@ -223,8 +238,9 @@ namespace TypedElementQueryBuilder
 	//
 	namespace Internal
 	{
-		using QueryContext = ITypedElementDataStorageInterface::FQueryContext;
+		using QueryContext = ITypedElementDataStorageInterface::IQueryContext;
 		using QueryAccess = ITypedElementDataStorageInterface::EQueryAccessType;
+		using QueryDescription = ITypedElementDataStorageInterface::FQueryDescription;
 
 		template <typename T>
 		struct HasStaticStructMethod
@@ -348,13 +364,11 @@ namespace TypedElementQueryBuilder
 			std::tuple<BaseColumnType<ColumnTypes>*...> Columns;
 			static constexpr bool bArePointerColumns = AreAllColumnsPointers<ColumnTypes...>();
 
-			explicit FFunctionColumnInfo(QueryContext& Context)
+			FFunctionColumnInfo(const QueryDescription& Description, QueryContext& Context)
 			{
-				static constexpr QueryAccess ColumnAccess[] = { GetAccessType<ColumnTypes>()... };
-				static TWeakObjectPtr<const UScriptStruct> ColumnDescriptions[] = { GetColumnType<ColumnTypes>()...};
-
 				char* ColumnAddresses[sizeof...(ColumnTypes)];
-				Context.GetColumnsUnguarded(sizeof...(ColumnTypes), ColumnAddresses, ColumnDescriptions, ColumnAccess);
+				Context.GetColumnsUnguarded(sizeof...(ColumnTypes), ColumnAddresses, 
+					Description.SelectionTypes.GetData(), Description.SelectionAccessTypes.GetData());
 				std::apply([ColumnAddresses](auto&&... Column)
 					{
 						int Index = 0;
@@ -366,7 +380,7 @@ namespace TypedElementQueryBuilder
 		template<>
 		struct FFunctionColumnInfo<>
 		{
-			explicit FFunctionColumnInfo(QueryContext&) {}
+			FFunctionColumnInfo(const QueryDescription&, QueryContext&) {}
 			static constexpr bool bArePointerColumns = false;
 		};
 
@@ -374,11 +388,29 @@ namespace TypedElementQueryBuilder
 		struct FContextInfo
 		{
 			using BaseContextType = std::remove_reference_t<ContextType>;
-			static constexpr bool bIsUndecoratedContext = std::is_same_v<BaseContextType, QueryContext>;
-			using ContextWrapperType = std::conditional_t<bIsUndecoratedContext, QueryContext&, BaseContextType>;
-			ContextWrapperType ContextWrapper;
+			BaseContextType ContextWrapper;
+
+			FContextInfo(const QueryDescription& Description, QueryContext& Context)
+				: ContextWrapper(Description, Context)
+			{}
+		};
+
+		template<>
+		struct FContextInfo<const QueryContext&>
+		{
+			const QueryContext& ContextWrapper;
+
+			FContextInfo(const QueryDescription& Description, QueryContext& Context)
+				: ContextWrapper(Context)
+			{}
+		};
+
+		template<>
+		struct FContextInfo<QueryContext&>
+		{
+			QueryContext& ContextWrapper;
 
-			explicit FContextInfo(QueryContext& Context)
+			FContextInfo(const QueryDescription& Description, QueryContext& Context)
 				: ContextWrapper(Context)
 			{}
 		};
@@ -389,9 +421,9 @@ namespace TypedElementQueryBuilder
 			using SuperColumn = FFunctionColumnInfo<Args...>;
 			using SuperContext = FContextInfo<ContextType>;
 
-			explicit FContextRowHandleColumnsFunction(QueryContext& Context)
-				: SuperColumn(Context)
-				, SuperContext(Context)
+			FContextRowHandleColumnsFunction(const QueryDescription& Description, QueryContext& Context)
+				: SuperColumn(Description, Context)
+				, SuperContext(Description, Context)
 			{}
 
 			template<typename CallerType>
@@ -442,9 +474,9 @@ namespace TypedElementQueryBuilder
 			using SuperColumn = FFunctionColumnInfo<Args...>;
 			using SuperContext = FContextInfo<ContextType>;
 
-			explicit FContextColumnsFunction(QueryContext& Context)
-				: SuperColumn(Context)
-				, SuperContext(Context)
+			FContextColumnsFunction(const QueryDescription& Description, QueryContext& Context)
+				: SuperColumn(Description, Context)
+				, SuperContext(Description, Context)
 			{}
 
 			template<typename CallerType>
@@ -490,8 +522,8 @@ namespace TypedElementQueryBuilder
 		{
 			using Super = FFunctionColumnInfo<Columns...>;
 
-			explicit FRowHandleColumnsFunction(QueryContext& Context)
-				: Super(Context)
+			FRowHandleColumnsFunction(const QueryDescription& Description, QueryContext& Context)
+				: Super(Description, Context)
 			{}
 
 			template<typename CallerType>
@@ -538,8 +570,8 @@ namespace TypedElementQueryBuilder
 		{
 			using Super = FFunctionColumnInfo<Args...>;
 			
-			explicit FColumnsFunction(QueryContext& Context)
-				: Super(Context)
+			FColumnsFunction(const QueryDescription& Description, QueryContext& Context)
+				: Super(Description, Context)
 			{}
 
 			template<typename CallerType>
@@ -627,8 +659,8 @@ namespace TypedElementQueryBuilder
 		{
 			using Super = typename FFunctionInfoHelper<Args...>::BaseClass;
 
-			explicit FFunctionInfo(QueryContext& Context)
-				: Super(Context)
+			FFunctionInfo(const QueryDescription& Description, QueryContext& Context)
+				: Super(Description, Context)
 			{}
 		};
 
@@ -636,9 +668,9 @@ namespace TypedElementQueryBuilder
 		void BindQueryFunction(ITypedElementDataStorageInterface::QueryCallback& Function, 
 			void (*Callback)(Args...))
 		{
-			Function = [Callback](QueryContext& Context)
+			Function = [Callback](const QueryDescription& Description, QueryContext& Context)
 			{
-				FFunctionInfo<Args...>(Context).Call(Context, Callback);
+				FFunctionInfo<Args...>(Description, Context).Call(Context, Callback);
 			};
 		}
 
@@ -646,9 +678,9 @@ namespace TypedElementQueryBuilder
 		void BindQueryFunction(ITypedElementDataStorageInterface::QueryCallback& Function, Class* Target, 
 			void (Class::*Callback)(Args...))
 		{
-			Function = [Target, Callback](QueryContext& Context)
+			Function = [Target, Callback](const QueryDescription& Description, QueryContext& Context)
 			{
-				FFunctionInfo<Args...>(Context).Call(Context, (Target->*Callback));
+				FFunctionInfo<Args...>(Description, Context).Call(Context, (Target->*Callback));
 			};
 		}
 
@@ -656,9 +688,9 @@ namespace TypedElementQueryBuilder
 		void BindQueryFunction(ITypedElementDataStorageInterface::QueryCallback& Function, Class* Target, 
 			void (Class::*Callback)(Args...) const)
 		{
-			Function = [Target, Callback](QueryContext& Context)
+			Function = [Target, Callback](const QueryDescription& Description, QueryContext& Context)
 			{
-				FFunctionInfo<Args...>(Context).Call(Context, (Target->*Callback));
+				FFunctionInfo<Args...>(Description, Context).Call(Context, (Target->*Callback));
 			};
 		}
 
@@ -666,9 +698,9 @@ namespace TypedElementQueryBuilder
 		void BindQueryFunction_Expand(ITypedElementDataStorageInterface::QueryCallback& Function, Functor&& CallbackObject, 
 			void (Class::*Callback)(Args...) const)
 		{
-			Function = [CallbackObject = std::forward<Functor>(CallbackObject)](QueryContext& Context)
+			Function = [CallbackObject = std::forward<Functor>(CallbackObject)](const QueryDescription& Description, QueryContext& Context)
 			{
-				FFunctionInfo<Args...>(Context).Call(Context, CallbackObject);
+				FFunctionInfo<Args...>(Description, Context).Call(Context, CallbackObject);
 			};
 		}
 
@@ -676,9 +708,9 @@ namespace TypedElementQueryBuilder
 		void BindQueryFunction_Expand(ITypedElementDataStorageInterface::QueryCallback& Function, Functor&& CallbackObject, 
 			void (Class::*Callback)(Args...))
 		{
-			Function = [CallbackObject = std::forward<Functor>(CallbackObject)](QueryContext& Context)
+			Function = [CallbackObject = std::forward<Functor>(CallbackObject)](const QueryDescription& Description, QueryContext& Context)
 			{
-				FFunctionInfo<Args...>(Context).Call(Context, CallbackObject);
+				FFunctionInfo<Args...>(Description, Context).Call(Context, CallbackObject);
 			};
 		}
 
@@ -701,7 +733,7 @@ namespace TypedElementQueryBuilder
 			}
 		}
 
-		template<typename ContextType> void RegisterDependencies(ITypedElementDataStorageInterface::FQueryDescription& Query)
+		template<typename ContextType> void RegisterDependencies(QueryDescription& Query)
 		{
 			using BaseType = typename TRemoveReference<ContextType>::Type;
 			if constexpr (TIsDerivedFrom<BaseType, FQueryContextForwarder>::Value)
@@ -713,7 +745,7 @@ namespace TypedElementQueryBuilder
 		template<typename... Args>
 		struct RegisterFunctionArgumentsHelper
 		{
-			void Register(ITypedElementDataStorageInterface::FQueryDescription&, Select& Target)
+			void Register(QueryDescription&, Select& Target)
 			{
 				(AddColumnToSelect<Args>(Target), ...);
 			}
@@ -722,13 +754,13 @@ namespace TypedElementQueryBuilder
 		template<>
 		struct RegisterFunctionArgumentsHelper<>
 		{
-			static void Register(ITypedElementDataStorageInterface::FQueryDescription&, Select&) {}
+			static void Register(QueryDescription&, Select&) {}
 		};
 
 		template<typename Arg0>
 		struct RegisterFunctionArgumentsHelper<Arg0>
 		{
-			static void Register(ITypedElementDataStorageInterface::FQueryDescription& Query, Select& Target)
+			static void Register(QueryDescription& Query, Select& Target)
 			{
 				if constexpr (IsValidContextType<Arg0>())
 				{
@@ -745,7 +777,7 @@ namespace TypedElementQueryBuilder
 		template<typename Arg0, typename Arg1, typename... Args>
 		struct RegisterFunctionArgumentsHelper<Arg0, Arg1, Args...>
 		{
-			static void Register(ITypedElementDataStorageInterface::FQueryDescription& Query, Select& Target)
+			static void Register(QueryDescription& Query, Select& Target)
 			{
 				if constexpr (IsValidContextType<Arg0>())
 				{
@@ -768,25 +800,25 @@ namespace TypedElementQueryBuilder
 		};
 
 		template<typename... Args>
-		void RegisterFunctionArguments(ITypedElementDataStorageInterface::FQueryDescription& Query, Select& Target, void (*)(Args...))
+		void RegisterFunctionArguments(QueryDescription& Query, Select& Target, void (*)(Args...))
 		{
 			RegisterFunctionArgumentsHelper<Args...>::Register(Query, Target);
 		}
 
 		template<typename Class, typename... Args>
-		void RegisterFunctionArguments(ITypedElementDataStorageInterface::FQueryDescription& Query, Select& Target, void (Class::*)(Args...))
+		void RegisterFunctionArguments(QueryDescription& Query, Select& Target, void (Class::*)(Args...))
 		{
 			RegisterFunctionArgumentsHelper<Args...>::Register(Query, Target);
 		}
 
 		template<typename Class, typename... Args>
-		void RegisterFunctionArguments(ITypedElementDataStorageInterface::FQueryDescription& Query, Select& Target, void (Class::*)(Args...) const)
+		void RegisterFunctionArguments(QueryDescription& Query, Select& Target, void (Class::*)(Args...) const)
 		{
 			RegisterFunctionArgumentsHelper<Args...>::Register(Query, Target);
 		}
 
 		template<typename Functor>
-		void RegisterFunctionArguments(ITypedElementDataStorageInterface::FQueryDescription& Query, Select& Target, Functor)
+		void RegisterFunctionArguments(QueryDescription& Query, Select& Target, Functor)
 		{
 			RegisterFunctionArguments(Query, Target, &Functor::operator());
 		}
@@ -879,7 +911,7 @@ namespace TypedElementQueryBuilder
 			}
 		};
 		
-		inline void PrepareForQueryBinding(ITypedElementDataStorageInterface::FQueryDescription& Query, const FProcessor& Processor)
+		inline void PrepareForQueryBinding(QueryDescription& Query, const FProcessor& Processor)
 		{
 			Query.Callback.Type = ITypedElementDataStorageInterface::EQueryCallbackType::Processor;
 			Query.Callback.Phase = Processor.Phase;
@@ -889,7 +921,7 @@ namespace TypedElementQueryBuilder
 			Query.Callback.bForceToGameThread = Processor.bForceToGameThread;
 		}
 
-		inline void PrepareForQueryBinding(ITypedElementDataStorageInterface::FQueryDescription& Query, const FObserver& Observer)
+		inline void PrepareForQueryBinding(QueryDescription& Query, const FObserver& Observer)
 		{
 			switch (Observer.Event)
 			{
@@ -905,7 +937,7 @@ namespace TypedElementQueryBuilder
 		}
 		
 		template<typename CallbackType, typename Function>
-		void PrepareForQueryBinding(Select& Target, ITypedElementDataStorageInterface::FQueryDescription& Query, FName Name, 
+		void PrepareForQueryBinding(Select& Target, QueryDescription& Query, FName Name,
 			const CallbackType& Type, Function Callback)
 		{
 			static_assert(TIsDerivedFrom<CallbackType, FQueryCallbackType>::Value, "The callback type provided isn't one of the available "
@@ -919,7 +951,7 @@ The following options are supported:
 - void(<Context>&, TypedElementRowHandle, [const]Column&...) 
 - void(<Context>&, [const]Column*...) 
 - void(<Context>&, const TypedElementRowHandle*, [const]Column*...) 
-Where <Context> is ITypedElementDataStorageInterface::FQueryContext or FCachedQueryContext<...>
+Where <Context> is ITypedElementDataStorageInterface::IQueryContext or FCachedQueryContext<...>
 e.g. void(FCachedQueryContext<Subsystem1, const Subsystem2>& Context, TypedElementRowHandle Row, ColumnType0& ColumnA, const ColumnType1& ColumnB) {...}
 )");
 			RegisterFunctionArguments(Query, Target, Callback);
diff --git a/Engine/Source/Runtime/TypedElementFramework/Public/Elements/Interfaces/TypedElementDataStorageInterface.h b/Engine/Source/Runtime/TypedElementFramework/Public/Elements/Interfaces/TypedElementDataStorageInterface.h
--- a/Engine/Source/Runtime/TypedElementFramework/Public/Elements/Interfaces/TypedElementDataStorageInterface.h
+++ b/Engine/Source/Runtime/TypedElementFramework/Public/Elements/Interfaces/TypedElementDataStorageInterface.h
@@ -7,6 +7,7 @@
 #include "Delegates/Delegate.h"
 #include "Elements/Framework/TypedElementColumnUtils.h"
 #include "Math/NumericLimits.h"
+#include "Misc/EnumClassFlags.h"
 #include "Templates/Function.h"
 #include "UObject/Interface.h"
 #include "UObject/NameTypes.h"
@@ -167,7 +168,7 @@ public:
 		 */
 		PrepareSyncExternalToDataStorage,
 		/** 
-		 * The group for processors that need to sync data from external sources such as subsystems or the world into
+		 * The group for queries that need to sync data from external sources such as subsystems or the world into
 		 * the Data Storage. These typically run early in a phase.
 		 */
 		SyncExternalToDataStorage,
@@ -176,22 +177,25 @@ public:
 		 * data with a dependency on external updates being completed.
 		 */
 		FinalizeSyncExternalToDataStorage,
-
 		
 		/** Prepares for data from the Data Storage to be synced to external sources such as subsystem or the world. */
 		PrepareSyncDataStorageToExternal,
 		/**
-		 * The group for processors that need to sync data from the Data Storage to external sources such as subsystems
+		 * The group for quries that need to sync data from the Data Storage to external sources such as subsystems
 		 * or the world into. These typically run late in a phase.
 		 */
 		SyncDataStorageToExternal,
 		/** Finalizes the sync of data from the Data Storage to external sources such as subsystem or the world. */
 		FinalizeSyncDataStorageToExternal,
 
+		/** Prepares for widgets getting synced to/from widgets. */
+		PrepareSyncWidgets,
 		/**
-		 * Processors grouped under this name will sync data to/from widgets.
+		 * Queries grouped under this name will sync data to/from widgets.
 		 */
-		SyncWidgets
+		SyncWidgets,
+		/** Finalize for widgets getting synced to/from widgets. */
+		FinalizeSyncWidgets
 	};
 
 	enum class EQueryCallbackType : uint8
@@ -206,26 +210,36 @@ public:
 		ObserveRemove
 	};
 
-	enum class EQueryAccessType
+	enum class EQueryAccessType : bool
 	{ 
 		ReadOnly, 
 		ReadWrite 
 	};
 
-	/** 
-	 * Base class to be provided to query callbacks.
-	 * Note that at the time of writing only subclasses of Subsystem are supported as dependencies.
+	enum class EQueryDependencyFlags : uint8
+	{
+		None = 0,
+		/** If set the dependency is accessed as read-only. If not set the dependency requires Read/Write access. */
+		ReadOnly = 1 << 0,
+		/** If set the dependecy can only be used from the game thread, otherwise it can be accessed from any thread. */
+		GameThreadBound = 1 << 1,
+		/** If set the dependency will be refetched every iteration, otherwise only if not fetched before. */
+		AlwaysRefresh = 1 << 2
+	};
+
+	/**
+	 * Base interface for any contexts provided to query callbacks.
 	 */
-	struct FQueryContext
+	struct ICommonQueryContext
 	{
-		virtual ~FQueryContext() = default;
-		
+		virtual ~ICommonQueryContext() = default;
+
 		/** Return the address of a immutable column matching the requested type or a nullptr if not found. */
 		virtual const void* GetColumn(const UScriptStruct* ColumnType) const = 0;
 		/** Return the address of a mutable column matching the requested type or a nullptr if not found. */
 		virtual void* GetMutableColumn(const UScriptStruct* ColumnType) = 0;
-		/** 
-		 * Get a list of columns or nullptrs if the columntype wasn't found. Mutable addresses are returned and it's up to 
+		/**
+		 * Get a list of columns or nullptrs if the columntype wasn't found. Mutable addresses are returned and it's up to
 		 * the caller to not change immutable addresses.
 		 */
 		virtual void GetColumns(TArrayView<char*> RetrievedAddresses, TConstArrayView<TWeakObjectPtr<const UScriptStruct>> ColumnTypes,
@@ -238,6 +252,30 @@ public:
 		virtual void GetColumnsUnguarded(int32 TypeCount, char** RetrievedAddresses, const TWeakObjectPtr<const UScriptStruct>* ColumnTypes,
 			const ITypedElementDataStorageInterface::EQueryAccessType* AccessTypes) = 0;
 
+		/** Returns the number rows in the batch. */
+		virtual uint32 GetRowCount() const = 0;
+		/**
+		 * Returns an immutable view that contains the row handles for all returned results. The returned size will be the same  as the
+		 * value returned by GetRowCount().
+		 */
+		virtual TConstArrayView<TypedElementRowHandle> GetRowHandles() const = 0;
+
+		// Utility functions
+
+		template<typename Column>
+		const Column* GetColumn() const;
+		template<typename Column>
+		Column* GetMutableColumn();
+	};
+	
+	/** 
+	 * Interface to be provided to query callbacks running with the Data Storage.
+	 * Note that at the time of writing only subclasses of Subsystem are supported as dependencies.
+	 */
+	struct IQueryContext : public ICommonQueryContext
+	{
+		virtual ~IQueryContext() = default;
+		
 		/** Returns an immutable instance of the requested dependency or a nullptr if not found. */
 		virtual const UObject* GetDependency(const UClass* DependencyClass) = 0;
 		/** Returns a mutable instance of the requested dependency or a nullptr if not found. */
@@ -249,13 +287,6 @@ public:
 		virtual void GetDependencies(TArrayView<UObject*> RetrievedAddresses, TConstArrayView<TWeakObjectPtr<const UClass>> DependencyTypes,
 			TConstArrayView<EQueryAccessType> AccessTypes) = 0;
 
-		/** Returns the number rows in the batch. */
-		virtual uint32 GetRowCount() const = 0;
-		/** 
-		 * Returns an immutable view that contains the row handles for all returned results. The returned size will be the same  as the
-		 * value returned by GetRowCount().
-		 */
-		virtual TConstArrayView<TypedElementRowHandle> GetRowHandles() const = 0;
 		/**
 		 * Removes the row with the provided row handle. The removal will not be immediately done but delayed until the end of the tick 
 		 * group.
@@ -300,10 +331,24 @@ public:
 		void RemoveColumns(TConstArrayView<TypedElementRowHandle> Rows);
 	};
 
-	using QueryCallback = TFunction<void(FQueryContext&)>;
+	/**
+	 * Interface to be provided to query callbacks that are directly called through RunQuery.
+	 */
+	struct IDirectQueryContext : public ICommonQueryContext
+	{
+		virtual ~IDirectQueryContext() = default;
+	};
+
+	struct FQueryDescription;
+	using QueryCallback = TFunction<void(const FQueryDescription&, IQueryContext&)>;
+	using DirectQueryCallbackRef = TFunctionRef<void(const FQueryDescription&, IDirectQueryContext&)>;
 
 	struct FQueryDescription final
 	{
+		static constexpr int32 NumInlineSelections = 8;
+		static constexpr int32 NumInlineConditions = 8;
+		static constexpr int32 NumInlineDependencies = 2;
+
 		enum class EActionType : uint8
 		{
 			None,	/** Do nothing. */
@@ -334,25 +379,7 @@ public:
 		{
 			FBinaryOperator Binary;
 			OperatorIndex Unary;
-			const UScriptStruct* Type;
-		};
-
-		struct FAccessControlledStruct
-		{
-			FAccessControlledStruct() = default;
-			inline FAccessControlledStruct(const UScriptStruct* Type, EQueryAccessType Access) : Type(Type), Access(Access) {}
-
-			const UScriptStruct* Type;
-			EQueryAccessType Access;
-		};
-
-		struct FAccessControlledClass
-		{
-			FAccessControlledClass() = default;
-			inline FAccessControlledClass(const UClass* Type, EQueryAccessType Access) : Type(Type), Access(Access) {}
-
-			const UClass* Type;
-			EQueryAccessType Access;
+			TWeakObjectPtr<const UScriptStruct> Type;
 		};
 
 		struct FCallbackData
@@ -369,16 +396,29 @@ public:
 		};
 		
 		FCallbackData Callback;
-		TArray<FAccessControlledStruct, TInlineAllocator<16>> Selection;
-		TArray<EOperatorType, TInlineAllocator<32>> ConditionTypes;
-		TArray<FOperator, TInlineAllocator<32>> ConditionOperators;
-		TArray<FAccessControlledClass, TInlineAllocator<4>> Dependencies;
+
+		// The list of arrays below are required to remain in the same order as they're added as the function binding expects certain entries
+		// to be in a specific location.
+
+		TArray<TWeakObjectPtr<const UScriptStruct>, TInlineAllocator<NumInlineSelections>> SelectionTypes;
+		TArray<EQueryAccessType, TInlineAllocator<NumInlineSelections>> SelectionAccessTypes;
+		
+		TArray<EOperatorType, TInlineAllocator<NumInlineConditions>> ConditionTypes;
+		TArray<FOperator, TInlineAllocator<NumInlineConditions>> ConditionOperators;
+		
+		TArray<TWeakObjectPtr<const UClass>, TInlineAllocator<NumInlineDependencies>> DependencyTypes;
+		TArray<EQueryDependencyFlags, TInlineAllocator<NumInlineDependencies>> DependencyFlags;
+		/** Cached instances of the dependencies. This will always match the count of the other Dependency*Types, but may contain null pointers. */
+		TArray<TWeakObjectPtr<UObject>, TInlineAllocator<NumInlineDependencies>> CachedDependencies;
+		
 		EActionType Action;
 		/** If true, this query only has simple operations and is guaranteed to be executed fully and at optimal performance. */
 		bool bSimpleQuery{ false };
 	};
 	virtual TypedElementQueryHandle RegisterQuery(FQueryDescription&& Query) = 0;
 	virtual void UnregisterQuery(TypedElementQueryHandle Query) = 0;
+	/** Returns the description of a previously registered query. If the query no longer exists an empty description will be returned. */
+	virtual const FQueryDescription& GetQueryDescription(TypedElementQueryHandle Query) const = 0;
 	/**
 	 * Tick groups for queries can be given any name and the Data Storage will figure out the order of execution based on found
 	 * dependencies. However keeping processors within the same query group can help promote better performance through parallelization.
@@ -399,13 +439,16 @@ public:
 			 */
 			Unsupported,
 			/** The provided query is no longer available. */
-			Unavailable
+			Unavailable,
+			/** One or more dependencies declared on the query could not be retrieved. */
+			MissingDependency
 		};
 		
 		uint32 Count{ 0 }; /** The number of rows were processed. */
 		ECompletion Completed{ ECompletion::Unavailable };
 	};
 	virtual FQueryResult RunQuery(TypedElementQueryHandle Query) = 0;
+	virtual FQueryResult RunQuery(TypedElementQueryHandle Query, DirectQueryCallbackRef Callback) = 0;
 	
 	/**
 	 * @section Misc
@@ -459,6 +502,9 @@ public:
 
 
 // Implementations
+
+ENUM_CLASS_FLAGS(ITypedElementDataStorageInterface::EQueryDependencyFlags);
+
 template<typename TagType>
 void ITypedElementDataStorageInterface::AddTag(TypedElementRowHandle Row)
 {
@@ -479,6 +525,18 @@ ColumnType* ITypedElementDataStorageInterface::AddOrGetColumn(TypedElementRowHan
 	return Result;
 }
 
+template<typename Column>
+const Column* ITypedElementDataStorageInterface::ICommonQueryContext::GetColumn() const
+{
+	return reinterpret_cast<const Column*>(GetColumn(Column::StaticStruct()));
+}
+
+template<typename Column>
+Column* ITypedElementDataStorageInterface::ICommonQueryContext::GetMutableColumn()
+{
+	return reinterpret_cast<Column*>(GetMutableColumn(Column::StaticStruct()));
+}
+
 template<typename ColumnType>
 ColumnType* ITypedElementDataStorageInterface::GetColumn(TypedElementRowHandle Row)
 {
@@ -493,25 +551,25 @@ SystemType* ITypedElementDataStorageInterface::GetExternalSystem()
 
 
 template<typename... Columns>
-void ITypedElementDataStorageInterface::FQueryContext::AddColumns(TypedElementRowHandle Row)
+void ITypedElementDataStorageInterface::IQueryContext::AddColumns(TypedElementRowHandle Row)
 {
 	AddColumns(Row, { Columns::StaticStruct()... });
 }
 
 template<typename... Columns>
-void ITypedElementDataStorageInterface::FQueryContext::AddColumns(TConstArrayView<TypedElementRowHandle> Rows)
+void ITypedElementDataStorageInterface::IQueryContext::AddColumns(TConstArrayView<TypedElementRowHandle> Rows)
 {
 	AddColumns(Rows, { Columns::StaticStruct()... });
 }
 
 template<typename... Columns>
-void ITypedElementDataStorageInterface::FQueryContext::RemoveColumns(TypedElementRowHandle Row)
+void ITypedElementDataStorageInterface::IQueryContext::RemoveColumns(TypedElementRowHandle Row)
 {
 	RemoveColumns(Row, { Columns::StaticStruct()... });
 }
 
 template<typename... Columns>
-void ITypedElementDataStorageInterface::FQueryContext::RemoveColumns(TConstArrayView<TypedElementRowHandle> Rows)
+void ITypedElementDataStorageInterface::IQueryContext::RemoveColumns(TConstArrayView<TypedElementRowHandle> Rows)
 {
 	RemoveColumns(Rows, { Columns::StaticStruct()... });
 }
\ No newline at end of file
