New and updated TEDS columns.

The following changes were were made:
- Added a new FTypedElementRowReferenceColumn column that holds the reference to another row.
- Added new columns to store type information for classes and structs in.
- All columns are now marked as "final" to avoid accidental inherting.
- Where possible, the column member variables have been marked with UPROPERTY() to allow them to be reflected, e.g. in the UI.

diff --git a/Engine/Source/Runtime/TypedElementFramework/Public/Elements/Columns/TypedElementLabelColumns.h b/Engine/Source/Runtime/TypedElementFramework/Public/Elements/Columns/TypedElementLabelColumns.h
--- a/Engine/Source/Runtime/TypedElementFramework/Public/Elements/Columns/TypedElementLabelColumns.h
+++ b/Engine/Source/Runtime/TypedElementFramework/Public/Elements/Columns/TypedElementLabelColumns.h
@@ -13,10 +13,11 @@
  * Column that stores a label.
  */
 USTRUCT()
-struct FTypedElementLabelColumn : public FTypedElementDataStorageColumn
+struct FTypedElementLabelColumn final : public FTypedElementDataStorageColumn
 {
 	GENERATED_BODY()
 
+	UPROPERTY()
 	FString Label;
 };
 
@@ -25,9 +26,10 @@ struct FTypedElementLabelColumn : public FTypedElementDataStorageColumn
  * kept separate in order to iterate quickly over all hash values.
  */
 USTRUCT()
-struct FTypedElementLabelHashColumn : public FTypedElementDataStorageColumn
+struct FTypedElementLabelHashColumn final : public FTypedElementDataStorageColumn
 {
 	GENERATED_BODY()
 
+	UPROPERTY()
 	uint64 LabelHash;
 };
\ No newline at end of file
diff --git a/Engine/Source/Runtime/TypedElementFramework/Public/Elements/Columns/TypedElementMiscColumns.h b/Engine/Source/Runtime/TypedElementFramework/Public/Elements/Columns/TypedElementMiscColumns.h
--- a/Engine/Source/Runtime/TypedElementFramework/Public/Elements/Columns/TypedElementMiscColumns.h
+++ b/Engine/Source/Runtime/TypedElementFramework/Public/Elements/Columns/TypedElementMiscColumns.h
@@ -14,7 +14,15 @@
  * will automatically be removed at the end of a tick.
  */
 USTRUCT()
-struct FTypedElementSyncBackToWorldTag : public FTypedElementDataStorageTag
+struct FTypedElementSyncBackToWorldTag final : public FTypedElementDataStorageTag
 {
 	GENERATED_BODY()
 };
+
+USTRUCT()
+struct FTypedElementRowReferenceColumn final : public FTypedElementDataStorageColumn
+{
+	GENERATED_BODY()
+
+	TypedElementRowHandle Row;
+};
\ No newline at end of file
diff --git a/Engine/Source/Runtime/TypedElementFramework/Public/Elements/Columns/TypedElementPackageColumns.h b/Engine/Source/Runtime/TypedElementFramework/Public/Elements/Columns/TypedElementPackageColumns.h
--- a/Engine/Source/Runtime/TypedElementFramework/Public/Elements/Columns/TypedElementPackageColumns.h
+++ b/Engine/Source/Runtime/TypedElementFramework/Public/Elements/Columns/TypedElementPackageColumns.h
@@ -15,10 +15,11 @@
  * Column that stores the path of a package.
  */
 USTRUCT()
-struct FTypedElementPackagePathColumn : public FTypedElementDataStorageColumn
+struct FTypedElementPackagePathColumn final : public FTypedElementDataStorageColumn
 {
 	GENERATED_BODY()
 
+	UPROPERTY()
 	FString Path;
 };
 
@@ -31,7 +32,7 @@ inline uint32 GetTypeHash(const FTypedElementPackagePathColumn& InStruct)
  * Column that stores the full loading path to a package.
  */
 USTRUCT()
-struct FTypedElementPackageLoadedPathColumn : public FTypedElementDataStorageColumn
+struct FTypedElementPackageLoadedPathColumn final : public FTypedElementDataStorageColumn
 {
 	GENERATED_BODY()
 
diff --git a/Engine/Source/Runtime/TypedElementFramework/Public/Elements/Columns/TypedElementSlateWidgetColumns.h b/Engine/Source/Runtime/TypedElementFramework/Public/Elements/Columns/TypedElementSlateWidgetColumns.h
--- a/Engine/Source/Runtime/TypedElementFramework/Public/Elements/Columns/TypedElementSlateWidgetColumns.h
+++ b/Engine/Source/Runtime/TypedElementFramework/Public/Elements/Columns/TypedElementSlateWidgetColumns.h
@@ -17,7 +17,7 @@ class SWidget;
  * be deleted.
  */
 USTRUCT()
-struct FTypedElementSlateWidgetReferenceColumn : public FTypedElementDataStorageColumn
+struct FTypedElementSlateWidgetReferenceColumn final : public FTypedElementDataStorageColumn
 {
 	GENERATED_BODY()
 
@@ -30,7 +30,7 @@ struct FTypedElementSlateWidgetReferenceColumn : public FTypedElementDataStorage
  * the column will be removed.
  */
 USTRUCT()
-struct FTypedElementSlateWidgetReferenceDeletesRowTag : public FTypedElementDataStorageTag
+struct FTypedElementSlateWidgetReferenceDeletesRowTag final : public FTypedElementDataStorageTag
 {
 	GENERATED_BODY()
 };
\ No newline at end of file
diff --git a/Engine/Source/Runtime/TypedElementFramework/Public/Elements/Columns/TypedElementTransformColumns.h b/Engine/Source/Runtime/TypedElementFramework/Public/Elements/Columns/TypedElementTransformColumns.h
--- a/Engine/Source/Runtime/TypedElementFramework/Public/Elements/Columns/TypedElementTransformColumns.h
+++ b/Engine/Source/Runtime/TypedElementFramework/Public/Elements/Columns/TypedElementTransformColumns.h
@@ -13,7 +13,7 @@
  * Column that stores a local transform. 
  */
 USTRUCT()
-struct FTypedElementLocalTransformColumn : public FTypedElementDataStorageColumn
+struct FTypedElementLocalTransformColumn final : public FTypedElementDataStorageColumn
 {
 	GENERATED_BODY()
 
@@ -21,5 +21,6 @@ struct FTypedElementLocalTransformColumn : public FTypedElementDataStorageColumn
 	// Transfrom will be updated the first and following ticks after it's creation. If this
 	// isn't initialized at the correct time, then the sync from source or the true initialization
 	// need to be moved to an earlier phase or group.
+	UPROPERTY()
 	FTransform Transform { NoInit };
 };
diff --git a/Engine/Source/Runtime/TypedElementFramework/Public/Elements/Columns/TypedElementTypeInfoColumns.h b/Engine/Source/Runtime/TypedElementFramework/Public/Elements/Columns/TypedElementTypeInfoColumns.h
new file mode 100644
--- /dev/null
+++ b/Engine/Source/Runtime/TypedElementFramework/Public/Elements/Columns/TypedElementTypeInfoColumns.h
@@ -0,0 +1,34 @@
+// Copyright Epic Games, Inc. All Rights Reserved.
+
+#pragma once
+
+#include "Elements/Interfaces/TypedElementDataStorageInterface.h"
+#include "UObject/ObjectMacros.h"
+#include "UObject/WeakObjectPtrTemplates.h"
+
+#include "TypedElementTypeInfoColumns.generated.h"
+
+class UClass;
+class UScriptStruct;
+
+/**
+ * Column that stores type information for classes.
+ */
+USTRUCT()
+struct FTypedElementClassTypeInfoColumn final : public FTypedElementDataStorageColumn
+{
+	GENERATED_BODY()
+
+	TWeakObjectPtr<const UClass> TypeInfo;
+};
+
+/**
+ * Column that stores type information for structs.
+ */
+USTRUCT()
+struct FTypedElementScriptStructTypeInfoColumn final : public FTypedElementDataStorageColumn
+{
+	GENERATED_BODY()
+
+	TWeakObjectPtr<const UScriptStruct> TypeInfo;
+};
diff --git a/Engine/Source/Runtime/TypedElementFramework/Public/Elements/Columns/TypedElementValueCacheColumns.h b/Engine/Source/Runtime/TypedElementFramework/Public/Elements/Columns/TypedElementValueCacheColumns.h
--- a/Engine/Source/Runtime/TypedElementFramework/Public/Elements/Columns/TypedElementValueCacheColumns.h
+++ b/Engine/Source/Runtime/TypedElementFramework/Public/Elements/Columns/TypedElementValueCacheColumns.h
@@ -12,10 +12,11 @@
  * Column that can be used to cache an unsigned 32-bit value in.
  */
 USTRUCT()
-struct FTypedElementU32IntValueCacheColumn : public FTypedElementDataStorageColumn
+struct FTypedElementU32IntValueCacheColumn final : public FTypedElementDataStorageColumn
 {
 	GENERATED_BODY()
 
+	UPROPERTY()
 	uint32 Value;
 };
 
@@ -23,10 +24,11 @@ struct FTypedElementU32IntValueCacheColumn : public FTypedElementDataStorageColu
  * Column that can be used to cache a signed 32-bit value in.
  */
 USTRUCT()
-struct FTypedElementI32IntValueCacheColumn : public FTypedElementDataStorageColumn
+struct FTypedElementI32IntValueCacheColumn final : public FTypedElementDataStorageColumn
 {
 	GENERATED_BODY()
 
+	UPROPERTY()
 	int32 Value;
 };
 
@@ -34,9 +36,10 @@ struct FTypedElementI32IntValueCacheColumn : public FTypedElementDataStorageColu
  * Column that can be used to cache a 32-bit floating point value in.
  */
 USTRUCT()
-struct FTypedElementFloatValueCacheColumn : public FTypedElementDataStorageColumn
+struct FTypedElementFloatValueCacheColumn final : public FTypedElementDataStorageColumn
 {
 	GENERATED_BODY()
 
+	UPROPERTY()
 	float Value;
 };
\ No newline at end of file
