Fixed uninitialized property errors in TEDS reported during startup.

diff --git a/Engine/Source/Runtime/TypedElementFramework/Public/Elements/Columns/TypedElementLabelColumns.h b/Engine/Source/Runtime/TypedElementFramework/Public/Elements/Columns/TypedElementLabelColumns.h
--- a/Engine/Source/Runtime/TypedElementFramework/Public/Elements/Columns/TypedElementLabelColumns.h
+++ b/Engine/Source/Runtime/TypedElementFramework/Public/Elements/Columns/TypedElementLabelColumns.h
@@ -30,6 +30,6 @@ struct FTypedElementLabelHashColumn final : public FTypedElementDataStorageColum
 {
 	GENERATED_BODY()
 
-	UPROPERTY()
+	UPROPERTY(meta = (IgnoreForMemberInitializationTest))
 	uint64 LabelHash;
 };
\ No newline at end of file
diff --git a/Engine/Source/Runtime/TypedElementFramework/Public/Elements/Columns/TypedElementTransformColumns.h b/Engine/Source/Runtime/TypedElementFramework/Public/Elements/Columns/TypedElementTransformColumns.h
--- a/Engine/Source/Runtime/TypedElementFramework/Public/Elements/Columns/TypedElementTransformColumns.h
+++ b/Engine/Source/Runtime/TypedElementFramework/Public/Elements/Columns/TypedElementTransformColumns.h
@@ -21,6 +21,6 @@ struct FTypedElementLocalTransformColumn final : public FTypedElementDataStorage
 	// Transfrom will be updated the first and following ticks after it's creation. If this
 	// isn't initialized at the correct time, then the sync from source or the true initialization
 	// need to be moved to an earlier phase or group.
-	UPROPERTY()
+	UPROPERTY(meta = (IgnoreForMemberInitializationTest))
 	FTransform Transform { NoInit };
 };
diff --git a/Engine/Source/Runtime/TypedElementFramework/Public/Elements/Columns/TypedElementValueCacheColumns.h b/Engine/Source/Runtime/TypedElementFramework/Public/Elements/Columns/TypedElementValueCacheColumns.h
--- a/Engine/Source/Runtime/TypedElementFramework/Public/Elements/Columns/TypedElementValueCacheColumns.h
+++ b/Engine/Source/Runtime/TypedElementFramework/Public/Elements/Columns/TypedElementValueCacheColumns.h
@@ -16,7 +16,7 @@ struct FTypedElementU32IntValueCacheColumn final : public FTypedElementDataStora
 {
 	GENERATED_BODY()
 
-	UPROPERTY()
+	UPROPERTY(meta = (IgnoreForMemberInitializationTest))
 	uint32 Value;
 };
 
@@ -28,7 +28,7 @@ struct FTypedElementI32IntValueCacheColumn final : public FTypedElementDataStora
 {
 	GENERATED_BODY()
 
-	UPROPERTY()
+	UPROPERTY(meta = (IgnoreForMemberInitializationTest))
 	int32 Value;
 };
 
@@ -40,6 +40,6 @@ struct FTypedElementFloatValueCacheColumn final : public FTypedElementDataStorag
 {
 	GENERATED_BODY()
 
-	UPROPERTY()
+	UPROPERTY(meta = (IgnoreForMemberInitializationTest))
 	float Value;
 };
\ No newline at end of file
