Gave the various TEDS columns a friendly name so they present nicer in the future in tools that use the friendly name from a type.

diff --git a/Engine/Source/Runtime/TypedElementFramework/Public/Elements/Columns/TypedElementLabelColumns.h b/Engine/Source/Runtime/TypedElementFramework/Public/Elements/Columns/TypedElementLabelColumns.h
--- a/Engine/Source/Runtime/TypedElementFramework/Public/Elements/Columns/TypedElementLabelColumns.h
+++ b/Engine/Source/Runtime/TypedElementFramework/Public/Elements/Columns/TypedElementLabelColumns.h
@@ -12,7 +12,7 @@
 /**
  * Column that stores a label.
  */
-USTRUCT()
+USTRUCT(meta = (DisplayName = "Label"))
 struct FTypedElementLabelColumn final : public FTypedElementDataStorageColumn
 {
 	GENERATED_BODY()
@@ -25,7 +25,7 @@ struct FTypedElementLabelColumn final : public FTypedElementDataStorageColumn
  * Column that stores the hash of a label. This is typically paired with FTypedElementLabelColumn, but 
  * kept separate in order to iterate quickly over all hash values.
  */
-USTRUCT()
+USTRUCT(meta = (DisplayName = "Label hash"))
 struct FTypedElementLabelHashColumn final : public FTypedElementDataStorageColumn
 {
 	GENERATED_BODY()
diff --git a/Engine/Source/Runtime/TypedElementFramework/Public/Elements/Columns/TypedElementMiscColumns.h b/Engine/Source/Runtime/TypedElementFramework/Public/Elements/Columns/TypedElementMiscColumns.h
--- a/Engine/Source/Runtime/TypedElementFramework/Public/Elements/Columns/TypedElementMiscColumns.h
+++ b/Engine/Source/Runtime/TypedElementFramework/Public/Elements/Columns/TypedElementMiscColumns.h
@@ -13,13 +13,13 @@
  * need to be copied out the Data Storage and into the original object. This tag
  * will automatically be removed at the end of a tick.
  */
-USTRUCT()
+USTRUCT(meta = (DisplayName = "Sync back to world"))
 struct FTypedElementSyncBackToWorldTag final : public FTypedElementDataStorageTag
 {
 	GENERATED_BODY()
 };
 
-USTRUCT()
+USTRUCT(meta = (DisplayName = "Row reference"))
 struct FTypedElementRowReferenceColumn final : public FTypedElementDataStorageColumn
 {
 	GENERATED_BODY()
diff --git a/Engine/Source/Runtime/TypedElementFramework/Public/Elements/Columns/TypedElementPackageColumns.h b/Engine/Source/Runtime/TypedElementFramework/Public/Elements/Columns/TypedElementPackageColumns.h
--- a/Engine/Source/Runtime/TypedElementFramework/Public/Elements/Columns/TypedElementPackageColumns.h
+++ b/Engine/Source/Runtime/TypedElementFramework/Public/Elements/Columns/TypedElementPackageColumns.h
@@ -14,7 +14,7 @@
 /**
  * Column that stores the path of a package.
  */
-USTRUCT()
+USTRUCT(meta = (DisplayName = "Package path"))
 struct FTypedElementPackagePathColumn final : public FTypedElementDataStorageColumn
 {
 	GENERATED_BODY()
@@ -31,7 +31,7 @@ inline uint32 GetTypeHash(const FTypedElementPackagePathColumn& InStruct)
 /**
  * Column that stores the full loading path to a package.
  */
-USTRUCT()
+USTRUCT(meta = (DisplayName = "Package loaded path"))
 struct FTypedElementPackageLoadedPathColumn final : public FTypedElementDataStorageColumn
 {
 	GENERATED_BODY()
diff --git a/Engine/Source/Runtime/TypedElementFramework/Public/Elements/Columns/TypedElementSlateWidgetColumns.h b/Engine/Source/Runtime/TypedElementFramework/Public/Elements/Columns/TypedElementSlateWidgetColumns.h
--- a/Engine/Source/Runtime/TypedElementFramework/Public/Elements/Columns/TypedElementSlateWidgetColumns.h
+++ b/Engine/Source/Runtime/TypedElementFramework/Public/Elements/Columns/TypedElementSlateWidgetColumns.h
@@ -16,7 +16,7 @@ class SWidget;
  * FTypedElementSlateWidgetDeletesRowTag is found then the entire row will
  * be deleted.
  */
-USTRUCT()
+USTRUCT(meta = (DisplayName = "Slate widget reference"))
 struct FTypedElementSlateWidgetReferenceColumn final : public FTypedElementDataStorageColumn
 {
 	GENERATED_BODY()
@@ -29,7 +29,7 @@ struct FTypedElementSlateWidgetReferenceColumn final : public FTypedElementDataS
  * FTypedElementSlateWidgetReferenceColumn is no longer valid, otherwise only
  * the column will be removed.
  */
-USTRUCT()
+USTRUCT(meta = (DisplayName = "Slate widget reference deletes row"))
 struct FTypedElementSlateWidgetReferenceDeletesRowTag final : public FTypedElementDataStorageTag
 {
 	GENERATED_BODY()
diff --git a/Engine/Source/Runtime/TypedElementFramework/Public/Elements/Columns/TypedElementTransformColumns.h b/Engine/Source/Runtime/TypedElementFramework/Public/Elements/Columns/TypedElementTransformColumns.h
--- a/Engine/Source/Runtime/TypedElementFramework/Public/Elements/Columns/TypedElementTransformColumns.h
+++ b/Engine/Source/Runtime/TypedElementFramework/Public/Elements/Columns/TypedElementTransformColumns.h
@@ -12,7 +12,7 @@
 /**
  * Column that stores a local transform. 
  */
-USTRUCT()
+USTRUCT(meta = (DisplayName = "Local Transform"))
 struct FTypedElementLocalTransformColumn final : public FTypedElementDataStorageColumn
 {
 	GENERATED_BODY()
diff --git a/Engine/Source/Runtime/TypedElementFramework/Public/Elements/Columns/TypedElementTypeInfoColumns.h b/Engine/Source/Runtime/TypedElementFramework/Public/Elements/Columns/TypedElementTypeInfoColumns.h
--- a/Engine/Source/Runtime/TypedElementFramework/Public/Elements/Columns/TypedElementTypeInfoColumns.h
+++ b/Engine/Source/Runtime/TypedElementFramework/Public/Elements/Columns/TypedElementTypeInfoColumns.h
@@ -14,7 +14,7 @@ class UScriptStruct;
 /**
  * Column that stores type information for classes.
  */
-USTRUCT()
+USTRUCT(meta = (DisplayName = "Class type info"))
 struct FTypedElementClassTypeInfoColumn final : public FTypedElementDataStorageColumn
 {
 	GENERATED_BODY()
@@ -25,7 +25,7 @@ struct FTypedElementClassTypeInfoColumn final : public FTypedElementDataStorageC
 /**
  * Column that stores type information for structs.
  */
-USTRUCT()
+USTRUCT(meta = (DisplayName = "ScriptStruct type info"))
 struct FTypedElementScriptStructTypeInfoColumn final : public FTypedElementDataStorageColumn
 {
 	GENERATED_BODY()
diff --git a/Engine/Source/Runtime/TypedElementFramework/Public/Elements/Columns/TypedElementValueCacheColumns.h b/Engine/Source/Runtime/TypedElementFramework/Public/Elements/Columns/TypedElementValueCacheColumns.h
--- a/Engine/Source/Runtime/TypedElementFramework/Public/Elements/Columns/TypedElementValueCacheColumns.h
+++ b/Engine/Source/Runtime/TypedElementFramework/Public/Elements/Columns/TypedElementValueCacheColumns.h
@@ -11,7 +11,7 @@
 /**
  * Column that can be used to cache an unsigned 32-bit value in.
  */
-USTRUCT()
+USTRUCT(meta = (DisplayName = "Uint32 value cache"))
 struct FTypedElementU32IntValueCacheColumn final : public FTypedElementDataStorageColumn
 {
 	GENERATED_BODY()
@@ -23,7 +23,7 @@ struct FTypedElementU32IntValueCacheColumn final : public FTypedElementDataStora
 /**
  * Column that can be used to cache a signed 32-bit value in.
  */
-USTRUCT()
+USTRUCT(meta = (DisplayName = "Int32 value cache"))
 struct FTypedElementI32IntValueCacheColumn final : public FTypedElementDataStorageColumn
 {
 	GENERATED_BODY()
@@ -35,7 +35,7 @@ struct FTypedElementI32IntValueCacheColumn final : public FTypedElementDataStora
 /**
  * Column that can be used to cache a 32-bit floating point value in.
  */
-USTRUCT()
+USTRUCT(meta = (DisplayName = "float value cache"))
 struct FTypedElementFloatValueCacheColumn final : public FTypedElementDataStorageColumn
 {
 	GENERATED_BODY()
