Two small changes to TEDS:

- Added support for caching 64-bit values as well.
- Moved updating transform information to the post physics phase. This avoids the need for transform being added in an earlier phase from also having to initialize the transform.

#nojira

diff --git a/Engine/Plugins/Experimental/TypedElementsDataStorage/Source/TypedElementsDataStorage/Private/Compatibility/TypedElementActorTransformProcessors.cpp b/Engine/Plugins/Experimental/TypedElementsDataStorage/Source/TypedElementsDataStorage/Private/Compatibility/TypedElementActorTransformProcessors.cpp
--- a/Engine/Plugins/Experimental/TypedElementsDataStorage/Source/TypedElementsDataStorage/Private/Compatibility/TypedElementActorTransformProcessors.cpp
+++ b/Engine/Plugins/Experimental/TypedElementsDataStorage/Source/TypedElementsDataStorage/Private/Compatibility/TypedElementActorTransformProcessors.cpp
@@ -48,7 +48,7 @@ void UTypedElementActorTransformFactory::RegisterActorLocalTransformToColumn(ITy
 	DataStorage.RegisterQuery(
 		Select(
 			TEXT("Sync actor transform to column"),
-			FProcessor(DSI::EQueryTickPhase::PrePhysics, DataStorage.GetQueryTickGroupName(DSI::EQueryTickGroups::SyncExternalToDataStorage))
+			FProcessor(DSI::EQueryTickPhase::PostPhysics, DataStorage.GetQueryTickGroupName(DSI::EQueryTickGroups::SyncExternalToDataStorage))
 				.ForceToGameThread(true),
 			[](DSI::IQueryContext& Context, TypedElementRowHandle Row, const FMassActorFragment& Actor, FTypedElementLocalTransformColumn& Transform)
 			{
diff --git a/Engine/Source/Runtime/TypedElementFramework/Public/Elements/Columns/TypedElementValueCacheColumns.h b/Engine/Source/Runtime/TypedElementFramework/Public/Elements/Columns/TypedElementValueCacheColumns.h
--- a/Engine/Source/Runtime/TypedElementFramework/Public/Elements/Columns/TypedElementValueCacheColumns.h
+++ b/Engine/Source/Runtime/TypedElementFramework/Public/Elements/Columns/TypedElementValueCacheColumns.h
@@ -32,6 +32,30 @@ struct FTypedElementI32IntValueCacheColumn final : public FTypedElementDataStora
 	int32 Value;
 };
 
+/**
+ * Column that can be used to cache an unsigned 64-bit value in.
+ */
+USTRUCT(meta = (DisplayName = "Uint64 value cache"))
+struct FTypedElementU64IntValueCacheColumn final : public FTypedElementDataStorageColumn
+{
+	GENERATED_BODY()
+
+	UPROPERTY(meta = (IgnoreForMemberInitializationTest))
+	uint64 Value;
+};
+
+/**
+ * Column that can be used to cache a signed 64-bit value in.
+ */
+USTRUCT(meta = (DisplayName = "Int64 value cache"))
+struct FTypedElementI64IntValueCacheColumn final : public FTypedElementDataStorageColumn
+{
+	GENERATED_BODY()
+
+	UPROPERTY(meta = (IgnoreForMemberInitializationTest))
+	int64 Value;
+};
+
 /**
  * Column that can be used to cache a 32-bit floating point value in.
  */
