Allowed subsystems used by TEDS to indicate whether or not they're dynamically reloaded and if they require the main thread.

diff --git a/Engine/Plugins/Experimental/TypedElementsDataStorage/Source/TypedElementsDataStorage/Public/TypedElementSubsystems.h b/Engine/Plugins/Experimental/TypedElementsDataStorage/Source/TypedElementsDataStorage/Public/TypedElementSubsystems.h
--- a/Engine/Plugins/Experimental/TypedElementsDataStorage/Source/TypedElementsDataStorage/Public/TypedElementSubsystems.h
+++ b/Engine/Plugins/Experimental/TypedElementsDataStorage/Source/TypedElementsDataStorage/Public/TypedElementSubsystems.h
@@ -21,6 +21,9 @@ class TYPEDELEMENTSDATASTORAGE_API UTypedElementDataStorageSubsystem : public UE
 	GENERATED_BODY()
 
 public:
+	static constexpr bool bRequiresGameThread = true;
+	static constexpr bool bIsHotReloadable = false;
+
 	~UTypedElementDataStorageSubsystem() override;
 
 	ITypedElementDataStorageInterface* Get();
@@ -40,6 +43,9 @@ class TYPEDELEMENTSDATASTORAGE_API UTypedElementDataStorageUiSubsystem : public
 	GENERATED_BODY()
 
 public:
+	static constexpr bool bRequiresGameThread = true;
+	static constexpr bool bIsHotReloadable = false;
+
 	~UTypedElementDataStorageUiSubsystem() override;
 
 	ITypedElementDataStorageUiInterface* Get();
@@ -50,8 +56,8 @@ protected:
 };
 
 /**
- * A subsystem to provide alternative access to the Typed Elements Data Storage Compatiblity. This can be used in situations where directly 
- * accessing the Compatiblity extension from the Typed Elements Registry is not recommended, such as for MASS.
+ * A subsystem to provide alternative access to the Typed Elements Data Storage Compatibility. This can be used in situations where directly 
+ * accessing the Compatibility extension from the Typed Elements Registry is not recommended, such as for MASS.
  */
 UCLASS()
 class TYPEDELEMENTSDATASTORAGE_API UTypedElementDataStorageCompatibilitySubsystem : public UEditorSubsystem
@@ -59,6 +65,9 @@ class TYPEDELEMENTSDATASTORAGE_API UTypedElementDataStorageCompatibilitySubsyste
 	GENERATED_BODY()
 
 public:
+	static constexpr bool bRequiresGameThread = true;
+	static constexpr bool bIsHotReloadable = false;
+
 	~UTypedElementDataStorageCompatibilitySubsystem() override;
 
 	ITypedElementDataStorageCompatibilityInterface* Get();
diff --git a/Engine/Source/Runtime/TypedElementFramework/Public/Elements/Framework/TypedElementQueryBuilder.h b/Engine/Source/Runtime/TypedElementFramework/Public/Elements/Framework/TypedElementQueryBuilder.h
--- a/Engine/Source/Runtime/TypedElementFramework/Public/Elements/Framework/TypedElementQueryBuilder.h
+++ b/Engine/Source/Runtime/TypedElementFramework/Public/Elements/Framework/TypedElementQueryBuilder.h
@@ -24,10 +24,10 @@ class USubsystem;
  * Calls to the sections become increasingly restrictive, e.g. after calling Where only DependsOn can be
  * called again.
  * 
- * Arguments to the various functions take a pointer to a description of a UStruct. These can be provived
+ * Arguments to the various functions take a pointer to a description of a UStruct. These can be provided
  * in the follow ways:
  * - By using the templated version, e.g. Any<FStructExample>()
- * - By callling the static StaticStruct() function on the UStruct, e.g. FStructExample::StaticStruct();
+ * - By calling the static StaticStruct() function on the UStruct, e.g. FStructExample::StaticStruct();
  * - By name using the Type or TypeOptional string operator, e.g. "/Script/ExamplePackage.FStructExample"_Type or
  *		"/Script/OptionalPackage.FStructOptional"_TypeOptional
  * All functions allow for a single type to be added or a list of types, e.g. ReadOnly(Type<FStructExample>() or
@@ -98,7 +98,7 @@ class USubsystem;
  *			.ReadWrite(USystemExample2::StaticClass())
  *		.Compile();
  *
- * Creating a query is expensive on the builder and the backend side. It's therefore recommended to create a query
+ * Creating a query is expensive on the builder and the back-end side. It's therefore recommended to create a query
  * and store its compiled form for repeated use instead of rebuilding the query on every update.
  */
 
diff --git a/Engine/Source/Runtime/TypedElementFramework/Public/Elements/Framework/TypedElementQueryBuilder.inl b/Engine/Source/Runtime/TypedElementFramework/Public/Elements/Framework/TypedElementQueryBuilder.inl
--- a/Engine/Source/Runtime/TypedElementFramework/Public/Elements/Framework/TypedElementQueryBuilder.inl
+++ b/Engine/Source/Runtime/TypedElementFramework/Public/Elements/Framework/TypedElementQueryBuilder.inl
@@ -5,13 +5,14 @@
 #include <utility>
 #include "Templates/IsConst.h"
 #include "Templates/UnrealTypeTraits.h"
+#include "Elements/Interfaces/TypedElementDataStorageCompatibilityInterface.h"
 
 namespace TypedElementQueryBuilder
 {
 	namespace Internal
 	{
 		// This assumes that the types are unique, but for queries this should be true and otherwise
-		// both resuls would point to the first found index.
+		// both results would point to the first found index.
 		template<typename Target, typename ArgsCurrent, typename... ArgsRemainder>
 		constexpr uint32 GetVarArgIndex()
 		{
@@ -42,13 +43,18 @@ namespace TypedElementQueryBuilder
 		template<typename Type>
 		constexpr ITypedElementDataStorageInterface::EQueryDependencyFlags GetDependencyFlags()
 		{
-			ITypedElementDataStorageInterface::EQueryDependencyFlags Result = ITypedElementDataStorageInterface::EQueryDependencyFlags::None;
-			// Until there's a way to pass in whether or not a dependency is tied to the main thread and whether
-			// it's safe to not update in between updates, default to using the game thread and always update.
-			EnumAddFlags(Result, ITypedElementDataStorageInterface::EQueryDependencyFlags::GameThreadBound);
-			EnumAddFlags(Result, ITypedElementDataStorageInterface::EQueryDependencyFlags::AlwaysRefresh);
-
 			using BaseType = typename std::remove_reference_t<Type>;
+			using SubsystemTraits = TTypedElementSubsystemTraits<std::remove_const_t<BaseType>>;
+
+			ITypedElementDataStorageInterface::EQueryDependencyFlags Result = ITypedElementDataStorageInterface::EQueryDependencyFlags::None;
+			if constexpr (SubsystemTraits::RequiresGameThread())
+			{
+				EnumAddFlags(Result, ITypedElementDataStorageInterface::EQueryDependencyFlags::GameThreadBound);
+			}
+			if constexpr (SubsystemTraits::IsHotReloadable())
+			{
+				EnumAddFlags(Result, ITypedElementDataStorageInterface::EQueryDependencyFlags::AlwaysRefresh);
+			}
 			if constexpr (TIsConst<BaseType>::Value)
 			{
 				EnumAddFlags(Result, ITypedElementDataStorageInterface::EQueryDependencyFlags::ReadOnly);
diff --git a/Engine/Source/Runtime/TypedElementFramework/Public/Elements/Interfaces/TypedElementDataStorageCompatibilityInterface.h b/Engine/Source/Runtime/TypedElementFramework/Public/Elements/Interfaces/TypedElementDataStorageCompatibilityInterface.h
--- a/Engine/Source/Runtime/TypedElementFramework/Public/Elements/Interfaces/TypedElementDataStorageCompatibilityInterface.h
+++ b/Engine/Source/Runtime/TypedElementFramework/Public/Elements/Interfaces/TypedElementDataStorageCompatibilityInterface.h
@@ -29,4 +29,60 @@ public:
 	virtual void AddCompatibleObject(AActor* Actor) = 0;
 	virtual void RemoveCompatibleObject(AActor* Actor) = 0;
 	virtual TypedElementRowHandle FindRowWithCompatibleObject(const TObjectKey<const AActor> Actor) const = 0;
+};
+
+template<typename Subsystem>
+struct TTypedElementSubsystemTraits final
+{
+	template<typename T, typename = void>
+	struct HasRequiresGameThreadVariable 
+	{ 
+		static constexpr bool bAvailable = false; 
+	};
+	template<typename T>
+	struct HasRequiresGameThreadVariable <T, decltype((void)T::bRequiresGameThread)>
+	{ 
+		static constexpr bool bAvailable = true; 
+	};
+
+	template<typename T, typename = void>
+	struct HasIsHotReloadableVariable
+	{ 
+		static constexpr bool bAvailable = false;
+	};
+	template<typename T>
+	struct HasIsHotReloadableVariable <T, decltype((void)T::bIsHotReloadable)>
+	{ 
+		static constexpr bool bAvailable = true;
+	};
+
+	static constexpr bool RequiresGameThread()
+	{
+		if constexpr (HasRequiresGameThreadVariable<Subsystem>::bAvailable)
+		{
+			return Subsystem::bRequiresGameThread;
+		}
+		else
+		{
+			static_assert(HasRequiresGameThreadVariable<Subsystem>::bAvailable, "Subsystem provided to the Typed Elements did not "
+				"have a 'static constexpr bool bRequiresGameThread = true|false` declared or have a specialization for "
+				"TTypedElementSubsystemTraits.");
+			return true;
+		}
+	}
+
+	static constexpr bool IsHotReloadable()
+	{
+		if constexpr (HasIsHotReloadableVariable<Subsystem>::bAvailable)
+		{
+			return Subsystem::bIsHotReloadable;
+		}
+		else
+		{
+			static_assert(HasIsHotReloadableVariable<Subsystem>::bAvailable, "Subsystem provided to the Typed Elements did not "
+				"have a 'static constexpr bool bIsHotReloadable = true|false` declared or have a specialization for "
+				"TTypedElementSubsystemTraits.");
+			return false;
+		}
+	}
 };
\ No newline at end of file
