Moved query information out of TEDS' main interface.

The query information that was embedded in the ITypedElementDataStorageInterface has been moved into a few separate headers. This was done to allow more flexibility for the query builder in the future and to keep the interface header manageable.

diff --git a/Engine/Source/Runtime/TypedElementFramework/Public/Elements/Common/TypedElementHandles.h b/Engine/Source/Runtime/TypedElementFramework/Public/Elements/Common/TypedElementHandles.h
new file mode 100644
--- /dev/null
+++ b/Engine/Source/Runtime/TypedElementFramework/Public/Elements/Common/TypedElementHandles.h
@@ -0,0 +1,17 @@
+// Copyright Epic Games, Inc. All Rights Reserved.
+
+#pragma once
+
+#include "Math/NumericLimits.h"
+
+namespace TypedElementDataStorage
+{
+	using TableHandle = uint64;
+	static constexpr TableHandle InvalidTableHandle = TNumericLimits<TableHandle>::Max();
+
+	using RowHandle = uint64;
+	static constexpr RowHandle InvalidRowHandle = TNumericLimits<RowHandle>::Max();
+
+	using QueryHandle = uint64;
+	static constexpr QueryHandle InvalidQueryHandle = TNumericLimits<QueryHandle>::Max();
+} // namespace TypedElementDataStorage
diff --git a/Engine/Source/Runtime/TypedElementFramework/Public/Elements/Common/TypedElementQueryDescription.h b/Engine/Source/Runtime/TypedElementFramework/Public/Elements/Common/TypedElementQueryDescription.h
new file mode 100644
--- /dev/null
+++ b/Engine/Source/Runtime/TypedElementFramework/Public/Elements/Common/TypedElementQueryDescription.h
@@ -0,0 +1,98 @@
+// Copyright Epic Games, Inc. All Rights Reserved.
+
+#pragma once
+
+#include "Containers/Array.h"
+#include "Containers/ContainerAllocationPolicies.h"
+#include "Elements/Common/TypedElementHandles.h"
+#include "Elements/Common/TypedElementQueryTypes.h"
+#include "Elements/Interfaces/TypedElementQueryStorageInterfaces.h"
+#include "Templates/Function.h"
+#include "UObject/NameTypes.h"
+#include "UObject/WeakObjectPtrTemplates.h"
+
+namespace TypedElementDataStorage
+{
+	struct FQueryDescription;
+	
+	using QueryCallback = TFunction<void(const FQueryDescription&, IQueryContext&)>;
+	
+	struct FQueryDescription final
+	{
+		static constexpr int32 NumInlineSelections = 8;
+		static constexpr int32 NumInlineConditions = 8;
+		static constexpr int32 NumInlineDependencies = 2;
+		static constexpr int32 NumInlineGroups = 2;
+
+		enum class EActionType : uint8
+		{
+			None,	//< Do nothing.
+			Select,	//< Selects a set of columns for further processing.
+			Count,	//< Counts the number of entries that match the filter condition.
+
+			Max //< Value indicating the maximum value in this enum. Not to be used as an enum value.
+		};
+
+		using OperatorIndex = int32;
+		enum class EOperatorType : uint16
+		{
+			SimpleAll,			//< Unary: Type
+			SimpleAny,			//< Unary: Type
+			SimpleNone,			//< Unary: Type
+			SimpleOptional,		//< Unary: Type
+			And,				//< Binary: left operator index, right operator index
+			Or,					//< Binary: left operator index, right operator index
+			Not,				//< Unary: condition index
+			Type,				//< Unary: Type
+
+			Max //< Value indicating the maximum value in this enum. Not to be used as an enum value.
+		};
+
+		struct FBinaryOperator final
+		{
+			OperatorIndex Left;
+			OperatorIndex Right;
+		};
+
+		union FOperator
+		{
+			FBinaryOperator Binary;
+			OperatorIndex Unary;
+			TWeakObjectPtr<const UScriptStruct> Type;
+		};
+
+		struct FCallbackData
+		{
+			TArray<FName, TInlineAllocator<NumInlineGroups>> BeforeGroups;
+			TArray<FName, TInlineAllocator<NumInlineGroups>> AfterGroups;
+			QueryCallback Function;
+			FName Name;
+			FName Group;
+			const UScriptStruct* MonitoredType{ nullptr };
+			EQueryCallbackType Type{ EQueryCallbackType::None };
+			EQueryTickPhase Phase;
+			bool bForceToGameThread{ false };
+		};
+
+		FCallbackData Callback;
+
+		// The list of arrays below are required to remain in the same order as they're added as the function binding expects certain entries
+		// to be in a specific location.
+
+		TArray<TWeakObjectPtr<const UScriptStruct>, TInlineAllocator<NumInlineSelections>> SelectionTypes;
+		TArray<EQueryAccessType, TInlineAllocator<NumInlineSelections>> SelectionAccessTypes;
+
+		TArray<EOperatorType, TInlineAllocator<NumInlineConditions>> ConditionTypes;
+		TArray<FOperator, TInlineAllocator<NumInlineConditions>> ConditionOperators;
+
+		TArray<TWeakObjectPtr<const UClass>, TInlineAllocator<NumInlineDependencies>> DependencyTypes;
+		TArray<EQueryDependencyFlags, TInlineAllocator<NumInlineDependencies>> DependencyFlags;
+		/** Cached instances of the dependencies. This will always match the count of the other Dependency*Types, but may contain null pointers. */
+		TArray<TWeakObjectPtr<UObject>, TInlineAllocator<NumInlineDependencies>> CachedDependencies;
+		TArray<QueryHandle> Subqueries;
+
+		EActionType Action;
+		/** If true, this query only has simple operations and is guaranteed to be executed fully and at optimal performance. */
+		bool bSimpleQuery{ false };
+	};
+} // namespace TypedElementDataStorage
diff --git a/Engine/Source/Runtime/TypedElementFramework/Public/Elements/Common/TypedElementQueryTypes.h b/Engine/Source/Runtime/TypedElementFramework/Public/Elements/Common/TypedElementQueryTypes.h
new file mode 100644
--- /dev/null
+++ b/Engine/Source/Runtime/TypedElementFramework/Public/Elements/Common/TypedElementQueryTypes.h
@@ -0,0 +1,109 @@
+// Copyright Epic Games, Inc. All Rights Reserved.
+
+#pragma once
+
+#include "CoreMinimal.h"
+
+namespace TypedElementDataStorage
+{
+	enum class EQueryTickPhase : uint8
+	{
+		PrePhysics, //< Queries are executed before physics simulation starts.
+		DuringPhysics, //< Queries that can be run in parallel with physics simulation work.
+		PostPhysics, //< Queries that need rigid body and cloth simulation to be completed before being executed.
+		FrameEnd, //< Catchall for queries demoted to the last possible moment.
+
+		Max //< Value indicating the maximum value in this enum. Not to be used as an enum value.
+	};
+
+	enum class EQueryTickGroups : uint8
+	{
+		/** The standard group to run work in. */
+		Default,
+
+		/**
+		 * The group for queries that need to sync data from external sources such as subsystems or the world into
+		 * the Data Storage. These typically run early in a phase.
+		 */
+		SyncExternalToDataStorage,
+		/**
+		 * The group for queries that need to sync data from the Data Storage to external sources such as subsystems
+		 * or the world into. These typically run late in a phase.
+		 */
+		SyncDataStorageToExternal,
+		/**
+		 * Queries grouped under this name will sync data to/from widgets.
+		 */
+		SyncWidgets,
+
+		Max //< Value indicating the maximum value in this enum. Not to be used as an enum value.
+	};
+
+	enum class EQueryCallbackType : uint8
+	{
+		/** No callback provided. */
+		None,
+		/** The query will be run every tick if at least one row matches. */
+		Processor,
+		/** The query will be run when a row is added that matches the query. The first recorded column will be actively monitored for changes. */
+		ObserveAdd,
+		/** The query will be run when a row is removed that matches the query. The first recorded column will be actively monitored for changes. */
+		ObserveRemove,
+		/**
+		 * At the start of the assigned phase this query will run if there are any matches. These queries will have any deferred operations such as
+		 * adding/removing rows/columns executed before the phase starts. This introduces sync points that hinder performance and are therefore
+		 * recommended only for queries that save on work later in the phase such as repeated checks for validity.
+		 */
+		PhasePreparation,
+		/**
+		 * At the end of the assigned phase this query will run if there are any matches. These queries will have any deferred operations such as
+		 * adding/removing rows/columns executed before the phase ends. This introduces sync points that hinder performance and are therefore
+		 * recommended only cases where delaying deferred operations is not possible e.g. when tables are known to be referenced outside the update
+		 * cycle.
+		 */
+		PhaseFinalization,
+
+		Max //< Value indicating the maximum value in this enum. Not to be used as an enum value.
+	};
+
+	enum class EQueryAccessType : bool
+	{
+		ReadOnly,
+		ReadWrite
+	};
+
+	enum class EQueryDependencyFlags : uint8
+	{
+		None = 0,
+		/** If set the dependency is accessed as read-only. If not set the dependency requires Read/Write access. */
+		ReadOnly = 1 << 0,
+		/** If set the dependency can only be used from the game thread, otherwise it can be accessed from any thread. */
+		GameThreadBound = 1 << 1,
+		/** If set the dependency will be re-fetched every iteration, otherwise only if not fetched before. */
+		AlwaysRefresh = 1 << 2
+	};
+	ENUM_CLASS_FLAGS(EQueryDependencyFlags);
+
+	struct FQueryResult final
+	{
+		enum class ECompletion
+		{
+			/** Query could be fully executed. */
+			Fully,
+			/** Only portions of the query were executed. This is caused by a problem that was encountered partway through processing. */
+			Partially,
+			/**
+			 * The back-end doesn't support the particular query. This may be a limitation in how/where the query is run or because
+			 * the query contains actions and/or operations that are not supported.
+			 */
+			Unsupported,
+			/** The provided query is no longer available. */
+			Unavailable,
+			/** One or more dependencies declared on the query could not be retrieved. */
+			MissingDependency
+		};
+
+		uint32 Count{ 0 }; /** The number of rows were processed. */
+		ECompletion Completed{ ECompletion::Unavailable };
+	};
+} // namespace TypedElementDataStorage
diff --git a/Engine/Source/Runtime/TypedElementFramework/Public/Elements/Interfaces/TypedElementDataStorageInterface.h b/Engine/Source/Runtime/TypedElementFramework/Public/Elements/Interfaces/TypedElementDataStorageInterface.h
--- a/Engine/Source/Runtime/TypedElementFramework/Public/Elements/Interfaces/TypedElementDataStorageInterface.h
+++ b/Engine/Source/Runtime/TypedElementFramework/Public/Elements/Interfaces/TypedElementDataStorageInterface.h
@@ -5,7 +5,11 @@
 #include "Containers/Array.h"
 #include "Containers/ContainersFwd.h"
 #include "Delegates/Delegate.h"
+#include "Elements/Common/TypedElementHandles.h"
+#include "Elements/Common/TypedElementQueryDescription.h"
+#include "Elements/Common/TypedElementQueryTypes.h"
 #include "Elements/Framework/TypedElementColumnUtils.h"
+#include "Elements/Interfaces/TypedElementQueryStorageInterfaces.h"
 #include "Math/NumericLimits.h"
 #include "Misc/EnumClassFlags.h"
 #include "Templates/Function.h"
@@ -27,12 +31,12 @@ struct ColumnDataResult
 	void* Data;
 };
 
-using TypedElementTableHandle = uint64;
-static constexpr auto TypedElementInvalidTableHandle = TNumericLimits<TypedElementTableHandle>::Max();
-using TypedElementRowHandle = uint64;
-static constexpr auto TypedElementInvalidRowHandle = TNumericLimits<TypedElementRowHandle>::Max();
-using TypedElementQueryHandle = uint64;
-static constexpr auto TypedElementInvalidQueryHandle = TNumericLimits<TypedElementQueryHandle>::Max();
+using TypedElementTableHandle = TypedElementDataStorage::TableHandle;
+static constexpr auto TypedElementInvalidTableHandle = TypedElementDataStorage::InvalidTableHandle;
+using TypedElementRowHandle = TypedElementDataStorage::RowHandle;
+static constexpr auto TypedElementInvalidRowHandle = TypedElementDataStorage::InvalidRowHandle;
+using TypedElementQueryHandle = TypedElementDataStorage::QueryHandle;
+static constexpr auto TypedElementInvalidQueryHandle = TypedElementDataStorage::InvalidQueryHandle;
 
 using FTypedElementOnDataStorageCreation = FSimpleMulticastDelegate;
 using FTypedElementOnDataStorageDestruction = FSimpleMulticastDelegate;
@@ -208,320 +212,24 @@ public:
 	 * can be used directly in the processor to do additional filtering. This will however impact performance and it's 
 	 * therefore recommended to try to simplify the query first before relying on extended query filtering in a processor.
 	 */
-	
-	enum class EQueryTickPhase : uint8
-	{
-		PrePhysics, //< Queries are executed before physics simulation starts.
-		DuringPhysics, //< Queries that can be run in parallel with physics simulation work.
-		PostPhysics, //< Queries that need rigid body and cloth simulation to be completed before being executed.
-		FrameEnd, //< Catchall for queries demoted to the last possible moment.
 
-		Max //< Value indicating the maximum value in this enum. Not to be used as an enum value.
-	};
-
-	enum class EQueryTickGroups : uint8
-	{
-		/** The standard group to run work in. */
-		Default,
-		
-		/** 
-		 * The group for queries that need to sync data from external sources such as subsystems or the world into
-		 * the Data Storage. These typically run early in a phase.
-		 */
-		SyncExternalToDataStorage,
-		/**
-		 * The group for queries that need to sync data from the Data Storage to external sources such as subsystems
-		 * or the world into. These typically run late in a phase.
-		 */
-		SyncDataStorageToExternal,
-		/**
-		 * Queries grouped under this name will sync data to/from widgets.
-		 */
-		SyncWidgets,
-
-		Max //< Value indicating the maximum value in this enum. Not to be used as an enum value.
-	};
-
-	enum class EQueryCallbackType : uint8
-	{
-		/** No callback provided. */
-		None,
-		/** The query will be run every tick if at least one row matches. */
-		Processor,
-		/** The query will be run when a row is added that matches the query. The first recorded column will be actively monitored for changes. */
-		ObserveAdd,
-		/** The query will be run when a row is removed that matches the query. The first recorded column will be actively monitored for changes. */
-		ObserveRemove,
-		/** 
-		 * At the start of the assigned phase this query will run if there are any matches. These queries will have any deferred operations such as
-		 * adding/removing rows/columns executed before the phase starts. This introduces sync points that hinder performance and are therefore
-		 * recommended only for queries that save on work later in the phase such as repeated checks for validity.
-		 */
-		PhasePreparation,
-		/**
-		 * At the end of the assigned phase this query will run if there are any matches. These queries will have any deferred operations such as
-		 * adding/removing rows/columns executed before the phase ends. This introduces sync points that hinder performance and are therefore
-		 * recommended only cases where delaying deferred operations is not possible e.g. when tables are known to be referenced outside the update
-		 * cycle.
-		 */
-		 PhaseFinalization,
-
-		 Max //< Value indicating the maximum value in this enum. Not to be used as an enum value.
-	};
-
-	enum class EQueryAccessType : bool
-	{ 
-		ReadOnly, 
-		ReadWrite 
-	};
-
-	enum class EQueryDependencyFlags : uint8
-	{
-		None = 0,
-		/** If set the dependency is accessed as read-only. If not set the dependency requires Read/Write access. */
-		ReadOnly = 1 << 0,
-		/** If set the dependency can only be used from the game thread, otherwise it can be accessed from any thread. */
-		GameThreadBound = 1 << 1,
-		/** If set the dependency will be re-fetched every iteration, otherwise only if not fetched before. */
-		AlwaysRefresh = 1 << 2
-	};
-
-	struct FQueryResult
-	{
-		enum class ECompletion
-		{
-			/** Query could be fully executed. */
-			Fully,
-			/** Only portions of the query were executed. This is caused by a problem that was encountered partway through processing. */
-			Partially,
-			/**
-			 * The back-end doesn't support the particular query. This may be a limitation in how/where the query is run or because
-			 * the query contains actions and/or operations that are not supported.
-			 */
-			 Unsupported,
-			 /** The provided query is no longer available. */
-			 Unavailable,
-			 /** One or more dependencies declared on the query could not be retrieved. */
-			 MissingDependency
-		};
-
-		uint32 Count{ 0 }; /** The number of rows were processed. */
-		ECompletion Completed{ ECompletion::Unavailable };
-	};
-
-	/**
-	 * Base interface for any contexts provided to query callbacks.
-	 */
-	struct ICommonQueryContext
-	{
-		virtual ~ICommonQueryContext() = default;
-
-		/** Return the address of a immutable column matching the requested type or a nullptr if not found. */
-		virtual const void* GetColumn(const UScriptStruct* ColumnType) const = 0;
-		/** Return the address of a mutable column matching the requested type or a nullptr if not found. */
-		virtual void* GetMutableColumn(const UScriptStruct* ColumnType) = 0;
-		/**
-		 * Get a list of columns or nullptrs if the column type wasn't found. Mutable addresses are returned and it's up to
-		 * the caller to not change immutable addresses.
-		 */
-		virtual void GetColumns(TArrayView<char*> RetrievedAddresses, TConstArrayView<TWeakObjectPtr<const UScriptStruct>> ColumnTypes,
-			TConstArrayView<ITypedElementDataStorageInterface::EQueryAccessType> AccessTypes) = 0;
-		/**
-		 * Get a list of columns or nullptrs if the column type wasn't found. Mutable addresses are returned and it's up to
-		 * the caller to not change immutable addresses. This version doesn't verify that the enough space is provided and
-		 * it's up to the caller to guarantee the target addresses have enough space.
-		 */
-		virtual void GetColumnsUnguarded(int32 TypeCount, char** RetrievedAddresses, const TWeakObjectPtr<const UScriptStruct>* ColumnTypes,
-			const ITypedElementDataStorageInterface::EQueryAccessType* AccessTypes) = 0;
-
-		/** Returns the number rows in the batch. */
-		virtual uint32 GetRowCount() const = 0;
-		/**
-		 * Returns an immutable view that contains the row handles for all returned results. The returned size will be the same  as the
-		 * value returned by GetRowCount().
-		 */
-		virtual TConstArrayView<TypedElementRowHandle> GetRowHandles() const = 0;
-
-		// Utility functions
-
-		template<typename Column>
-		const Column* GetColumn() const;
-		template<typename Column>
-		Column* GetMutableColumn();
-	};
-	
-	/** 
-	 * Interface to be provided to query callbacks running with the Data Storage.
-	 * Note that at the time of writing only subclasses of Subsystem are supported as dependencies.
-	 */
-	struct IQueryContext : public ICommonQueryContext
-	{
-		virtual ~IQueryContext() = default;
-		
-		/** Returns an immutable instance of the requested dependency or a nullptr if not found. */
-		virtual const UObject* GetDependency(const UClass* DependencyClass) = 0;
-		/** Returns a mutable instance of the requested dependency or a nullptr if not found. */
-		virtual UObject* GetMutableDependency(const UClass* DependencyClass) = 0;
-		/** 
-		 * Returns a list of dependencies or nullptrs if a dependency wasn't found. Mutable versions are return and it's up to the 
-		 * caller to not change immutable dependencies.
-		 */
-		virtual void GetDependencies(TArrayView<UObject*> RetrievedAddresses, TConstArrayView<TWeakObjectPtr<const UClass>> DependencyTypes,
-			TConstArrayView<EQueryAccessType> AccessTypes) = 0;
-
-		/**
-		 * Removes the row with the provided row handle. The removal will not be immediately done but delayed until the end of the tick 
-		 * group.
-		 */
-		virtual void RemoveRow(TypedElementRowHandle Row) = 0;
-		/**
-		 * Removes rows with the provided row handles. The removal will not be immediately done but delayed until the end of the tick
-		 * group.
-		 */
-		virtual void RemoveRows(TConstArrayView<TypedElementRowHandle> Rows) = 0;
-
-		/**
-		 * Adds new empty columns to a row of the provided type. The addition will not be immediately done but delayed until the end of the
-		 * tick group.
-		 */
-		virtual void AddColumns(TypedElementRowHandle Row, TConstArrayView<const UScriptStruct*> ColumnTypes) = 0;
-		/**
-		 * Adds new empty columns to the listed rows of the provided type. The addition will not be immediately done but delayed until the end of the
-		 * tick group.
-		 */
-		virtual void AddColumns(TConstArrayView<TypedElementRowHandle> Rows, TConstArrayView<const UScriptStruct*> ColumnTypes) = 0;
-		/**
-		 * Removes columns of the provided types from a row. The removal will not be immediately done but delayed until the end of the
-		 * tick group.
-		 */
-		virtual void RemoveColumns(TypedElementRowHandle Row, TConstArrayView<const UScriptStruct*> ColumnTypes) = 0;
-		/**
-		 * Removes columns of the provided types from the listed rows. The removal will not be immediately done but delayed until the end of the
-		 * tick group.
-		 */
-		virtual void RemoveColumns(TConstArrayView<TypedElementRowHandle> Rows, TConstArrayView<const UScriptStruct*> ColumnTypes) = 0;
-
-		/**
-		 * Runs a previously created query. This version takes an arbitrary query, but is limited to running queries that do not directly 
-		 * access data from rows such as count queries.
-		 */
-		virtual FQueryResult RunQuery(TypedElementQueryHandle Query) = 0;
-		/** Runs a subquery registered with the current query. The subquery index is in the order of registration with the query. */
-		virtual FQueryResult RunSubquery(int32 SubqueryIndex) = 0;
-		
-
-
-		// Utility functions
-
-		template<typename... Columns>
-		void AddColumns(TypedElementRowHandle Row);
-		template<typename... Columns>
-		void AddColumns(TConstArrayView<TypedElementRowHandle> Rows);
-		template<typename... Columns>
-		void RemoveColumns(TypedElementRowHandle Row);
-		template<typename... Columns>
-		void RemoveColumns(TConstArrayView<TypedElementRowHandle> Rows);
-	};
+	using EQueryTickPhase = TypedElementDataStorage::EQueryTickPhase;
+	using EQueryTickGroups = TypedElementDataStorage::EQueryTickGroups;
+	using EQueryCallbackType = TypedElementDataStorage::EQueryCallbackType;
+	using EQueryAccessType = TypedElementDataStorage::EQueryAccessType;
+	using EQueryDependencyFlags = TypedElementDataStorage::EQueryDependencyFlags;
+	using FQueryResult = TypedElementDataStorage::FQueryResult;
 
-	/**
-	 * Interface to be provided to query callbacks that are directly called through RunQuery from outside a query callback.
-	 */
-	struct IDirectQueryContext : public ICommonQueryContext
-	{
-		virtual ~IDirectQueryContext() = default;
-	};
+	using IQueryContext = TypedElementDataStorage::IQueryContext;
+	using IDirectQueryContext = TypedElementDataStorage::IDirectQueryContext;
+	using ISubqueryCallbackContext = TypedElementDataStorage::ISubqueryCallbackContext;
 
-	/**
-	 * Interface to be provided to query callbacks that are directly called through from a query callback.
-	 */
-	struct ISubqueryCallbackContext : public ICommonQueryContext
-	{
-		virtual ~ISubqueryCallbackContext() = default;
-	};
+	using FQueryDescription = TypedElementDataStorage::FQueryDescription;
+	using QueryCallback = TypedElementDataStorage::QueryCallback;
 
-	struct FQueryDescription;
-	using QueryCallback = TFunction<void(const FQueryDescription&, IQueryContext&)>;
 	using QueryCallbackRef = TFunctionRef<void(const FQueryDescription&, IQueryContext&)>;
 	using DirectQueryCallbackRef = TFunctionRef<void(const FQueryDescription&, IDirectQueryContext&)>;
 
-	struct FQueryDescription final
-	{
-		static constexpr int32 NumInlineSelections = 8;
-		static constexpr int32 NumInlineConditions = 8;
-		static constexpr int32 NumInlineDependencies = 2;
-		static constexpr int32 NumInlineGroups = 2;
-
-		enum class EActionType : uint8
-		{
-			None,	//< Do nothing.
-			Select,	//< Selects a set of columns for further processing.
-			Count,	//< Counts the number of entries that match the filter condition.
-
-			Max //< Value indicating the maximum value in this enum. Not to be used as an enum value.
-		};
-
-		using OperatorIndex = int32;
-		enum class EOperatorType : uint16
-		{
-			SimpleAll,			//< Unary: Type
-			SimpleAny,			//< Unary: Type
-			SimpleNone,			//< Unary: Type
-			SimpleOptional,		//< Unary: Type
-			And,				//< Binary: left operator index, right operator index
-			Or,					//< Binary: left operator index, right operator index
-			Not,				//< Unary: condition index
-			Type,				//< Unary: Type
-
-			Max //< Value indicating the maximum value in this enum. Not to be used as an enum value.
-		};
-
-		struct FBinaryOperator final
-		{
-			OperatorIndex Left;
-			OperatorIndex Right;
-		};
-		
-		union FOperator
-		{
-			FBinaryOperator Binary;
-			OperatorIndex Unary;
-			TWeakObjectPtr<const UScriptStruct> Type;
-		};
-
-		struct FCallbackData
-		{
-			TArray<FName, TInlineAllocator<NumInlineGroups>> BeforeGroups;
-			TArray<FName, TInlineAllocator<NumInlineGroups>> AfterGroups;
-			QueryCallback Function;
-			FName Name;
-			FName Group;
-			const UScriptStruct* MonitoredType{ nullptr };
-			EQueryCallbackType Type{ EQueryCallbackType::None };
-			EQueryTickPhase Phase;
-			bool bForceToGameThread{ false };
-		};
-		
-		FCallbackData Callback;
-
-		// The list of arrays below are required to remain in the same order as they're added as the function binding expects certain entries
-		// to be in a specific location.
-
-		TArray<TWeakObjectPtr<const UScriptStruct>, TInlineAllocator<NumInlineSelections>> SelectionTypes;
-		TArray<EQueryAccessType, TInlineAllocator<NumInlineSelections>> SelectionAccessTypes;
-		
-		TArray<EOperatorType, TInlineAllocator<NumInlineConditions>> ConditionTypes;
-		TArray<FOperator, TInlineAllocator<NumInlineConditions>> ConditionOperators;
-		
-		TArray<TWeakObjectPtr<const UClass>, TInlineAllocator<NumInlineDependencies>> DependencyTypes;
-		TArray<EQueryDependencyFlags, TInlineAllocator<NumInlineDependencies>> DependencyFlags;
-		/** Cached instances of the dependencies. This will always match the count of the other Dependency*Types, but may contain null pointers. */
-		TArray<TWeakObjectPtr<UObject>, TInlineAllocator<NumInlineDependencies>> CachedDependencies;
-		TArray<TypedElementQueryHandle> Subqueries;
-		
-		EActionType Action;
-		/** If true, this query only has simple operations and is guaranteed to be executed fully and at optimal performance. */
-		bool bSimpleQuery{ false };
-	};
 	/** 
 	 * Registers a query with the data storage. The description is processed into an internal format and may be changed. If no valid
 	 * could be created an invalid query handle will be returned. It's recommended to use the Query Builder for a more convenient
@@ -544,8 +252,7 @@ public:
 	 * Directly runs a query. The callback will be called for batches of matching rows. During a single call to RunQuery the callback
 	 * may be called multiple times. If the query handle is invalid or has been deleted nothing happens and the callback won't be called
 	 */
-	virtual FQueryResult RunQuery(TypedElementQueryHandle Query, 
-		DirectQueryCallbackRef Callback) = 0;
+	virtual FQueryResult RunQuery(TypedElementQueryHandle Query, DirectQueryCallbackRef Callback) = 0;
 	
 	/**
 	 * @section Misc
@@ -621,8 +328,6 @@ public:
 
 // Implementations
 
-ENUM_CLASS_FLAGS(ITypedElementDataStorageInterface::EQueryDependencyFlags);
-
 template<typename Column>
 bool ITypedElementDataStorageInterface::AddColumn(TypedElementRowHandle Row)
 {
@@ -678,45 +383,3 @@ SystemType* ITypedElementDataStorageInterface::GetExternalSystem()
 {
 	return reinterpret_cast<SystemType*>(GetExternalSystemAddress(SystemType::StaticClass()));
 }
-
-
-
-//
-// ITypedElementDataStorageInterface::IQueryContext
-//
-
-template<typename Column>
-const Column* ITypedElementDataStorageInterface::ICommonQueryContext::GetColumn() const
-{
-	return reinterpret_cast<const Column*>(GetColumn(Column::StaticStruct()));
-}
-
-template<typename Column>
-Column* ITypedElementDataStorageInterface::ICommonQueryContext::GetMutableColumn()
-{
-	return reinterpret_cast<Column*>(GetMutableColumn(Column::StaticStruct()));
-}
-
-template<typename... Columns>
-void ITypedElementDataStorageInterface::IQueryContext::AddColumns(TypedElementRowHandle Row)
-{
-	AddColumns(Row, { Columns::StaticStruct()... });
-}
-
-template<typename... Columns>
-void ITypedElementDataStorageInterface::IQueryContext::AddColumns(TConstArrayView<TypedElementRowHandle> Rows)
-{
-	AddColumns(Rows, { Columns::StaticStruct()... });
-}
-
-template<typename... Columns>
-void ITypedElementDataStorageInterface::IQueryContext::RemoveColumns(TypedElementRowHandle Row)
-{
-	RemoveColumns(Row, { Columns::StaticStruct()... });
-}
-
-template<typename... Columns>
-void ITypedElementDataStorageInterface::IQueryContext::RemoveColumns(TConstArrayView<TypedElementRowHandle> Rows)
-{
-	RemoveColumns(Rows, { Columns::StaticStruct()... });
-}
\ No newline at end of file
diff --git a/Engine/Source/Runtime/TypedElementFramework/Public/Elements/Interfaces/TypedElementQueryStorageInterfaces.h b/Engine/Source/Runtime/TypedElementFramework/Public/Elements/Interfaces/TypedElementQueryStorageInterfaces.h
new file mode 100644
--- /dev/null
+++ b/Engine/Source/Runtime/TypedElementFramework/Public/Elements/Interfaces/TypedElementQueryStorageInterfaces.h
@@ -0,0 +1,191 @@
+// Copyright Epic Games, Inc. All Rights Reserved.
+
+#pragma once
+
+#include "Containers/ContainersFwd.h"
+#include "Elements/Common/TypedElementHandles.h"
+#include "Elements/Common/TypedElementQueryTypes.h"
+#include "UObject/WeakObjectPtrTemplates.h"
+
+class UClass;
+class UObject;
+class UScriptStruct;
+
+namespace TypedElementDataStorage
+{
+	/**
+	 * Base interface for any contexts provided to query callbacks.
+	 */
+	struct ICommonQueryContext
+	{
+		virtual ~ICommonQueryContext() = default;
+
+		/** Return the address of a immutable column matching the requested type or a nullptr if not found. */
+		virtual const void* GetColumn(const UScriptStruct* ColumnType) const = 0;
+		/** Return the address of a mutable column matching the requested type or a nullptr if not found. */
+		virtual void* GetMutableColumn(const UScriptStruct* ColumnType) = 0;
+		/**
+		 * Get a list of columns or nullptrs if the column type wasn't found. Mutable addresses are returned and it's up to
+		 * the caller to not change immutable addresses.
+		 */
+		virtual void GetColumns(TArrayView<char*> RetrievedAddresses, TConstArrayView<TWeakObjectPtr<const UScriptStruct>> ColumnTypes,
+			TConstArrayView<EQueryAccessType> AccessTypes) = 0;
+		/**
+		 * Get a list of columns or nullptrs if the column type wasn't found. Mutable addresses are returned and it's up to
+		 * the caller to not change immutable addresses. This version doesn't verify that the enough space is provided and
+		 * it's up to the caller to guarantee the target addresses have enough space.
+		 */
+		virtual void GetColumnsUnguarded(int32 TypeCount, char** RetrievedAddresses, const TWeakObjectPtr<const UScriptStruct>* ColumnTypes,
+			const EQueryAccessType* AccessTypes) = 0;
+
+		/** Returns the number rows in the batch. */
+		virtual uint32 GetRowCount() const = 0;
+		/**
+		 * Returns an immutable view that contains the row handles for all returned results. The returned size will be the same  as the
+		 * value returned by GetRowCount().
+		 */
+		virtual TConstArrayView<RowHandle> GetRowHandles() const = 0;
+
+		// Utility functions
+
+		template<typename Column>
+		const Column* GetColumn() const;
+		template<typename Column>
+		Column* GetMutableColumn();
+	};
+
+	/**
+	 * Interface to be provided to query callbacks that are directly called through RunQuery from outside a query callback.
+	 */
+	struct IDirectQueryContext : public ICommonQueryContext
+	{
+		virtual ~IDirectQueryContext() = default;
+	};
+
+	/**
+	 * Interface to be provided to query callbacks that are directly called through from a query callback.
+	 */
+	struct ISubqueryCallbackContext : public ICommonQueryContext
+	{
+		virtual ~ISubqueryCallbackContext() = default;
+	};
+
+	/**
+	 * Interface to be provided to query callbacks running with the Data Storage.
+	 * Note that at the time of writing only subclasses of Subsystem are supported as dependencies.
+	 */
+	struct IQueryContext : public ICommonQueryContext
+	{
+		virtual ~IQueryContext() = default;
+
+		/** Returns an immutable instance of the requested dependency or a nullptr if not found. */
+		virtual const UObject* GetDependency(const UClass* DependencyClass) = 0;
+		/** Returns a mutable instance of the requested dependency or a nullptr if not found. */
+		virtual UObject* GetMutableDependency(const UClass* DependencyClass) = 0;
+		/**
+		 * Returns a list of dependencies or nullptrs if a dependency wasn't found. Mutable versions are return and it's up to the
+		 * caller to not change immutable dependencies.
+		 */
+		virtual void GetDependencies(TArrayView<UObject*> RetrievedAddresses, TConstArrayView<TWeakObjectPtr<const UClass>> DependencyTypes,
+			TConstArrayView<EQueryAccessType> AccessTypes) = 0;
+
+		/**
+		 * Removes the row with the provided row handle. The removal will not be immediately done but delayed until the end of the tick
+		 * group.
+		 */
+		virtual void RemoveRow(RowHandle Row) = 0;
+		/**
+		 * Removes rows with the provided row handles. The removal will not be immediately done but delayed until the end of the tick
+		 * group.
+		 */
+		virtual void RemoveRows(TConstArrayView<RowHandle> Rows) = 0;
+
+		/**
+		 * Adds new empty columns to a row of the provided type. The addition will not be immediately done but delayed until the end of the
+		 * tick group.
+		 */
+		virtual void AddColumns(RowHandle Row, TConstArrayView<const UScriptStruct*> ColumnTypes) = 0;
+		/**
+		 * Adds new empty columns to the listed rows of the provided type. The addition will not be immediately done but delayed until the end of the
+		 * tick group.
+		 */
+		virtual void AddColumns(TConstArrayView<RowHandle> Rows, TConstArrayView<const UScriptStruct*> ColumnTypes) = 0;
+		/**
+		 * Removes columns of the provided types from a row. The removal will not be immediately done but delayed until the end of the
+		 * tick group.
+		 */
+		virtual void RemoveColumns(RowHandle Row, TConstArrayView<const UScriptStruct*> ColumnTypes) = 0;
+		/**
+		 * Removes columns of the provided types from the listed rows. The removal will not be immediately done but delayed until the end of the
+		 * tick group.
+		 */
+		virtual void RemoveColumns(TConstArrayView<RowHandle> Rows, TConstArrayView<const UScriptStruct*> ColumnTypes) = 0;
+
+		/**
+		 * Runs a previously created query. This version takes an arbitrary query, but is limited to running queries that do not directly
+		 * access data from rows such as count queries.
+		 */
+		virtual FQueryResult RunQuery(QueryHandle Query) = 0;
+		/** Runs a subquery registered with the current query. The subquery index is in the order of registration with the query. */
+		virtual FQueryResult RunSubquery(int32 SubqueryIndex) = 0;
+
+
+
+		// Utility functions
+
+		template<typename... Columns>
+		void AddColumns(RowHandle Row);
+		template<typename... Columns>
+		void AddColumns(TConstArrayView<RowHandle> Rows);
+		template<typename... Columns>
+		void RemoveColumns(RowHandle Row);
+		template<typename... Columns>
+		void RemoveColumns(TConstArrayView<RowHandle> Rows);
+	};
+} // namespace TypedElementDataStorage
+
+
+
+
+//
+// Implementations
+//
+
+namespace TypedElementDataStorage
+{
+	template<typename Column>
+	const Column* ICommonQueryContext::GetColumn() const
+	{
+		return reinterpret_cast<const Column*>(GetColumn(Column::StaticStruct()));
+	}
+
+	template<typename Column>
+	Column* ICommonQueryContext::GetMutableColumn()
+	{
+		return reinterpret_cast<Column*>(GetMutableColumn(Column::StaticStruct()));
+	}
+
+	template<typename... Columns>
+	void IQueryContext::AddColumns(RowHandle Row)
+	{
+		AddColumns(Row, { Columns::StaticStruct()... });
+	}
+
+	template<typename... Columns>
+	void IQueryContext::AddColumns(TConstArrayView<RowHandle> Rows)
+	{
+		AddColumns(Rows, { Columns::StaticStruct()... });
+	}
+
+	template<typename... Columns>
+	void IQueryContext::RemoveColumns(RowHandle Row)
+	{
+		RemoveColumns(Row, { Columns::StaticStruct()... });
+	}
+
+	template<typename... Columns>
+	void IQueryContext::RemoveColumns(TConstArrayView<RowHandle> Rows)
+	{
+		RemoveColumns(Rows, { Columns::StaticStruct()... });
+	}
+} // namespace TypedElementDataStorage
