Fixed non-unity compile error for missing WeakPtr header.

#nojira

diff --git a/Engine/Source/Runtime/TypedElementFramework/Public/Elements/Common/TypedElementQueryDescription.h b/Engine/Source/Runtime/TypedElementFramework/Public/Elements/Common/TypedElementQueryDescription.h
--- a/Engine/Source/Runtime/TypedElementFramework/Public/Elements/Common/TypedElementQueryDescription.h
+++ b/Engine/Source/Runtime/TypedElementFramework/Public/Elements/Common/TypedElementQueryDescription.h
@@ -9,8 +9,11 @@
 #include "Elements/Interfaces/TypedElementQueryStorageInterfaces.h"
 #include "Templates/Function.h"
 #include "UObject/NameTypes.h"
+#include "UObject/WeakObjectPtr.h"
 #include "UObject/WeakObjectPtrTemplates.h"
 
+class UScriptStruct;
+
 namespace TypedElementDataStorage
 {
 	struct FQueryDescription;
