Introduced the ability to bind to an already created query.

This change introduces the ability to bind a function to an already existing query. It takes the provided functions and creates a wrapper around it that can be used by standard RunQuery calls in TEDS. To support the various different query contexts, the query binder now also no longer requires IQueryContext for it's binding and can use any object that has the minimally required functions. For convenience the current root functions still explicity set the appropriate context as this will produce more readable error messages.

Additional notes:
- Several parts of the query binding were simplified by using C++20 concepts.
- Some cleanup to be able to remove the "using" redirections introduced at a temporary fix with the previous code reorganization.

diff --git a/Engine/Plugins/Experimental/TypedElementsDataStorage/Source/TypedElementsDataStorage/Private/Queries/TypedElementExtendedQueryStore.cpp b/Engine/Plugins/Experimental/TypedElementsDataStorage/Source/TypedElementsDataStorage/Private/Queries/TypedElementExtendedQueryStore.cpp
--- a/Engine/Plugins/Experimental/TypedElementsDataStorage/Source/TypedElementsDataStorage/Private/Queries/TypedElementExtendedQueryStore.cpp
+++ b/Engine/Plugins/Experimental/TypedElementsDataStorage/Source/TypedElementsDataStorage/Private/Queries/TypedElementExtendedQueryStore.cpp
@@ -190,37 +190,30 @@ ITypedElementDataStorageInterface::FQueryResult FTypedElementExtendedQueryStore:
 
 	if (FTypedElementExtendedQuery* QueryData = Get(Query))
 	{
-		if (QueryData->Description.bSimpleQuery)
+		switch (QueryData->Description.Action)
 		{
-			switch (QueryData->Description.Action)
+		case ActionType::None:
+			Result.Completed = CompletionType::Fully;
+			break;
+		case ActionType::Select:
+			if (!QueryData->Processor.IsValid())
+			{
+				Result = FTypedElementQueryProcessorData::Execute(
+					Callback, QueryData->Description, QueryData->NativeQuery, EntityManager);
+			}
+			else
 			{
-			case ActionType::None:
-				Result.Completed = CompletionType::Fully;
-				break;
-			case ActionType::Select:
-				if (!QueryData->Processor.IsValid())
-				{
-					Result = FTypedElementQueryProcessorData::Execute(
-						Callback, QueryData->Description, QueryData->NativeQuery, EntityManager);
-				}
-				else
-				{
-					Result.Completed = CompletionType::Unsupported;
-				}
-				break;
-			case ActionType::Count:
-				// Only the count is requested so no need to trigger the callback.
-				Result.Count = QueryData->NativeQuery.GetNumMatchingEntities(EntityManager);
-				Result.Completed = CompletionType::Fully;
-				break;
-			default:
 				Result.Completed = CompletionType::Unsupported;
-				break;
 			}
-		}
-		else
-		{
+			break;
+		case ActionType::Count:
+			// Only the count is requested so no need to trigger the callback.
+			Result.Count = QueryData->NativeQuery.GetNumMatchingEntities(EntityManager);
+			Result.Completed = CompletionType::Fully;
+			break;
+		default:
 			Result.Completed = CompletionType::Unsupported;
+			break;
 		}
 	}
 	else
diff --git a/Engine/Source/Runtime/TypedElementFramework/Public/Elements/Common/TypedElementQueryDescription.h b/Engine/Source/Runtime/TypedElementFramework/Public/Elements/Common/TypedElementQueryDescription.h
--- a/Engine/Source/Runtime/TypedElementFramework/Public/Elements/Common/TypedElementQueryDescription.h
+++ b/Engine/Source/Runtime/TypedElementFramework/Public/Elements/Common/TypedElementQueryDescription.h
@@ -19,6 +19,11 @@ namespace TypedElementDataStorage
 	struct FQueryDescription;
 	
 	using QueryCallback = TFunction<void(const FQueryDescription&, IQueryContext&)>;
+	using QueryCallbackRef = TFunctionRef<void(const FQueryDescription&, IQueryContext&)>;
+	using DirectQueryCallback = TFunction<void(const FQueryDescription&, IDirectQueryContext&)>;
+	using DirectQueryCallbackRef = TFunctionRef<void(const FQueryDescription&, IDirectQueryContext&)>;
+	using SubqueryCallback = TFunction<void(const FQueryDescription&, ISubqueryContext&)>;
+	using SubqueryCallbackRef = TFunctionRef<void(const FQueryDescription&, ISubqueryContext&)>;
 	
 	struct FQueryDescription final
 	{
diff --git a/Engine/Source/Runtime/TypedElementFramework/Public/Elements/Framework/TypedElementQueryBuilder.h b/Engine/Source/Runtime/TypedElementFramework/Public/Elements/Framework/TypedElementQueryBuilder.h
--- a/Engine/Source/Runtime/TypedElementFramework/Public/Elements/Framework/TypedElementQueryBuilder.h
+++ b/Engine/Source/Runtime/TypedElementFramework/Public/Elements/Framework/TypedElementQueryBuilder.h
@@ -315,6 +315,12 @@ namespace TypedElementQueryBuilder
 	private:
 		ITypedElementDataStorageInterface::FQueryDescription Query;
 	};
+
+	template<typename Function>
+	ITypedElementDataStorageInterface::DirectQueryCallback CreateDirectQueryCallbackBinding(Function&& Callback);
+	template<typename Function>
+	ITypedElementDataStorageInterface::DirectQueryCallback CreateSubqueryCallbackBinding(Function&& Callback);
+
 } // namespace TypedElementQueryBuilder
 
 #include "Elements/Framework/TypedElementQueryBuilder.inl"
diff --git a/Engine/Source/Runtime/TypedElementFramework/Public/Elements/Framework/TypedElementQueryBuilder.inl b/Engine/Source/Runtime/TypedElementFramework/Public/Elements/Framework/TypedElementQueryBuilder.inl
--- a/Engine/Source/Runtime/TypedElementFramework/Public/Elements/Framework/TypedElementQueryBuilder.inl
+++ b/Engine/Source/Runtime/TypedElementFramework/Public/Elements/Framework/TypedElementQueryBuilder.inl
@@ -6,6 +6,31 @@
 #include "Templates/IsConst.h"
 #include "Templates/UnrealTypeTraits.h"
 #include "Elements/Interfaces/TypedElementDataStorageCompatibilityInterface.h"
+#include "Elements/Interfaces/TypedElementQueryStorageInterfaces.h"
+#include "Elements/Common/TypedElementQueryDescription.h"
+#include "Elements/Common/TypedElementQueryTypes.h"
+
+// Work around missing header on some platforms
+#if __has_include(<concepts>)
+#	include <concepts>
+#else
+namespace std
+{
+	namespace detail
+	{
+		template< class T, class U >
+		concept SameHelper = std::is_same_v<T, U>;
+	}
+	template< class T, class U >
+	concept same_as = detail::SameHelper<T, U> && detail::SameHelper<U, T>;
+
+	template <class From, class To>
+	concept convertible_to = std::is_convertible_v<From, To> && requires { static_cast<To>(std::declval<From>()); };
+
+	template< class Derived, class Base >
+	concept derived_from = std::is_base_of_v<Base, Derived> && std::is_convertible_v<const volatile Derived*, const volatile Base*>;
+}
+#endif
 
 namespace TypedElementQueryBuilder
 {
@@ -27,37 +52,37 @@ namespace TypedElementQueryBuilder
 		}
 
 		template<typename Type>
-		constexpr ITypedElementDataStorageInterface::EQueryAccessType GetAccessType()
+		constexpr TypedElementDataStorage::EQueryAccessType GetAccessType()
 		{
 			using BaseType = typename std::remove_reference_t<Type>;
 			if constexpr (TIsConst<BaseType>::Value)
 			{
-				return ITypedElementDataStorageInterface::EQueryAccessType::ReadOnly;
+				return TypedElementDataStorage::EQueryAccessType::ReadOnly;
 			}
 			else
 			{
-				return ITypedElementDataStorageInterface::EQueryAccessType::ReadWrite;
+				return TypedElementDataStorage::EQueryAccessType::ReadWrite;
 			}
 		}
 
 		template<typename Type>
-		constexpr ITypedElementDataStorageInterface::EQueryDependencyFlags GetDependencyFlags()
+		constexpr TypedElementDataStorage::EQueryDependencyFlags GetDependencyFlags()
 		{
 			using BaseType = typename std::remove_reference_t<Type>;
 			using SubsystemTraits = TTypedElementSubsystemTraits<std::remove_const_t<BaseType>>;
 
-			ITypedElementDataStorageInterface::EQueryDependencyFlags Result = ITypedElementDataStorageInterface::EQueryDependencyFlags::None;
+			TypedElementDataStorage::EQueryDependencyFlags Result = TypedElementDataStorage::EQueryDependencyFlags::None;
 			if constexpr (SubsystemTraits::RequiresGameThread())
 			{
-				EnumAddFlags(Result, ITypedElementDataStorageInterface::EQueryDependencyFlags::GameThreadBound);
+				EnumAddFlags(Result, TypedElementDataStorage::EQueryDependencyFlags::GameThreadBound);
 			}
 			if constexpr (SubsystemTraits::IsHotReloadable())
 			{
-				EnumAddFlags(Result, ITypedElementDataStorageInterface::EQueryDependencyFlags::AlwaysRefresh);
+				EnumAddFlags(Result, TypedElementDataStorage::EQueryDependencyFlags::AlwaysRefresh);
 			}
 			if constexpr (TIsConst<BaseType>::Value)
 			{
-				EnumAddFlags(Result, ITypedElementDataStorageInterface::EQueryDependencyFlags::ReadOnly);
+				EnumAddFlags(Result, TypedElementDataStorage::EQueryDependencyFlags::ReadOnly);
 			}
 			
 			return Result;
@@ -104,8 +129,8 @@ namespace TypedElementQueryBuilder
 	//
 
 	FQueryContextForwarder::FQueryContextForwarder(
-		const ITypedElementDataStorageInterface::FQueryDescription& InDescription, 
-		ITypedElementDataStorageInterface::IQueryContext& InParentContext)
+		const TypedElementDataStorage::FQueryDescription& InDescription,
+		TypedElementDataStorage::IQueryContext& InParentContext)
 		: ParentContext(InParentContext)
 		, Description(InDescription)
 	{}
@@ -121,13 +146,13 @@ namespace TypedElementQueryBuilder
 	}
 
 	void FQueryContextForwarder::GetColumns(TArrayView<char*> RetrievedAddresses, TConstArrayView<TWeakObjectPtr<const UScriptStruct>> ColumnTypes,
-		TConstArrayView<ITypedElementDataStorageInterface::EQueryAccessType> AccessTypes)
+		TConstArrayView<TypedElementDataStorage::EQueryAccessType> AccessTypes)
 	{
 		ParentContext.GetColumns(RetrievedAddresses, ColumnTypes, AccessTypes);
 	}
 
 	void FQueryContextForwarder::GetColumnsUnguarded(int32 TypeCount, char** RetrievedAddresses, const TWeakObjectPtr<const UScriptStruct>* ColumnTypes,
-		const ITypedElementDataStorageInterface::EQueryAccessType* AccessTypes)
+		const TypedElementDataStorage::EQueryAccessType* AccessTypes)
 	{
 		ParentContext.GetColumnsUnguarded(TypeCount, RetrievedAddresses, ColumnTypes, AccessTypes);
 	}
@@ -143,7 +168,7 @@ namespace TypedElementQueryBuilder
 	}
 
 	void FQueryContextForwarder::GetDependencies(TArrayView<UObject*> RetrievedAddresses, TConstArrayView<TWeakObjectPtr<const UClass>> DependencyTypes,
-		TConstArrayView<ITypedElementDataStorageInterface::EQueryAccessType> AccessTypes)
+		TConstArrayView<TypedElementDataStorage::EQueryAccessType> AccessTypes)
 	{
 		ParentContext.GetDependencies(RetrievedAddresses, DependencyTypes, AccessTypes);
 	}
@@ -188,12 +213,12 @@ namespace TypedElementQueryBuilder
 		ParentContext.RemoveColumns(Rows, ColumnTypes);
 	}
 
-	ITypedElementDataStorageInterface::FQueryResult FQueryContextForwarder::RunQuery(TypedElementQueryHandle Query)
+	TypedElementDataStorage::FQueryResult FQueryContextForwarder::RunQuery(TypedElementQueryHandle Query)
 	{
 		return ParentContext.RunQuery(Query);
 	}
 	
-	ITypedElementDataStorageInterface::FQueryResult FQueryContextForwarder::RunSubquery(int32 SubqueryIndex)
+	TypedElementDataStorage::FQueryResult FQueryContextForwarder::RunSubquery(int32 SubqueryIndex)
 	{
 		return ParentContext.RunSubquery(SubqueryIndex);
 	}
@@ -204,13 +229,13 @@ namespace TypedElementQueryBuilder
 	
 	template<typename... Dependencies>
 	FCachedQueryContext<Dependencies...>::FCachedQueryContext(
-		const ITypedElementDataStorageInterface::FQueryDescription& InDescription, 
-		ITypedElementDataStorageInterface::IQueryContext& InParentContext)
+		const TypedElementDataStorage::FQueryDescription& InDescription,
+		TypedElementDataStorage::IQueryContext& InParentContext)
 		: FQueryContextForwarder(InDescription, InParentContext)
 	{}
 
 	template<typename... Dependencies>
-	void FCachedQueryContext<Dependencies...>::Register(ITypedElementDataStorageInterface::FQueryDescription& Query)
+	void FCachedQueryContext<Dependencies...>::Register(TypedElementDataStorage::FQueryDescription& Query)
 	{
 		Query.DependencyTypes.Reserve(sizeof...(Dependencies));
 		Query.DependencyFlags.Reserve(sizeof...(Dependencies));
@@ -254,38 +279,25 @@ namespace TypedElementQueryBuilder
 	//
 	namespace Internal
 	{
-		using QueryContext = ITypedElementDataStorageInterface::IQueryContext;
-		using QueryAccess = ITypedElementDataStorageInterface::EQueryAccessType;
-		using QueryDescription = ITypedElementDataStorageInterface::FQueryDescription;
-
-		template <typename T>
-		struct HasStaticStructMethod
+		template<typename T>
+		concept HasStaticStructMethod = requires
 		{
-		private:
-			template <typename C> static constexpr bool Check(decltype(&C::StaticStruct)) { return std::is_same_v<decltype(&C::StaticStruct), UScriptStruct* (*)()>; }
-			template <typename C> static constexpr bool Check(...) { return false; }
-
-		public:
-			static constexpr bool Value = Check<T>(nullptr);
+			{ T::StaticStruct() } -> std::same_as<UScriptStruct*>;
 		};
 
 		template<typename Context>
-		constexpr bool IsValidContextType()
+		concept SourceQueryContext = requires(Context& Ctx, int32 TypeCount, char** RetrievedAddresses,
+			const TWeakObjectPtr<const UScriptStruct>*ColumnTypes, const TypedElementDataStorage::EQueryAccessType * AccessTypes)
 		{
-			using BaseContextType = typename std::remove_const_t<typename TRemoveReference<Context>::Type>;
-			if constexpr (std::is_same_v<BaseContextType, QueryContext>)
-			{
-				return true;
-			}
-			else if constexpr (TIsDerivedFrom<BaseContextType, FQueryContextForwarder>::Value)
-			{
-				return true;
-			}
-			else
-			{
-				return false;
-			}
-		}
+			{ Ctx.GetColumnsUnguarded(TypeCount, RetrievedAddresses, ColumnTypes, AccessTypes) };
+			{ Ctx.GetRowHandles() } -> std::convertible_to<TConstArrayView<TypedElementRowHandle>>;
+			{ Ctx.GetRowCount() } -> std::convertible_to<uint32>;
+		};
+
+		template<typename Source, typename Target>
+		concept IsCompatibleTargetContextType = 
+			SourceQueryContext<Source> && 
+			std::derived_from<std::remove_const_t<std::remove_reference_t<Target>>, Source>;
 
 		template<typename... Columns>
 		constexpr bool AreAllColumnsReferences()
@@ -326,12 +338,12 @@ namespace TypedElementQueryBuilder
 			if constexpr (std::is_reference_v<Column>)
 			{
 				using BaseType = typename std::remove_const_t<std::remove_reference_t<Column>>;
-				return HasStaticStructMethod<BaseType>::Value;
+				return HasStaticStructMethod<BaseType>;
 			}
 			else if constexpr (std::is_pointer_v<Column>)
 			{
 				using BaseType = typename std::remove_const_t<std::remove_pointer_t<Column>>;
-				return HasStaticStructMethod<BaseType>::Value;
+				return HasStaticStructMethod<BaseType>;
 			}
 			else
 			{
@@ -346,10 +358,7 @@ namespace TypedElementQueryBuilder
 		}
 
 		template<typename RowType>
-		constexpr bool IsRowHandleType()
-		{
-			return std::is_same_v<std::remove_cv_t<std::remove_pointer_t<RowType>>, TypedElementRowHandle>;
-		}
+		concept IsRowHandleType = std::is_same_v<std::remove_cv_t<std::remove_pointer_t<RowType>>, TypedElementRowHandle>;
 
 		template<typename RowType, typename... Columns>
 		constexpr bool IsRowTypeCompatibleWithColumns()
@@ -364,24 +373,63 @@ namespace TypedElementQueryBuilder
 			}
 		}
 
-		template <typename T>
-		class IsFunctor
-		{
-			template <typename C> static constexpr bool Check(decltype(&C::operator())) { return true; }
-			template <typename C> static constexpr bool Check(...) { return false; }
-
-		public:
-			static constexpr bool Value = Check<T>(nullptr);
-		};
+		template<typename T> concept IsFunctor = requires{ &T::operator(); };
 
-		template<typename... ColumnTypes>
+		template<SourceQueryContext SourceContext, bool ValidateColumns, typename... ColumnTypes>
 		struct FFunctionColumnInfo
 		{
 			std::tuple<BaseColumnType<ColumnTypes>*...> Columns;
 			static constexpr bool bArePointerColumns = AreAllColumnsPointers<ColumnTypes...>();
 
-			FFunctionColumnInfo(const QueryDescription& Description, QueryContext& Context)
+			template<typename Column>
+			bool CheckValidity(const TypedElementDataStorage::FQueryDescription& Description)
 			{
+				const UScriptStruct* ColumnType = Column::StaticStruct();
+				if (const TWeakObjectPtr<const UScriptStruct>* Match = Description.SelectionTypes.FindByPredicate(
+					[ColumnType](const TWeakObjectPtr<const UScriptStruct>& Entry)
+					{
+						return Entry == ColumnType;
+					}); Match != nullptr)
+				{
+					size_t Index = Match - Description.SelectionTypes.GetData();
+					if (Description.SelectionAccessTypes[Index] == TypedElementDataStorage::EQueryAccessType::ReadOnly)
+					{
+						if (std::is_const_v<Column>)
+						{
+							return true;
+						}
+						else
+						{
+							checkf(false, 
+								TEXT("Provided query callback couldn't be bound because target read-only column of type '%s' is bound to a mutable reference."),
+								*ColumnType->GetName());
+							return false;
+						}
+					}
+					else
+					{
+						return true;
+					}
+				}
+				else
+				{
+					checkf(false, 
+						TEXT("Provided query callback couldn't be bound because target query doesn't contain a column of type '%s'."), 
+						*ColumnType->GetName());
+					return false;
+				}
+			}
+
+			FFunctionColumnInfo(const TypedElementDataStorage::FQueryDescription& Description, SourceContext& Context)
+			{
+				if constexpr (ValidateColumns)
+				{
+					if (!(CheckValidity<BaseColumnType<ColumnTypes>>(Description) && ...))
+					{
+						return;
+					}
+				}
+
 				char* ColumnAddresses[sizeof...(ColumnTypes)];
 				Context.GetColumnsUnguarded(sizeof...(ColumnTypes), ColumnAddresses, 
 					Description.SelectionTypes.GetData(), Description.SelectionAccessTypes.GetData());
@@ -393,57 +441,61 @@ namespace TypedElementQueryBuilder
 			}
 		};
 
-		template<>
-		struct FFunctionColumnInfo<>
+		template<SourceQueryContext SourceContext, bool ValidateColumns>
+		struct FFunctionColumnInfo<SourceContext, ValidateColumns>
 		{
-			FFunctionColumnInfo(const QueryDescription&, QueryContext&) {}
+			FFunctionColumnInfo(const TypedElementDataStorage::FQueryDescription&, SourceContext&) {}
 			static constexpr bool bArePointerColumns = false;
 		};
 
-		template<typename ContextType>
+		template<SourceQueryContext SourceType, typename TargetType> 
+			requires IsCompatibleTargetContextType<SourceType, TargetType>
 		struct FContextInfo
 		{
-			using BaseContextType = std::remove_reference_t<ContextType>;
+			using BaseContextType = std::remove_reference_t<TargetType>;
 			BaseContextType ContextWrapper;
 
-			FContextInfo(const QueryDescription& Description, QueryContext& Context)
+			FContextInfo(const TypedElementDataStorage::FQueryDescription& Description, SourceType& Context)
 				: ContextWrapper(Description, Context)
 			{}
 		};
 
-		template<>
-		struct FContextInfo<const QueryContext&>
+		template<SourceQueryContext SourceType>
+		struct FContextInfo<SourceType, const SourceType&>
 		{
-			const QueryContext& ContextWrapper;
+			const SourceType& ContextWrapper;
 
-			FContextInfo(const QueryDescription& Description, QueryContext& Context)
+			FContextInfo(const TypedElementDataStorage::FQueryDescription& Description, SourceType& Context)
 				: ContextWrapper(Context)
 			{}
 		};
 
-		template<>
-		struct FContextInfo<QueryContext&>
+		template<SourceQueryContext SourceType>
+		struct FContextInfo<SourceType, SourceType&>
 		{
-			QueryContext& ContextWrapper;
+			SourceType& ContextWrapper;
 
-			FContextInfo(const QueryDescription& Description, QueryContext& Context)
+			FContextInfo(const TypedElementDataStorage::FQueryDescription& Description, SourceType& Context)
 				: ContextWrapper(Context)
 			{}
 		};
 
-		template<typename ContextType, typename RowHandleType, typename... Args>
-		struct FContextRowHandleColumnsFunction : public FFunctionColumnInfo<Args...>, public FContextInfo<ContextType>
+		template<SourceQueryContext SourceContext, bool ValidateColumns, typename TargetContext, typename RowHandleType, typename... Args>
+			requires IsCompatibleTargetContextType<SourceContext, TargetContext>
+		struct FContextRowHandleColumnsFunction 
+			: public FFunctionColumnInfo<SourceContext, ValidateColumns, Args...>
+			, public FContextInfo<SourceContext, TargetContext>
 		{
-			using SuperColumn = FFunctionColumnInfo<Args...>;
-			using SuperContext = FContextInfo<ContextType>;
+			using SuperColumn = FFunctionColumnInfo<SourceContext, ValidateColumns, Args...>;
+			using SuperContext = FContextInfo<SourceContext, TargetContext>;
 
-			FContextRowHandleColumnsFunction(const QueryDescription& Description, QueryContext& Context)
+			FContextRowHandleColumnsFunction(const TypedElementDataStorage::FQueryDescription& Description, SourceContext& Context)
 				: SuperColumn(Description, Context)
 				, SuperContext(Description, Context)
 			{}
 
 			template<typename CallerType>
-			void Call(QueryContext& Context, CallerType&& Caller)
+			void Call(SourceContext& Context, const CallerType& Caller)
 			{
 				TConstArrayView<TypedElementRowHandle> Rows = Context.GetRowHandles();
 				if constexpr (SuperColumn::bArePointerColumns || std::is_pointer_v<RowHandleType>)
@@ -484,33 +536,36 @@ namespace TypedElementQueryBuilder
 			}
 		};
 
-		template<typename ContextType, typename... Args>
-		struct FContextColumnsFunction : public FFunctionColumnInfo<Args...>, public FContextInfo<ContextType>
+		template<SourceQueryContext SourceContext, bool ValidateColumns, typename TargetContext, typename... Args>
+			requires IsCompatibleTargetContextType<SourceContext, TargetContext>
+		struct FContextColumnsFunction 
+			: public FFunctionColumnInfo<SourceContext, ValidateColumns, Args...>
+			, public FContextInfo<SourceContext, TargetContext>
 		{
-			using SuperColumn = FFunctionColumnInfo<Args...>;
-			using SuperContext = FContextInfo<ContextType>;
+			using SuperColumn = FFunctionColumnInfo<SourceContext, ValidateColumns, Args...>;
+			using SuperContext = FContextInfo<SourceContext, TargetContext>;
 
-			FContextColumnsFunction(const QueryDescription& Description, QueryContext& Context)
+			FContextColumnsFunction(const TypedElementDataStorage::FQueryDescription& Description, SourceContext& Context)
 				: SuperColumn(Description, Context)
 				, SuperContext(Description, Context)
 			{}
 
 			template<typename CallerType>
-			void Call(QueryContext& Context, CallerType&& Caller)
+			void Call(SourceContext& Context, const CallerType& Caller)
 			{
 				if constexpr (sizeof...(Args) > 0)
 				{
 					if constexpr (SuperColumn::bArePointerColumns)
 					{
 						std::apply(
-							[this, Caller = std::forward<CallerType>(Caller)](auto&&... Column)
+							[this, &Caller](auto&&... Column)
 							{
 								Caller(this->ContextWrapper, Column...);
 							}, this->Columns);
 					}
 					else
 					{
-						auto CallForwarder = [this, Caller = std::forward<CallerType>(Caller)](auto&&... Column)
+						auto CallForwarder = [this, &Caller](auto&&... Column)
 						{
 							Caller(this->ContextWrapper, *Column...);
 							((++Column), ...);
@@ -533,17 +588,17 @@ namespace TypedElementQueryBuilder
 			}
 		};
 
-		template<typename RowHandleType, typename... Columns>
-		struct FRowHandleColumnsFunction : public FFunctionColumnInfo<Columns...>
+		template<SourceQueryContext SourceContext, bool ValidateColumns, typename RowHandleType, typename... Columns>
+		struct FRowHandleColumnsFunction : public FFunctionColumnInfo<SourceContext, ValidateColumns, Columns...>
 		{
-			using Super = FFunctionColumnInfo<Columns...>;
+			using Super = FFunctionColumnInfo<SourceContext, ValidateColumns, Columns...>;
 
-			FRowHandleColumnsFunction(const QueryDescription& Description, QueryContext& Context)
+			FRowHandleColumnsFunction(const TypedElementDataStorage::FQueryDescription& Description, SourceContext& Context)
 				: Super(Description, Context)
 			{}
 
 			template<typename CallerType>
-			void Call(QueryContext& Context, CallerType&& Caller)
+			void Call(SourceContext& Context, const CallerType& Caller)
 			{
 				TConstArrayView<TypedElementRowHandle> Rows = Context.GetRowHandles();
 				if constexpr (Super::bArePointerColumns || std::is_pointer_v<RowHandleType>)
@@ -581,31 +636,31 @@ namespace TypedElementQueryBuilder
 			}
 		};
 
-		template<typename... Args>
-		struct FColumnsFunction : public FFunctionColumnInfo<Args...>
+		template<SourceQueryContext SourceContext, bool ValidateColumns, typename... Args>
+		struct FColumnsFunction : public FFunctionColumnInfo<SourceContext, ValidateColumns, Args...>
 		{
-			using Super = FFunctionColumnInfo<Args...>;
+			using Super = FFunctionColumnInfo<SourceContext, ValidateColumns, Args...>;
 			
-			FColumnsFunction(const QueryDescription& Description, QueryContext& Context)
+			FColumnsFunction(const TypedElementDataStorage::FQueryDescription& Description, SourceContext& Context)
 				: Super(Description, Context)
 			{}
 
 			template<typename CallerType>
-			void Call(QueryContext& Context, CallerType&& Caller)
+			void Call(SourceContext& Context, const CallerType& Caller)
 			{
 				if constexpr (sizeof...(Args) > 0)
 				{
 					if constexpr (Super::bArePointerColumns)
 					{
 						std::apply(
-							[Caller = std::forward<CallerType>(Caller)](auto&&... Column)
+							[&Caller](auto&&... Column)
 							{
 								Caller(Column...);
 							}, this->Columns);
 					}
 					else
 					{
-						auto CallForwarder = [Caller = std::forward<CallerType>(Caller)](auto&&... Column)
+						auto CallForwarder = [&Caller](auto&&... Column)
 						{
 							Caller(*Column...);
 							((++Column), ...);
@@ -628,114 +683,143 @@ namespace TypedElementQueryBuilder
 			}
 		};
 
-		template<typename... Args>
+		template<SourceQueryContext SourceContext, bool ValidateColumns, typename... Args>
 		struct FFunctionInfoHelper
 		{
-			using BaseClass = FColumnsFunction<Args...>;
+			using BaseClass = FColumnsFunction<SourceContext, ValidateColumns, Args...>;
 		};
 
-		template<typename Arg0>
-		struct FFunctionInfoHelper<Arg0>
+		template<SourceQueryContext SourceContext, bool ValidateColumns, typename Arg0, typename... Args>
+			requires IsCompatibleTargetContextType<SourceContext, Arg0>
+		struct FFunctionInfoHelper<SourceContext, ValidateColumns, Arg0, Args...>
 		{
-			static constexpr bool bArg0IsContext = IsValidContextType<Arg0>();
-			static constexpr bool bArg0IsRowHandle = IsRowHandleType<Arg0>();
-			using BaseClass = 
-				std::conditional_t<bArg0IsContext,
-					FContextColumnsFunction<Arg0>,
-					std::conditional_t<bArg0IsRowHandle,
-						FRowHandleColumnsFunction<Arg0>,
-						FColumnsFunction<Arg0>
-					>
-				>;
+			using BaseClass = FContextColumnsFunction<SourceContext, ValidateColumns, Arg0, Args...>;
 		};
 
-		template<typename Arg0, typename Arg1, typename... Args>
-		struct FFunctionInfoHelper<Arg0, Arg1, Args...>
+		template<SourceQueryContext SourceContext, bool ValidateColumns, typename Arg0, typename... Args>
+			requires IsRowHandleType<Arg0>
+		struct FFunctionInfoHelper<SourceContext, ValidateColumns, Arg0, Args...>
 		{
-			static constexpr bool bArg0IsContext = IsValidContextType<Arg0>();
-			static constexpr bool bArg0IsRowHandle = IsRowHandleType<Arg0>();
-			
-			static constexpr bool bArg1IsRowHandle = IsRowHandleType<Arg1>();
-
-			using BaseClass =
-				std::conditional_t<bArg0IsContext,
-					std::conditional_t<bArg1IsRowHandle,
-						FContextRowHandleColumnsFunction<Arg0, Arg1, Args...>,
-						FContextColumnsFunction<Arg0, Arg1, Args...>
-					>,
-					std::conditional_t<bArg0IsRowHandle,
-						FRowHandleColumnsFunction<Arg0, Arg1, Args...>,
-						FColumnsFunction<Arg0, Arg1, Args...>
-					>
-				>;
+			using BaseClass = FRowHandleColumnsFunction<SourceContext, ValidateColumns, Arg0, Args...>;
+		};
+
+		template<SourceQueryContext SourceContext, bool ValidateColumns, typename Arg0, typename Arg1, typename... Args>
+			requires IsCompatibleTargetContextType<SourceContext, Arg0> && IsRowHandleType<Arg1>
+		struct FFunctionInfoHelper<SourceContext, ValidateColumns, Arg0, Arg1, Args...>
+		{
+			using BaseClass = FContextRowHandleColumnsFunction<SourceContext, ValidateColumns, Arg0, Arg1, Args...>;
 		};
 
-		template<typename... Args>
-		struct FFunctionInfo final : public FFunctionInfoHelper<Args...>::BaseClass
+		template<SourceQueryContext SourceContext, bool ValidateColumns, typename... Args>
+		struct FFunctionInfo final : public FFunctionInfoHelper<SourceContext, ValidateColumns, Args...>::BaseClass
 		{
-			using Super = typename FFunctionInfoHelper<Args...>::BaseClass;
+			using Super = typename FFunctionInfoHelper<SourceContext, ValidateColumns, Args...>::BaseClass;
 
-			FFunctionInfo(const QueryDescription& Description, QueryContext& Context)
+			FFunctionInfo(const TypedElementDataStorage::FQueryDescription& Description, SourceContext& Context)
 				: Super(Description, Context)
 			{}
 		};
 
-		template<typename... Args>
-		void BindQueryFunction(ITypedElementDataStorageInterface::QueryCallback& Function, 
+		template<SourceQueryContext SourceContext, bool ValidateColumns, typename... Args>
+		void BindQueryFunction(TypedElementDataStorage::QueryCallback& Function,
 			void (*Callback)(Args...))
 		{
-			Function = [Callback](const QueryDescription& Description, QueryContext& Context)
+			Function = [Callback](const TypedElementDataStorage::FQueryDescription& Description, SourceContext& Context)
 			{
-				FFunctionInfo<Args...>(Description, Context).Call(Context, Callback);
+				FFunctionInfo<SourceContext, ValidateColumns, Args...>(Description, Context).Call(Context, Callback);
 			};
 		}
 
-		template<typename Class, typename... Args>
-		void BindQueryFunction(ITypedElementDataStorageInterface::QueryCallback& Function, Class* Target, 
+		template<SourceQueryContext SourceContext, bool ValidateColumns, typename Class, typename... Args>
+		void BindQueryFunction(TypedElementDataStorage::QueryCallback& Function, Class* Target,
 			void (Class::*Callback)(Args...))
 		{
-			Function = [Target, Callback](const QueryDescription& Description, QueryContext& Context)
+			Function = [Target, Callback](const TypedElementDataStorage::FQueryDescription& Description, SourceContext& Context)
 			{
-				FFunctionInfo<Args...>(Description, Context).Call(Context, (Target->*Callback));
+				FFunctionInfo<SourceContext, ValidateColumns, Args...>(Description, Context).Call(Context, (Target->*Callback));
 			};
 		}
 
-		template<typename Class, typename... Args>
-		void BindQueryFunction(ITypedElementDataStorageInterface::QueryCallback& Function, Class* Target, 
+		template<SourceQueryContext SourceContext, bool ValidateColumns, typename Class, typename... Args>
+		void BindQueryFunction(TypedElementDataStorage::QueryCallback& Function, Class* Target,
 			void (Class::*Callback)(Args...) const)
 		{
-			Function = [Target, Callback](const QueryDescription& Description, QueryContext& Context)
+			Function = [Target, Callback](const TypedElementDataStorage::FQueryDescription& Description, SourceContext& Context)
 			{
-				FFunctionInfo<Args...>(Description, Context).Call(Context, (Target->*Callback));
+				FFunctionInfo<SourceContext, ValidateColumns, Args...>(Description, Context).Call(Context, (Target->*Callback));
 			};
 		}
 
-		template<typename Functor, typename Class, typename... Args>
-		void BindQueryFunction_Expand(ITypedElementDataStorageInterface::QueryCallback& Function, Functor&& CallbackObject, 
+		template<SourceQueryContext SourceContext, bool ValidateColumns, typename Functor, typename Class, typename... Args>
+		void BindQueryFunction_Expand(TypedElementDataStorage::QueryCallback& Function, Functor&& CallbackObject,
 			void (Class::*Callback)(Args...) const)
 		{
-			Function = [CallbackObject = std::forward<Functor>(CallbackObject)](const QueryDescription& Description, QueryContext& Context)
+			Function = [CallbackObject = std::forward<Functor>(CallbackObject)](const TypedElementDataStorage::FQueryDescription& Description, SourceContext& Context)
 			{
-				FFunctionInfo<Args...>(Description, Context).Call(Context, CallbackObject);
+				FFunctionInfo<SourceContext, ValidateColumns, Args...>(Description, Context).Call(Context, CallbackObject);
 			};
 		}
 
-		template<typename Functor, typename Class, typename... Args>
-		void BindQueryFunction_Expand(ITypedElementDataStorageInterface::QueryCallback& Function, Functor&& CallbackObject, 
+		template<SourceQueryContext SourceContext, bool ValidateColumns, typename Functor, typename Class, typename... Args>
+		void BindQueryFunction_Expand(TypedElementDataStorage::QueryCallback& Function, Functor&& CallbackObject,
 			void (Class::*Callback)(Args...))
 		{
-			Function = [CallbackObject = std::forward<Functor>(CallbackObject)](const QueryDescription& Description, QueryContext& Context)
+			Function = [CallbackObject = std::forward<Functor>(CallbackObject)](const TypedElementDataStorage::FQueryDescription& Description, SourceContext& Context)
 			{
-				FFunctionInfo<Args...>(Description, Context).Call(Context, CallbackObject);
+				FFunctionInfo<SourceContext, ValidateColumns, Args...>(Description, Context).Call(Context, CallbackObject);
 			};
 		}
 
-		template<typename Functor>
-		void BindQueryFunction(ITypedElementDataStorageInterface::QueryCallback& Function, Functor&& Callback)
+		template<SourceQueryContext SourceContext, bool ValidateColumns, typename Functor>
+		void BindQueryFunction(TypedElementDataStorage::QueryCallback& Function, Functor&& Callback)
+		{
+			BindQueryFunction_Expand<SourceContext, ValidateColumns>(Function, std::forward<Functor>(Callback), &Functor::operator());
+		}
+
+		
+		template<bool ValidateColumns, SourceQueryContext SourceContext, typename... Args>
+		void CallQueryFunction(const TypedElementDataStorage::FQueryDescription& Description, SourceContext& Context,
+			TypedElementDataStorage::QueryCallback& Function, void (*Callback)(Args...))
+		{
+			FFunctionInfo<SourceContext, ValidateColumns, Args...>(Description, Context).Call(Context, Callback);
+		}
+
+		template<bool ValidateColumns, SourceQueryContext SourceContext, typename Class, typename... Args>
+		void CallQueryFunction(const TypedElementDataStorage::FQueryDescription& Description, SourceContext& Context,
+			Class* Target, void (Class::* Callback)(Args...))
+		{
+			FFunctionInfo<SourceContext, ValidateColumns, Args...>(Description, Context).Call(Context, (Target->*Callback));
+		}
+
+		template<bool ValidateColumns, SourceQueryContext SourceContext, typename Class, typename... Args>
+		void CallQueryFunction(const TypedElementDataStorage::FQueryDescription& Description, SourceContext& Context,
+			Class* Target, void (Class::* Callback)(Args...) const)
+		{
+			FFunctionInfo<SourceContext, ValidateColumns, Args...>(Description, Context).Call(Context, (Target->*Callback));
+		}
+
+		template<bool ValidateColumns, SourceQueryContext SourceContext, typename Functor, typename Class, typename... Args>
+		void CallQueryFunction_Expand(const TypedElementDataStorage::FQueryDescription& Description, SourceContext& Context,
+			const Functor& CallbackObject, void (Class::* Callback)(Args...) const)
+		{
+			FFunctionInfo<SourceContext, ValidateColumns, Args...>(Description, Context).Call(Context, CallbackObject);
+		}
+
+		template<bool ValidateColumns, SourceQueryContext SourceContext, typename Functor, typename Class, typename... Args>
+		void CallQueryFunction_Expand(const TypedElementDataStorage::FQueryDescription& Description, SourceContext& Context,
+			const Functor& CallbackObject, void (Class::* Callback)(Args...))
 		{
-			BindQueryFunction_Expand(Function, std::forward<Functor>(Callback), &Functor::operator());
+			FFunctionInfo<SourceContext, ValidateColumns, Args...>(Description, Context).Call(Context, CallbackObject);
+		}
+
+		template<bool ValidateColumns, SourceQueryContext SourceContext, typename Functor>
+		void CallQueryFunction(const TypedElementDataStorage::FQueryDescription& Description, SourceContext& Context,
+			const Functor& Callback)
+		{
+			CallQueryFunction_Expand<ValidateColumns>(Description, Context, Callback, &Functor::operator());
 		}
 		
+
 		template<typename Arg>
 		void AddColumnToSelect(Select& Target)
 		{
@@ -749,7 +833,7 @@ namespace TypedElementQueryBuilder
 			}
 		}
 
-		template<typename ContextType> void RegisterDependencies(QueryDescription& Query)
+		template<typename ContextType> void RegisterDependencies(TypedElementDataStorage::FQueryDescription& Query)
 		{
 			using BaseType = typename TRemoveReference<ContextType>::Type;
 			if constexpr (TIsDerivedFrom<BaseType, FQueryContextForwarder>::Value)
@@ -758,27 +842,27 @@ namespace TypedElementQueryBuilder
 			}
 		}
 
-		template<typename... Args>
+		template<SourceQueryContext SourceContext, typename... Args>
 		struct RegisterFunctionArgumentsHelper
 		{
-			void Register(QueryDescription&, Select& Target)
+			void Register(TypedElementDataStorage::FQueryDescription&, Select& Target)
 			{
 				(AddColumnToSelect<Args>(Target), ...);
 			}
 		};
 
-		template<>
-		struct RegisterFunctionArgumentsHelper<>
+		template<SourceQueryContext SourceContext>
+		struct RegisterFunctionArgumentsHelper<SourceContext>
 		{
-			static void Register(QueryDescription&, Select&) {}
+			static void Register(TypedElementDataStorage::FQueryDescription&, Select&) {}
 		};
 
-		template<typename Arg0>
-		struct RegisterFunctionArgumentsHelper<Arg0>
+		template<SourceQueryContext SourceContext, typename Arg0>
+		struct RegisterFunctionArgumentsHelper<SourceContext, Arg0>
 		{
-			static void Register(QueryDescription& Query, Select& Target)
+			static void Register(TypedElementDataStorage::FQueryDescription& Query, Select& Target)
 			{
-				if constexpr (IsValidContextType<Arg0>())
+				if constexpr (IsCompatibleTargetContextType<SourceContext, Arg0>)
 				{
 					RegisterDependencies<Arg0>(Query);
 				}
@@ -790,12 +874,12 @@ namespace TypedElementQueryBuilder
 			}
 		};
 
-		template<typename Arg0, typename Arg1, typename... Args>
-		struct RegisterFunctionArgumentsHelper<Arg0, Arg1, Args...>
+		template<SourceQueryContext SourceContext, typename Arg0, typename Arg1, typename... Args>
+		struct RegisterFunctionArgumentsHelper<SourceContext, Arg0, Arg1, Args...>
 		{
-			static void Register(QueryDescription& Query, Select& Target)
+			static void Register(TypedElementDataStorage::FQueryDescription& Query, Select& Target)
 			{
-				if constexpr (IsValidContextType<Arg0>())
+				if constexpr (IsCompatibleTargetContextType<SourceContext, Arg0>)
 				{
 					RegisterDependencies<Arg0>(Query);
 				}
@@ -815,41 +899,44 @@ namespace TypedElementQueryBuilder
 			}
 		};
 
-		template<typename... Args>
-		void RegisterFunctionArguments(QueryDescription& Query, Select& Target, void (*)(Args...))
+		template<SourceQueryContext SourceContext, typename... Args>
+		void RegisterFunctionArguments(TypedElementDataStorage::FQueryDescription& Query, Select& Target, void (*)(Args...))
 		{
-			RegisterFunctionArgumentsHelper<Args...>::Register(Query, Target);
+			RegisterFunctionArgumentsHelper<SourceContext, Args...>::Register(Query, Target);
 		}
 
-		template<typename Class, typename... Args>
-		void RegisterFunctionArguments(QueryDescription& Query, Select& Target, void (Class::*)(Args...))
+		template<SourceQueryContext SourceContext, typename Class, typename... Args>
+		void RegisterFunctionArguments(TypedElementDataStorage::FQueryDescription& Query, Select& Target, void (Class::*)(Args...))
 		{
-			RegisterFunctionArgumentsHelper<Args...>::Register(Query, Target);
+			RegisterFunctionArgumentsHelper<SourceContext, Args...>::Register(Query, Target);
 		}
 
-		template<typename Class, typename... Args>
-		void RegisterFunctionArguments(QueryDescription& Query, Select& Target, void (Class::*)(Args...) const)
+		template<SourceQueryContext SourceContext, typename Class, typename... Args>
+		void RegisterFunctionArguments(TypedElementDataStorage::FQueryDescription& Query, Select& Target, void (Class::*)(Args...) const)
 		{
-			RegisterFunctionArgumentsHelper<Args...>::Register(Query, Target);
+			RegisterFunctionArgumentsHelper<SourceContext, Args...>::Register(Query, Target);
 		}
 
-		template<typename Functor>
-		void RegisterFunctionArguments(QueryDescription& Query, Select& Target, Functor)
+		template<SourceQueryContext SourceContext, typename Functor>
+		void RegisterFunctionArguments(TypedElementDataStorage::FQueryDescription& Query, Select& Target, Functor)
 		{
-			RegisterFunctionArguments(Query, Target, &Functor::operator());
+			RegisterFunctionArguments<SourceContext>(Query, Target, &Functor::operator());
 		}
 
-		template<typename... Args>
+		template<SourceQueryContext SourceContext, typename... Args>
 		struct IsValidSelectFunctionSignatureImpl2
 		{
-			constexpr static bool Value = (IsValidColumnType<Args>() && ...);
+			constexpr static bool HasValidColumnTypes = (IsValidColumnType<Args>() && ...);
+			static_assert(HasValidColumnTypes, "One or more of the provided columns is not compatible with query callback.");
+
+			constexpr static bool Value = HasValidColumnTypes;
 		};
 
-		template<typename Arg0, typename... Args>
-		struct IsValidSelectFunctionSignatureImpl2<Arg0, Args...>
+		template<SourceQueryContext SourceContext, typename Arg0, typename... Args>
+		struct IsValidSelectFunctionSignatureImpl2<SourceContext, Arg0, Args...>
 		{
-			constexpr static bool Arg0IsContext = IsValidContextType<Arg0>();
-			constexpr static bool Arg0IsRowHandle = IsRowHandleType<Arg0>();
+			constexpr static bool Arg0IsContext = IsCompatibleTargetContextType<SourceContext, Arg0>;
+			constexpr static bool Arg0IsRowHandle = IsRowHandleType<Arg0>;
 			constexpr static bool Arg0IsArgument = IsValidColumnType<Arg0>();
 			constexpr static bool Arg0IsValidRowHandle = Arg0IsRowHandle && IsRowTypeCompatibleWithColumns<Arg0, Args...>();
 			constexpr static bool Arg0IsValid = (Arg0IsContext || Arg0IsValidRowHandle || Arg0IsArgument);
@@ -859,14 +946,17 @@ namespace TypedElementQueryBuilder
 			static_assert(Arg0IsValid, "The first provided argument has to be a reference to a compatible query reference, a row handle"
 				" or a (const) reference/pointer to a column.");
 			
-			constexpr static bool Value = Arg0IsValid && (IsValidColumnType<Args>() && ...);
+			constexpr static bool HasValidColumnTypes = (IsValidColumnType<Args>() && ...);
+			static_assert(HasValidColumnTypes, "One or more of the provided columns is not compatible with query callback.");
+
+			constexpr static bool Value = Arg0IsValid && HasValidColumnTypes;
 		};
 
-		template<typename Arg0, typename Arg1, typename... Args>
-		struct IsValidSelectFunctionSignatureImpl2<Arg0, Arg1, Args...>
+		template<SourceQueryContext SourceContext, typename Arg0, typename Arg1, typename... Args>
+		struct IsValidSelectFunctionSignatureImpl2<SourceContext, Arg0, Arg1, Args...>
 		{
-			constexpr static bool Arg0IsContext = IsValidContextType<Arg0>();
-			constexpr static bool Arg0IsRowHandle = IsRowHandleType<Arg0>();
+			constexpr static bool Arg0IsContext = IsCompatibleTargetContextType<SourceContext, Arg0>;
+			constexpr static bool Arg0IsRowHandle = IsRowHandleType<Arg0>;
 			constexpr static bool Arg0IsArgument = IsValidColumnType<Arg0>();
 			constexpr static bool Arg0IsValidRowHandle = Arg0IsRowHandle && IsRowTypeCompatibleWithColumns<Arg0, Args...>();
 			constexpr static bool Arg0IsValid = (Arg0IsContext || Arg0IsValidRowHandle || Arg0IsArgument);
@@ -876,7 +966,7 @@ namespace TypedElementQueryBuilder
 			static_assert(Arg0IsValid, "The first argument to a query callback has to be a reference to a compatible query reference, a "
 				"row handle or a (const) reference/pointer to a column.");
 			
-			constexpr static bool Arg1IsRowHandle = IsRowHandleType<Arg1>();
+			constexpr static bool Arg1IsRowHandle = IsRowHandleType<Arg1>;
 			constexpr static bool Arg1IsArgument = IsValidColumnType<Arg1>();
 			constexpr static bool Arg1IsValidRowHandle = Arg1IsRowHandle && IsRowTypeCompatibleWithColumns<Arg1, Args...>();
 			constexpr static bool Arg1IsValid = (Arg1IsValidRowHandle || Arg1IsArgument);
@@ -886,50 +976,57 @@ namespace TypedElementQueryBuilder
 			static_assert(Arg1IsValid, "The second argument to a query callback has to be a row handle or a (const) reference/pointer "
 				"to a column.");
 			
-			constexpr static bool Value = Arg0IsValid && Arg1IsValid && (IsValidColumnType<Args>() && ...);
+			constexpr static bool HasValidColumnTypes = (IsValidColumnType<Args>() && ...);
+			static_assert(HasValidColumnTypes, "One or more of the provided columns is not compatible with query callback.");
+
+			constexpr static bool Value = Arg0IsValid && Arg1IsValid && HasValidColumnTypes;
 		};
 
-		template<>
-		struct IsValidSelectFunctionSignatureImpl2<>
+		template<SourceQueryContext SourceContext>
+		struct IsValidSelectFunctionSignatureImpl2<SourceContext>
 		{
 			constexpr static bool Value = false;
 		};
 				
-		template<typename T> struct IsValidSelectFunctionSignatureImpl
+		template<SourceQueryContext SourceContext, typename T> 
+		struct IsValidSelectFunctionSignatureImpl
 		{ 
 			constexpr static bool Value = false; 
 		};
 
-		template<typename... Args> struct IsValidSelectFunctionSignatureImpl<void (*)(Args...)>
+		template<SourceQueryContext SourceContext, typename... Args> 
+		struct IsValidSelectFunctionSignatureImpl<SourceContext, void (*)(Args...)>
 		{ 
-			constexpr static bool Value = IsValidSelectFunctionSignatureImpl2<Args...>::Value;
+			constexpr static bool Value = IsValidSelectFunctionSignatureImpl2<SourceContext, Args...>::Value;
 		};
 		
-		template<typename Class, typename... Args> struct IsValidSelectFunctionSignatureImpl<void (Class::*)(Args...)>
+		template<SourceQueryContext SourceContext, typename Class, typename... Args> 
+		struct IsValidSelectFunctionSignatureImpl<SourceContext, void (Class::*)(Args...)>
 		{ 
-			constexpr static bool Value = IsValidSelectFunctionSignatureImpl2<Args...>::Value;
+			constexpr static bool Value = IsValidSelectFunctionSignatureImpl2<SourceContext, Args...>::Value;
 		};
 
-		template<typename Class, typename... Args> struct IsValidSelectFunctionSignatureImpl<void (Class::*)(Args...) const>
+		template<SourceQueryContext SourceContext, typename Class, typename... Args>
+		struct IsValidSelectFunctionSignatureImpl<SourceContext, void (Class::*)(Args...) const>
 		{ 
-			constexpr static bool Value = IsValidSelectFunctionSignatureImpl2<Args...>::Value;
+			constexpr static bool Value = IsValidSelectFunctionSignatureImpl2<SourceContext, Args...>::Value;
 		};
 
-		template<typename T> constexpr bool IsValidSelectFunctionSignature()
+		template<SourceQueryContext SourceContext, typename T> constexpr bool IsValidSelectFunctionSignature()
 		{ 
-			if constexpr (IsFunctor<T>::Value)
+			if constexpr (IsFunctor<T>)
 			{
-				return IsValidSelectFunctionSignatureImpl<decltype(&T::operator())>::Value;
+				return IsValidSelectFunctionSignatureImpl<SourceContext, decltype(&T::operator())>::Value;
 			}
 			else
 			{
-				return IsValidSelectFunctionSignatureImpl<T>::Value;
+				return IsValidSelectFunctionSignatureImpl<SourceContext, T>::Value;
 			}
 		};
 		
-		inline void PrepareForQueryBinding(QueryDescription& Query, const FProcessor& Processor)
+		inline void PrepareForQueryBinding(TypedElementDataStorage::FQueryDescription& Query, const FProcessor& Processor)
 		{
-			Query.Callback.Type = ITypedElementDataStorageInterface::EQueryCallbackType::Processor;
+			Query.Callback.Type = TypedElementDataStorage::EQueryCallbackType::Processor;
 			Query.Callback.Phase = Processor.Phase;
 			Query.Callback.Group = Processor.Group;
 			if (!Processor.BeforeGroup.IsNone())
@@ -943,30 +1040,30 @@ namespace TypedElementQueryBuilder
 			Query.Callback.bForceToGameThread = Processor.bForceToGameThread;
 		}
 
-		inline void PrepareForQueryBinding(QueryDescription& Query, const FObserver& Observer)
+		inline void PrepareForQueryBinding(TypedElementDataStorage::FQueryDescription& Query, const FObserver& Observer)
 		{
 			switch (Observer.Event)
 			{
 			case FObserver::EEvent::Add:
-				Query.Callback.Type = ITypedElementDataStorageInterface::EQueryCallbackType::ObserveAdd;
+				Query.Callback.Type = TypedElementDataStorage::EQueryCallbackType::ObserveAdd;
 				break;
 			case FObserver::EEvent::Remove:
-				Query.Callback.Type = ITypedElementDataStorageInterface::EQueryCallbackType::ObserveRemove;
+				Query.Callback.Type = TypedElementDataStorage::EQueryCallbackType::ObserveRemove;
 				break;
 			}
 			Query.Callback.MonitoredType = Observer.Monitor;
 			Query.Callback.bForceToGameThread = Observer.bForceToGameThread;
 		}
 
-		inline void PrepareForQueryBinding(QueryDescription& Query, const FPhaseAmble& PhaseAmble)
+		inline void PrepareForQueryBinding(TypedElementDataStorage::FQueryDescription& Query, const FPhaseAmble& PhaseAmble)
 		{
 			switch (PhaseAmble.Location)
 			{
 			case FPhaseAmble::ELocation::Preamble:
-				Query.Callback.Type = ITypedElementDataStorageInterface::EQueryCallbackType::PhasePreparation;
+				Query.Callback.Type = TypedElementDataStorage::EQueryCallbackType::PhasePreparation;
 				break;
 			case FPhaseAmble::ELocation::Postamble:
-				Query.Callback.Type = ITypedElementDataStorageInterface::EQueryCallbackType::PhaseFinalization;
+				Query.Callback.Type = TypedElementDataStorage::EQueryCallbackType::PhaseFinalization;
 				break;
 			}
 			Query.Callback.Phase = PhaseAmble.Phase;
@@ -974,12 +1071,12 @@ namespace TypedElementQueryBuilder
 		}
 		
 		template<typename CallbackType, typename Function>
-		void PrepareForQueryBinding(Select& Target, QueryDescription& Query, FName Name,
+		void PrepareForQueryBinding(Select& Target, TypedElementDataStorage::FQueryDescription& Query, FName Name,
 			const CallbackType& Type, Function Callback)
 		{
 			static_assert(TIsDerivedFrom<CallbackType, FQueryCallbackType>::Value, "The callback type provided isn't one of the available "
 				"classes derived from FQueryCallbackType.");
-			static_assert(IsValidSelectFunctionSignature<Function>(),
+			static_assert(IsValidSelectFunctionSignature<TypedElementDataStorage::IQueryContext, Function>(),
 				R"(The function provided to the Query Builder's Select call wasn't invocable or doesn't contain a supported combination of arguments.
 The following options are supported:
 - void([const]Column&...) 
@@ -988,29 +1085,31 @@ The following options are supported:
 - void(<Context>&, TypedElementRowHandle, [const]Column&...) 
 - void(<Context>&, [const]Column*...) 
 - void(<Context>&, const TypedElementRowHandle*, [const]Column*...) 
-Where <Context> is ITypedElementDataStorageInterface::IQueryContext or FCachedQueryContext<...>
+Where <Context> is TypedElementDataStorage::IQueryContext or FCachedQueryContext<...>
 e.g. void(FCachedQueryContext<Subsystem1, const Subsystem2>& Context, TypedElementRowHandle Row, ColumnType0& ColumnA, const ColumnType1& ColumnB) {...}
 )");
-			RegisterFunctionArguments(Query, Target, Callback);
+			RegisterFunctionArguments<TypedElementDataStorage::IQueryContext>(Query, Target, Callback);
 			PrepareForQueryBinding(Query, Type);
 			Query.Callback.Name = Name;
 		}
-	}
+	} // namespace Internal
 
 	template<typename CallbackType, typename Function>
 	Select::Select(FName Name, const CallbackType& Type, Function&& Callback)
 		: Select()
 	{
+		static constexpr bool ValidateColumns = false;
 		Internal::PrepareForQueryBinding(*this, Query, Name, Type, Callback);
-		Internal::BindQueryFunction(Query.Callback.Function, std::forward<Function>(Callback));
+		Internal::BindQueryFunction<TypedElementDataStorage::IQueryContext, ValidateColumns>(Query.Callback.Function, std::forward<Function>(Callback));
 	}
 
 	template<typename CallbackType, typename Class, typename Function>
 	Select::Select(FName Name, const CallbackType& Type, Class* Instance, Function&& Callback)
 		: Select()
 	{
+		static constexpr bool ValidateColumns = false;
 		Internal::PrepareForQueryBinding(*this, Query, Name, Type, Callback);
-		Internal::BindQueryFunction(Query.Callback.Function, Instance, Callback);
+		Internal::BindQueryFunction<TypedElementDataStorage::IQueryContext, ValidateColumns>(Query.Callback.Function, Instance, Callback);
 	}
 
 	template<typename... TargetTypes>
@@ -1052,4 +1151,59 @@ e.g. void(FCachedQueryContext<Subsystem1, const Subsystem2>& Context, TypedEleme
 		None({ TargetTypes::StaticStruct()... });
 		return *this;
 	}
-} // namespace TypedElementQueryBuilder
\ No newline at end of file
+
+
+
+	//
+	// External query bindings
+	//
+	template<typename Function>
+	TypedElementDataStorage::DirectQueryCallback CreateDirectQueryCallbackBinding(Function&& Callback)
+	{
+		static_assert(Internal::IsValidSelectFunctionSignature<TypedElementDataStorage::IDirectQueryContext, Function>(),
+			R"(The function provided to the Query Builder's CreateDirectQueryCallbackBinding call wasn't invocable or doesn't contain a supported combination of arguments.
+The following options are supported:
+- void([const]Column&...) 
+- void(TypedElementRowHandle, [const]Column&...) 
+- void(<Context>&, [const]Column&...) 
+- void(<Context>&, TypedElementRowHandle, [const]Column&...) 
+- void(<Context>&, [const]Column*...) 
+- void(<Context>&, const TypedElementRowHandle*, [const]Column*...) 
+Where <Context> is TypedElementDataStorage::IDirectQueryContext
+e.g. void(IDirectQueryContext& Context, TypedElementRowHandle Row, ColumnType0& ColumnA, const ColumnType1& ColumnB) {...}
+)");
+
+		return [Callback = Forward<Function>(Callback)](
+			const TypedElementDataStorage::FQueryDescription& Description,
+			TypedElementDataStorage::IDirectQueryContext& Context)
+		{
+			static constexpr bool ValidateColumns = true;
+			Internal::CallQueryFunction<ValidateColumns>(Description, Context, Callback);
+		};
+	}
+
+	template<typename Function>
+	TypedElementDataStorage::DirectQueryCallback CreateSubqueryCallbackBinding(Function&& Callback)
+	{
+		static_assert(Internal::IsValidSelectFunctionSignature<ITypedElementDataStorageInterface::ISubqueryContext, Function>(),
+			R"(The function provided to the Query Builder's CreateSubqueryCallbackBinding call wasn't invocable or doesn't contain a supported combination of arguments.
+The following options are supported:
+- void([const]Column&...) 
+- void(TypedElementRowHandle, [const]Column&...) 
+- void(<Context>&, [const]Column&...) 
+- void(<Context>&, TypedElementRowHandle, [const]Column&...) 
+- void(<Context>&, [const]Column*...) 
+- void(<Context>&, const TypedElementRowHandle*, [const]Column*...) 
+Where <Context> is TypedElementDataStorage::ISubqueryContext
+e.g. void(ISubqueryContext& Context, TypedElementRowHandle Row, ColumnType0& ColumnA, const ColumnType1& ColumnB) {...}
+)");
+
+		return [Callback = Forward<Function>(Callback)](
+			const TypedElementDataStorage::FQueryDescription& Description,
+			TypedElementDataStorage::ISubqueryContext& Context)
+		{
+			static constexpr bool ValidateColumns = true;
+			Internal::CallQueryFunction<ValidateColumns>(Description, Context, Callback);
+		};
+	}
+} // namespace TypedElementQueryBuilder
diff --git a/Engine/Source/Runtime/TypedElementFramework/Public/Elements/Interfaces/TypedElementDataStorageInterface.h b/Engine/Source/Runtime/TypedElementFramework/Public/Elements/Interfaces/TypedElementDataStorageInterface.h
--- a/Engine/Source/Runtime/TypedElementFramework/Public/Elements/Interfaces/TypedElementDataStorageInterface.h
+++ b/Engine/Source/Runtime/TypedElementFramework/Public/Elements/Interfaces/TypedElementDataStorageInterface.h
@@ -222,13 +222,15 @@ public:
 
 	using IQueryContext = TypedElementDataStorage::IQueryContext;
 	using IDirectQueryContext = TypedElementDataStorage::IDirectQueryContext;
-	using ISubqueryCallbackContext = TypedElementDataStorage::ISubqueryCallbackContext;
+	using ISubqueryContext = TypedElementDataStorage::ISubqueryContext;
 
 	using FQueryDescription = TypedElementDataStorage::FQueryDescription;
 	using QueryCallback = TypedElementDataStorage::QueryCallback;
-
-	using QueryCallbackRef = TFunctionRef<void(const FQueryDescription&, IQueryContext&)>;
-	using DirectQueryCallbackRef = TFunctionRef<void(const FQueryDescription&, IDirectQueryContext&)>;
+	using QueryCallbackRef = TypedElementDataStorage::QueryCallbackRef;
+	using DirectQueryCallback = TypedElementDataStorage::DirectQueryCallback;
+	using DirectQueryCallbackRef = TypedElementDataStorage::DirectQueryCallbackRef;
+	using SubqueryCallback = TypedElementDataStorage::SubqueryCallback;
+	using SubqueryCallbackRef = TypedElementDataStorage::SubqueryCallbackRef;
 
 	/** 
 	 * Registers a query with the data storage. The description is processed into an internal format and may be changed. If no valid
diff --git a/Engine/Source/Runtime/TypedElementFramework/Public/Elements/Interfaces/TypedElementQueryStorageInterfaces.h b/Engine/Source/Runtime/TypedElementFramework/Public/Elements/Interfaces/TypedElementQueryStorageInterfaces.h
--- a/Engine/Source/Runtime/TypedElementFramework/Public/Elements/Interfaces/TypedElementQueryStorageInterfaces.h
+++ b/Engine/Source/Runtime/TypedElementFramework/Public/Elements/Interfaces/TypedElementQueryStorageInterfaces.h
@@ -65,9 +65,9 @@ namespace TypedElementDataStorage
 	/**
 	 * Interface to be provided to query callbacks that are directly called through from a query callback.
 	 */
-	struct ISubqueryCallbackContext : public ICommonQueryContext
+	struct ISubqueryContext : public ICommonQueryContext
 	{
-		virtual ~ISubqueryCallbackContext() = default;
+		virtual ~ISubqueryContext() = default;
 	};
 
 	/**
