Fix for more platforms having misisng/incomplete C++20 headers.

diff --git a/Engine/Source/Runtime/TypedElementFramework/Public/Elements/Framework/TypedElementQueryBuilder.inl b/Engine/Source/Runtime/TypedElementFramework/Public/Elements/Framework/TypedElementQueryBuilder.inl
--- a/Engine/Source/Runtime/TypedElementFramework/Public/Elements/Framework/TypedElementQueryBuilder.inl
+++ b/Engine/Source/Runtime/TypedElementFramework/Public/Elements/Framework/TypedElementQueryBuilder.inl
@@ -10,11 +10,8 @@
 #include "Elements/Common/TypedElementQueryDescription.h"
 #include "Elements/Common/TypedElementQueryTypes.h"
 
-// Work around missing header on some platforms
-#if __has_include(<concepts>)
-#	include <concepts>
-#else
-namespace std
+// Work around missing header/implementations on some platforms
+namespace UE
 {
 	namespace detail
 	{
@@ -30,7 +27,6 @@ namespace std
 	template< class Derived, class Base >
 	concept derived_from = std::is_base_of_v<Base, Derived> && std::is_convertible_v<const volatile Derived*, const volatile Base*>;
 }
-#endif
 
 namespace TypedElementQueryBuilder
 {
@@ -282,7 +278,7 @@ namespace TypedElementQueryBuilder
 		template<typename T>
 		concept HasStaticStructMethod = requires
 		{
-			{ T::StaticStruct() } -> std::same_as<UScriptStruct*>;
+			{ T::StaticStruct() } -> UE::same_as<UScriptStruct*>;
 		};
 
 		template<typename Context>
@@ -290,14 +286,14 @@ namespace TypedElementQueryBuilder
 			const TWeakObjectPtr<const UScriptStruct>*ColumnTypes, const TypedElementDataStorage::EQueryAccessType * AccessTypes)
 		{
 			{ Ctx.GetColumnsUnguarded(TypeCount, RetrievedAddresses, ColumnTypes, AccessTypes) };
-			{ Ctx.GetRowHandles() } -> std::convertible_to<TConstArrayView<TypedElementRowHandle>>;
-			{ Ctx.GetRowCount() } -> std::convertible_to<uint32>;
+			{ Ctx.GetRowHandles() } -> UE::convertible_to<TConstArrayView<TypedElementRowHandle>>;
+			{ Ctx.GetRowCount() } -> UE::convertible_to<uint32>;
 		};
 
 		template<typename Source, typename Target>
 		concept IsCompatibleTargetContextType = 
 			SourceQueryContext<Source> && 
-			std::derived_from<std::remove_const_t<std::remove_reference_t<Target>>, Source>;
+			UE::derived_from<std::remove_const_t<std::remove_reference_t<Target>>, Source>;
 
 		template<typename... Columns>
 		constexpr bool AreAllColumnsReferences()
