Syntax update for TEDS' observers.

The original syntax for constructing observer from templates was incorrect and ended up not working. The new syntax uses explicit OnAdd and OnRemove calls to construct an observer for a specific type. This ended up being easier to read as well.

diff --git a/Engine/Source/Runtime/TypedElementFramework/Public/Elements/Framework/TypedElementQueryBuilder.h b/Engine/Source/Runtime/TypedElementFramework/Public/Elements/Framework/TypedElementQueryBuilder.h
--- a/Engine/Source/Runtime/TypedElementFramework/Public/Elements/Framework/TypedElementQueryBuilder.h
+++ b/Engine/Source/Runtime/TypedElementFramework/Public/Elements/Framework/TypedElementQueryBuilder.h
@@ -189,10 +189,13 @@ namespace TypedElementQueryBuilder
 			Remove
 		};
 
-		template<typename ColumnType>
-		explicit FObserver(EEvent MonitorForEvent);
 		TYPEDELEMENTFRAMEWORK_API FObserver(EEvent MonitorForEvent, const UScriptStruct* MonitoredColumn);
 
+		template<typename ColumnType>
+		static FObserver OnAdd();
+		template<typename ColumnType>
+		static FObserver OnRemove();
+
 		TYPEDELEMENTFRAMEWORK_API FObserver& SetEvent(EEvent MonitorForEvent);
 		TYPEDELEMENTFRAMEWORK_API FObserver& SetMonitoredColumn(const UScriptStruct* MonitoredColumn);
 		template<typename ColumnType>
diff --git a/Engine/Source/Runtime/TypedElementFramework/Public/Elements/Framework/TypedElementQueryBuilder.inl b/Engine/Source/Runtime/TypedElementFramework/Public/Elements/Framework/TypedElementQueryBuilder.inl
--- a/Engine/Source/Runtime/TypedElementFramework/Public/Elements/Framework/TypedElementQueryBuilder.inl
+++ b/Engine/Source/Runtime/TypedElementFramework/Public/Elements/Framework/TypedElementQueryBuilder.inl
@@ -109,9 +109,16 @@ namespace TypedElementQueryBuilder
 	//
 	
 	template<typename ColumnType>
-	FObserver::FObserver(EEvent MonitorForEvent)
-		: FObserver(MonitorForEvent, ColumnType::StaticStruct())
-	{}
+	FObserver FObserver::OnAdd()
+	{
+		return FObserver(FObserver::EEvent::Add, ColumnType::StaticStruct());
+	}
+
+	template<typename ColumnType>
+	FObserver FObserver::OnRemove()
+	{
+		return FObserver(FObserver::EEvent::Remove, ColumnType::StaticStruct());
+	}
 
 	template<typename ColumnType>
 	FObserver& FObserver::SetMonitoredColumn()
