Fixed missing header for non-unity builds on Android.

diff --git a/Engine/Source/Runtime/TypedElementFramework/Public/Elements/Interfaces/TypedElementQueryStorageInterfaces.h b/Engine/Source/Runtime/TypedElementFramework/Public/Elements/Interfaces/TypedElementQueryStorageInterfaces.h
--- a/Engine/Source/Runtime/TypedElementFramework/Public/Elements/Interfaces/TypedElementQueryStorageInterfaces.h
+++ b/Engine/Source/Runtime/TypedElementFramework/Public/Elements/Interfaces/TypedElementQueryStorageInterfaces.h
@@ -5,6 +5,7 @@
 #include "Containers/ContainersFwd.h"
 #include "Elements/Common/TypedElementHandles.h"
 #include "Elements/Common/TypedElementQueryTypes.h"
+#include "UObject/Class.h"
 #include "UObject/WeakObjectPtrTemplates.h"
 #include "Templates/Function.h"
 
