Fixed TEDS callback binding when no arguments present
* Existing code will become a noop if query function signature does not include any additional arguments.  Problem is due to the first `if constexpr` failing and eliding the next else block - which should be what is called.
* When mementos enabled, it caused the cleanup processor to never run as it uses the signature void(*::)(TypedElementDataStorage::IQueryContext&, TypedElementRowHandle) thus triggering the behaviour.

diff --git a/Engine/Source/Runtime/TypedElementFramework/Public/Elements/Framework/TypedElementQueryBuilder.inl b/Engine/Source/Runtime/TypedElementFramework/Public/Elements/Framework/TypedElementQueryBuilder.inl
--- a/Engine/Source/Runtime/TypedElementFramework/Public/Elements/Framework/TypedElementQueryBuilder.inl
+++ b/Engine/Source/Runtime/TypedElementFramework/Public/Elements/Framework/TypedElementQueryBuilder.inl
@@ -529,22 +529,19 @@ namespace TypedElementQueryBuilder
 				}
 				else
 				{
-					if constexpr (sizeof...(Args) > 0)
+					for (TypedElementRowHandle Row : Rows)
 					{
-						for (TypedElementRowHandle Row : Rows)
+						if constexpr (sizeof...(Args) > 0)
 						{
-							if constexpr (sizeof...(Args) > 0)
-							{
-								std::apply([this, Row, &Caller](auto&&... Column)
-									{
-										Caller(this->ContextWrapper, Row, *Column...);
-										((++Column), ...);
-									}, this->Columns);
-							}
-							else
-							{
-								Caller(this->ContextWrapper, Row);
-							}
+							std::apply([this, Row, &Caller](auto&&... Column)
+								{
+									Caller(this->ContextWrapper, Row, *Column...);
+									((++Column), ...);
+								}, this->Columns);
+						}
+						else
+						{
+							Caller(this->ContextWrapper, Row);
 						}
 					}
 				}
