Fix incorrect column additions for external objects in TEDS Compatibility layer
* The existing code was adding UObjectColumn and ClassTypeInfoColumns for registered UScriptStruct objects. Compatibility should have been registering ExternalObjectColumn and StructTypeInfoColumn as these hold the data that is used for these objects.
* More interestingly is why this compiled in the first place.  The behaviour as it was was that a UObjectColumn was created from a sliced ExternalObjectColumn - that is, the ExternalObjectColumn was sliced down to its TypedElementDataStorageColumn base type and the FTypedElementUObjectColumn constructed with the implicit constructor.  The result was the the FTypedElementUObjectColumn had an empty Object member.
* Structural fix by disabling the ability to construct a column from a slice within the AddOrGetColumn template.  It is still _possible_ to do but needs to be explicitly done by the caller.

diff --git a/Engine/Plugins/Experimental/TypedElementsDataStorage/Source/TypedElementsDataStorage/Private/TypedElementDatabaseCompatibility.cpp b/Engine/Plugins/Experimental/TypedElementsDataStorage/Source/TypedElementsDataStorage/Private/TypedElementDatabaseCompatibility.cpp
--- a/Engine/Plugins/Experimental/TypedElementsDataStorage/Source/TypedElementsDataStorage/Private/TypedElementDatabaseCompatibility.cpp
+++ b/Engine/Plugins/Experimental/TypedElementsDataStorage/Source/TypedElementsDataStorage/Private/TypedElementDatabaseCompatibility.cpp
@@ -531,8 +531,8 @@ void UTypedElementDatabaseCompatibility::TickPendingExternalObjectRegistration()
 			It->Value.ProcessEntries(*Storage, It->Key, 
 				[this](TypedElementRowHandle Row, const ExternalObjectRegistration& Object)
 				{
-					Storage->AddOrGetColumn<FTypedElementUObjectColumn>(Row, FTypedElementExternalObjectColumn{ .Object = Object.Object });
-					Storage->AddOrGetColumn<FTypedElementClassTypeInfoColumn>(Row, FTypedElementScriptStructTypeInfoColumn{ .TypeInfo = Object.TypeInfo });
+					Storage->AddOrGetColumn<FTypedElementExternalObjectColumn>(Row, FTypedElementExternalObjectColumn{ .Object = Object.Object });
+					Storage->AddOrGetColumn<FTypedElementScriptStructTypeInfoColumn>(Row, FTypedElementScriptStructTypeInfoColumn{ .TypeInfo = Object.TypeInfo });
 					// Make sure the new row is tagged for update.
 					Storage->AddColumn<FTypedElementSyncFromWorldTag>(Row);
 
diff --git a/Engine/Source/Runtime/TypedElementFramework/Public/Elements/Interfaces/TypedElementDataStorageInterface.h b/Engine/Source/Runtime/TypedElementFramework/Public/Elements/Interfaces/TypedElementDataStorageInterface.h
--- a/Engine/Source/Runtime/TypedElementFramework/Public/Elements/Interfaces/TypedElementDataStorageInterface.h
+++ b/Engine/Source/Runtime/TypedElementFramework/Public/Elements/Interfaces/TypedElementDataStorageInterface.h
@@ -354,12 +354,25 @@ void ITypedElementDataStorageInterface::RemoveColumns(TypedElementRowHandle Row)
 	RemoveColumns(Row, { Columns::StaticStruct()...});
 }
 
+template<typename ColumnType, typename FirstArg, typename... NextArgs>
+struct TCompileTimeSliceDisabler
+{
+	// Fail assertion if the only argument is derived from FTypedElementDataStorageColumn and not the same as ColumnType
+	// This gives a good indication that object slicing is probably happening.
+	static_assert(!(sizeof...(NextArgs) == 0 && std::is_base_of_v<FTypedElementDataStorageColumn, FirstArg> && !std::is_same_v<ColumnType, FirstArg>),
+		"Probable object slicing detected. The invoked constructor of ColumnType uses a different column type as it's first argument. "
+		"This is detected as a likely object slice and disabled. "
+		"If this is what you intended, then explicitly slice the object using the constructor before passing it as an argument");
+};
+
 template<typename ColumnType, typename... Args>
 ColumnType* ITypedElementDataStorageInterface::AddOrGetColumn(TypedElementRowHandle Row, Args... Arguments)
 {
 	auto* Result = reinterpret_cast<ColumnType*>(AddOrGetColumnData(Row, ColumnType::StaticStruct()));
 	if constexpr (sizeof...(Arguments) > 0)
 	{
+		[[maybe_unused]] TCompileTimeSliceDisabler<ColumnType, Args...> SliceDisabler;
+		
 		if (Result)
 		{
 			new(Result) ColumnType{ std::forward<Args>(Arguments)... };
