Follow up changes for improving AddOrGetColumn
* Improve readability of the static_assert
* Added overload to allow type deduction of column

diff --git a/Engine/Source/Runtime/TypedElementFramework/Public/Elements/Interfaces/TypedElementDataStorageInterface.h b/Engine/Source/Runtime/TypedElementFramework/Public/Elements/Interfaces/TypedElementDataStorageInterface.h
--- a/Engine/Source/Runtime/TypedElementFramework/Public/Elements/Interfaces/TypedElementDataStorageInterface.h
+++ b/Engine/Source/Runtime/TypedElementFramework/Public/Elements/Interfaces/TypedElementDataStorageInterface.h
@@ -313,6 +313,16 @@ public:
 	 */
 	template<typename ColumnType, typename... Args>
 	ColumnType* AddOrGetColumn(TypedElementRowHandle Row, Args... Arguments);
+
+	/**
+	 * Returns a pointer to the column of the given row or creates a new one if not found.
+	 * Enables type deduction of ColumnType from Column argument.
+	 * 
+	 * For example, FTransformColumn added and deduced from second argument:
+	 * StorageInterface->AddOrGetColumn(Row, FTransformColumn{.Transform = Transform});
+	 */
+	template<typename ColumnType>
+	ColumnType* AddOrGetColumn(TypedElementRowHandle Row, ColumnType&& Column);
 	
 	/** Returns a pointer to the column of the given row or a nullptr if the type couldn't be found or the row doesn't exist. */
 	template<typename ColumnType>
@@ -355,11 +365,14 @@ void ITypedElementDataStorageInterface::RemoveColumns(TypedElementRowHandle Row)
 }
 
 template<typename ColumnType, typename FirstArg, typename... NextArgs>
-struct TCompileTimeSliceDisabler
+struct TConstructFromSlicedObjectDisabler
 {
 	// Fail assertion if the only argument is derived from FTypedElementDataStorageColumn and not the same as ColumnType
 	// This gives a good indication that object slicing is probably happening.
-	static_assert(!(sizeof...(NextArgs) == 0 && std::is_base_of_v<FTypedElementDataStorageColumn, FirstArg> && !std::is_same_v<ColumnType, FirstArg>),
+	static_assert(!(
+		sizeof...(NextArgs) == 0 &&                                    // There is only one argument to the callback and ...
+		std::is_base_of_v<FTypedElementDataStorageColumn, FirstArg> && // ... the type of the argument derives from FTypedElementDataStorageColumn
+		!std::is_same_v<ColumnType, FirstArg>),                        // ... but it isn't the same type as the column we are trying to add
 		"Probable object slicing detected. The invoked constructor of ColumnType uses a different column type as it's first argument. "
 		"This is detected as a likely object slice and disabled. "
 		"If this is what you intended, then explicitly slice the object using the constructor before passing it as an argument");
@@ -368,10 +381,10 @@ struct TCompileTimeSliceDisabler
 template<typename ColumnType, typename... Args>
 ColumnType* ITypedElementDataStorageInterface::AddOrGetColumn(TypedElementRowHandle Row, Args... Arguments)
 {
-	auto* Result = reinterpret_cast<ColumnType*>(AddOrGetColumnData(Row, ColumnType::StaticStruct()));
+	ColumnType* Result = static_cast<ColumnType*>(AddOrGetColumnData(Row, ColumnType::StaticStruct()));
 	if constexpr (sizeof...(Arguments) > 0)
 	{
-		[[maybe_unused]] TCompileTimeSliceDisabler<ColumnType, Args...> SliceDisabler;
+		[[maybe_unused]] TConstructFromSlicedObjectDisabler<ColumnType, Args...> SliceDisabler;
 		
 		if (Result)
 		{
@@ -381,6 +394,18 @@ ColumnType* ITypedElementDataStorageInterface::AddOrGetColumn(TypedElementRowHan
 	return Result;
 }
 
+template<typename ColumnType>
+ColumnType* ITypedElementDataStorageInterface::AddOrGetColumn(TypedElementRowHandle Row, ColumnType&& Column)
+{
+	ColumnType* Result = static_cast<ColumnType*>(AddOrGetColumnData(Row, ColumnType::StaticStruct()));
+
+	if (Result)
+	{
+		new(Result) ColumnType(MoveTemp(Column));
+	}
+	return Result;
+}
+
 template<typename ColumnType>
 ColumnType* ITypedElementDataStorageInterface::GetColumn(TypedElementRowHandle Row)
 {
