Boolean Query Conditions.

Introduces the ability to combine columns together into a query condition. Example: FColumn<A>() && (FColumn<B>() || FColumn<C>));
At the moment only ANDs and ORs are supported plus scopes. This is enough for the first intended use which is to allow widget constructors to have broader registration conditions with TEDS UI.

diff --git a/Engine/Source/Runtime/TypedElementFramework/Private/Elements/Framework/TypedElementQueryBuilder.cpp b/Engine/Source/Runtime/TypedElementFramework/Private/Elements/Framework/TypedElementQueryBuilder.cpp
--- a/Engine/Source/Runtime/TypedElementFramework/Private/Elements/Framework/TypedElementQueryBuilder.cpp
+++ b/Engine/Source/Runtime/TypedElementFramework/Private/Elements/Framework/TypedElementQueryBuilder.cpp
@@ -2,7 +2,9 @@
 
 #include "Elements/Framework/TypedElementQueryBuilder.h"
 
+#include "Algo/BinarySearch.h"
 #include "Elements/Framework/TypedElementMetaData.h"
+#include "GenericPlatform/GenericPlatformMath.h"
 
 namespace TypedElementQueryBuilder
 {
@@ -29,9 +31,494 @@ namespace TypedElementQueryBuilder
 		return TypeOptional(FTopLevelAssetPath{ FAnsiStringView{ Name, IntCastChecked<int32>(NameSize) } });
 	}
 
-	/**
-	 * DependsOn
-	 */
+
+	//
+	// FColumnBase
+	//
+
+	constexpr FColumnBase::FColumnBase(TWeakObjectPtr<const UScriptStruct> ColumnTypeInfo)
+		: TypeInfo(ColumnTypeInfo)
+	{
+	}
+
+
+
+	//
+	// FQueryConditions
+	//
+
+	FQueryConditions::FQueryConditions(FColumnBase Column)
+		: ColumnCount(1)
+	{
+		Columns[0] = Column.TypeInfo;
+	}
+
+	void FQueryConditions::AppendToString(FString& Output) const
+	{
+		if (TokenCount > 0)
+		{
+			Output += "{ ";
+			uint8_t ColumnIndex = 0;
+			if (Tokens[0] != Token::ScopeOpen)
+			{
+				ColumnIndex = 1;
+				AppendName(Output, Columns[0]);
+			}
+			for (uint8_t TokenIndex = 0; TokenIndex < TokenCount; ++TokenIndex)
+			{
+				switch (Tokens[TokenIndex])
+				{
+				case Token::And:
+					Output += " && ";
+					if (!EntersScopeNext(TokenIndex))
+					{
+						AppendName(Output, Columns[ColumnIndex++]);
+					}
+					break;
+				case Token::Or:
+					Output += " || ";
+					if (!EntersScopeNext(TokenIndex))
+					{
+						AppendName(Output, Columns[ColumnIndex++]);
+					}
+					break;
+				case Token::ScopeOpen:
+					Output += "( ";
+					if (!EntersScopeNext(TokenIndex))
+					{
+						AppendName(Output, Columns[ColumnIndex++]);
+					}
+					break;
+				case Token::ScopeClose:
+					Output += " )";
+					break;
+				default:
+					checkf(false, TEXT("Invalid query token"));
+					break;
+				}
+			}
+			Output += " }";
+		}
+	}
+
+	bool FQueryConditions::Verify(TConstArrayView<FColumnBase> AvailableColumns) const
+	{
+		return VerifyBootstrap(
+			[&AvailableColumns](uint8_t ColumnIndex, TWeakObjectPtr<const UScriptStruct> Column)
+			{
+				for (const FColumnBase& Target : AvailableColumns)
+				{
+					if (Target.TypeInfo == Column)
+					{
+						return true;
+					}
+				}
+				return false;
+			});
+	}
+
+	bool FQueryConditions::Verify(TArray<TWeakObjectPtr<const UScriptStruct>>& MatchedColumns, TConstArrayView<FColumnBase> AvailableColumns,
+		bool AvailableColumnsAreSorted) const
+	{
+		static_assert(MaxColumnCount < 64, "Query conditions use a bit mask to locate matches. As a result MaxColumnCount can be larger than 64.");
+		uint64 Matches = 0;
+		bool Result = AvailableColumnsAreSorted
+			? VerifyBootstrap(
+				[&Matches, &AvailableColumns](uint8_t ColumnIndex, TWeakObjectPtr<const UScriptStruct> Column)
+				{
+					auto Projection = [](const FColumnBase& ColumnBase) 
+					{ 
+						return ColumnBase.TypeInfo.Get();
+					};
+					if (Algo::BinarySearchBy(AvailableColumns, Column.Get(), Projection) != INDEX_NONE)
+					{
+						Matches |= (uint64)1 << ColumnIndex;
+						return true;
+					}
+					return false;
+				})
+			: VerifyBootstrap(
+				[&Matches, &AvailableColumns](uint8_t ColumnIndex, TWeakObjectPtr<const UScriptStruct> Column)
+				{
+					for (const FColumnBase& Target : AvailableColumns)
+					{
+						if (Target.TypeInfo == Column)
+						{
+							Matches |= (uint64)1 << ColumnIndex;
+							return true;
+						}
+					}
+				return false;
+				});
+		if (Result)
+		{
+			ConvertColumnBitToArray(MatchedColumns, Matches);
+		}
+		return Result;
+	}
+
+	bool FQueryConditions::Verify(TConstArrayView<TWeakObjectPtr<const UScriptStruct>> AvailableColumns, bool AvailableColumnsAreSorted) const
+	{
+		return AvailableColumnsAreSorted
+			? VerifyBootstrap(
+				[&AvailableColumns](uint8_t ColumnIndex, TWeakObjectPtr<const UScriptStruct> Column)
+				{
+					auto Projection = [](TWeakObjectPtr<const UScriptStruct> ColumnBase)
+					{ 
+						return ColumnBase.Get();
+					};
+					return Algo::BinarySearchBy(AvailableColumns, Column.Get(), Projection) != INDEX_NONE;
+				})
+			: VerifyBootstrap(
+				[&AvailableColumns](uint8_t ColumnIndex, TWeakObjectPtr<const UScriptStruct> Column)
+				{
+					return AvailableColumns.Find(Column) != INDEX_NONE;
+				});
+	}
+
+	bool FQueryConditions::Verify(TArray<TWeakObjectPtr<const UScriptStruct>>& MatchedColumns,
+		TConstArrayView<TWeakObjectPtr<const UScriptStruct>> AvailableColumns, bool AvailableColumnsAreSorted) const
+	{
+		static_assert(MaxColumnCount < 64, "Query conditions use a bit mask to locate matches. As a result MaxColumnCount can be larger than 64.");
+		uint64 Matches = 0;
+		bool Result = AvailableColumnsAreSorted
+			? VerifyBootstrap(
+				[&Matches, &AvailableColumns](uint8_t ColumnIndex, TWeakObjectPtr<const UScriptStruct> Column)
+				{
+					auto Projection = [](TWeakObjectPtr<const UScriptStruct> ColumnBase)
+					{
+						return ColumnBase.Get();
+					};
+					if (Algo::BinarySearchBy(AvailableColumns, Column.Get(), Projection) != INDEX_NONE)
+					{
+						Matches |= (uint64)1 << ColumnIndex;
+						return true;
+					}
+					return false;
+				})
+			: VerifyBootstrap(
+				[&Matches, &AvailableColumns](uint8_t ColumnIndex, TWeakObjectPtr<const UScriptStruct> Column)
+				{
+					return AvailableColumns.Find(Column) != INDEX_NONE;
+				});
+		if (Result)
+		{
+			ConvertColumnBitToArray(MatchedColumns, Matches);
+		}
+		return Result;
+	}
+
+	bool FQueryConditions::Verify(TSet<TWeakObjectPtr<const UScriptStruct>> AvailableColumns) const
+	{
+		return VerifyBootstrap(
+			[&AvailableColumns](uint8_t ColumnIndex, TWeakObjectPtr<const UScriptStruct> Column)
+			{
+				return AvailableColumns.Find(Column) != nullptr;
+			});
+	}
+
+	uint8_t FQueryConditions::MinimumColumnMatchRequired() const
+	{
+		if (TokenCount > 0)
+		{
+			uint8_t Front = 0;
+			return MinimumColumnMatchRequiredRange(Front);
+		}
+		else if (ColumnCount == 1)
+		{
+			return 1;
+		}
+		else
+		{
+			return 0;
+		}
+	}
+
+	TConstArrayView<TWeakObjectPtr<const UScriptStruct>> FQueryConditions::GetColumns() const
+	{
+		return TConstArrayView<TWeakObjectPtr<const UScriptStruct>>(Columns, ColumnCount);
+	}
+
+	bool FQueryConditions::IsEmpty() const
+	{
+		return ColumnCount == 0;
+	}
+
+	void FQueryConditions::AppendName(FString& Output, TWeakObjectPtr<const UScriptStruct> TypeInfo) const
+	{
+#if WITH_EDITORONLY_DATA
+		static FName DisplayNameName(TEXT("DisplayName"));
+		if (const FString* Name = TypeInfo->FindMetaData(DisplayNameName))
+		{
+			Output += *Name;
+		}
+#else
+		Output += TEXT("<Unavailable>");
+#endif
+	}
+
+	bool FQueryConditions::EntersScopeNext(uint8_t Index) const
+	{
+		return Index < (TokenCount - 1) && Tokens[Index + 1] == Token::ScopeOpen;
+	}
+
+	bool FQueryConditions::EntersScope(uint8_t Index) const
+	{
+		return Tokens[Index] == Token::ScopeOpen;
+	}
+
+	bool FQueryConditions::Contains(TWeakObjectPtr<const UScriptStruct> ColumnType, const TArray<FColumnBase>& AvailableColumns) const
+	{
+		for (const FColumnBase& Column : AvailableColumns)
+		{
+			if (ColumnType == Column.TypeInfo)
+			{
+				return true;
+			}
+		}
+		return false;
+	}
+
+	template<typename ContainsCallback>
+	bool FQueryConditions::VerifyBootstrap(ContainsCallback&& Contains) const
+	{
+		if (TokenCount > 0)
+		{
+			uint8_t TokenIndex = 0;
+			uint8_t ColumnIndex = 0;
+			return VerifyRange(TokenIndex, ColumnIndex, Forward<ContainsCallback>(Contains));
+		}
+		else if (ColumnCount == 1)
+		{
+			return Contains(0, Columns[0]);
+		}
+		else
+		{
+			return false;
+		}
+	}
+
+	template<typename ContainsCallback>
+	bool FQueryConditions::VerifyRange(uint8_t& TokenIndex, uint8_t& ColumnIndex, ContainsCallback Contains) const
+	{
+		auto Init = [&]()
+		{
+			if (EntersScope(TokenIndex))
+			{
+				return VerifyRange(++TokenIndex, ColumnIndex, Contains);
+			}
+			else
+			{
+				bool Result =  Contains(ColumnIndex, Columns[ColumnIndex]);
+				ColumnIndex++;
+				return Result;
+			}
+			return true;
+		};
+		bool Result = Init();
+
+		while (TokenIndex < TokenCount)
+		{
+			Token T = Tokens[TokenIndex];
+			++TokenIndex;
+			switch (T)
+			{
+			case Token::And:
+				if (EntersScope(TokenIndex))
+				{
+					Result = VerifyRange(++TokenIndex, ColumnIndex, Contains) && Result;
+				}
+				else
+				{
+					Result = Contains(ColumnIndex, Columns[ColumnIndex]) && Result;
+					ColumnIndex++;
+				}
+				break;
+			case Token::Or:
+				if (EntersScope(TokenIndex))
+				{
+					Result = VerifyRange(++TokenIndex, ColumnIndex, Contains) || Result;
+				}
+				else
+				{
+					Result = Contains(ColumnIndex, Columns[ColumnIndex]) || Result;
+					ColumnIndex++;
+				}
+				break;
+			case Token::ScopeOpen:
+				checkf(false, TEXT("The scope open in a query should be called during processing as it should be captured by an earlier statement."));
+				break;
+			case Token::ScopeClose:
+				return Result;
+			default:
+				checkf(false, TEXT("Encountered an unknown query token."));
+				break;
+			}
+		}
+		return Result;
+	}
+
+	void FQueryConditions::ConvertColumnBitToArray(TArray<TWeakObjectPtr<const UScriptStruct>>& MatchedColumns, uint64 ColumnBits) const
+	{
+		uint64 Index = 0;
+		while (ColumnBits)
+		{
+			if (ColumnBits & 1)
+			{
+				MatchedColumns.Add(Columns[Index]);
+			}
+			++Index;
+			ColumnBits >>= 1;
+		}
+	}
+
+	uint8_t FQueryConditions::MinimumColumnMatchRequiredRange(uint8_t& Front) const
+	{
+		uint8_t Result = EntersScope(Front) ? MinimumColumnMatchRequiredRange(++Front) : 1;
+		for (; Front < TokenCount; ++Front)
+		{
+			Token T = Tokens[Front];
+			switch (T)
+			{
+			case Token::And:
+				if (EntersScopeNext(Front))
+				{
+					Front += 2;
+					Result += MinimumColumnMatchRequiredRange(Front);
+				}
+				else
+				{
+					++Result;
+				}
+				break;
+			case Token::Or:
+				if (EntersScopeNext(Front))
+				{
+					Front += 2;
+					uint8_t Lhs = MinimumColumnMatchRequiredRange(Front);
+					Result = FMath::Min(Result, Lhs);
+				}
+				break;
+			case Token::ScopeOpen:
+				checkf(false, TEXT("The scope open in a query should be called during processing as it should be captured by an earlier statement."));
+				break;
+			case Token::ScopeClose:
+				++Front;
+				return Result;
+			default:
+				checkf(false, TEXT("Encountered an unknown query token."));
+				break;
+			}
+		}
+		return Result;
+	}
+
+	void FQueryConditions::AppendQuery(FQueryConditions& Target, const FQueryConditions& Source)
+	{
+		checkf(Target.ColumnCount + Source.ColumnCount < MaxColumnCount, TEXT("Too many columns in the query."));
+		for (uint8_t Index = 0; Index < Source.ColumnCount; ++Index)
+		{
+			Target.Columns[Target.ColumnCount + Index] = Source.Columns[Index];
+		}
+
+		checkf(Target.TokenCount + Source.TokenCount < MaxTokenCount, TEXT("Too many operations in the query. Try simplifying your query."));
+		for (uint8_t Index = 0; Index < Source.TokenCount; ++Index)
+		{
+			Target.Tokens[Target.TokenCount + Index] = Source.Tokens[Index];
+		}
+
+		Target.ColumnCount += Source.ColumnCount;
+		Target.TokenCount += Source.TokenCount;
+	}
+
+	FQueryConditions operator&&(const FQueryConditions& Lhs, FColumnBase Rhs)
+	{
+		FQueryConditions Result = Lhs;
+		Result.Columns[Result.ColumnCount++] = Rhs.TypeInfo;
+		Result.Tokens[Result.TokenCount++] = FQueryConditions::Token::And;
+		return Result;
+	}
+
+	FQueryConditions operator&&(const FQueryConditions& Lhs, const FQueryConditions& Rhs)
+	{
+		FQueryConditions Result;
+		Result.Tokens[Result.TokenCount++] = FQueryConditions::Token::ScopeOpen;
+		FQueryConditions::AppendQuery(Result, Lhs);
+		Result.Tokens[Result.TokenCount++] = FQueryConditions::Token::ScopeClose;
+
+		Result.Tokens[Result.TokenCount++] = FQueryConditions::Token::And;
+
+		Result.Tokens[Result.TokenCount++] = FQueryConditions::Token::ScopeOpen;
+		FQueryConditions::AppendQuery(Result, Rhs);
+		Result.Tokens[Result.TokenCount++] = FQueryConditions::Token::ScopeClose;
+
+		return Result;
+	}
+
+	FQueryConditions operator&&(FColumnBase Lhs, FColumnBase Rhs)
+	{
+		FQueryConditions Result(Lhs);
+		Result.Columns[Result.ColumnCount++] = Rhs.TypeInfo;
+		Result.Tokens[Result.TokenCount++] = FQueryConditions::Token::And;
+		return Result;
+	}
+
+	FQueryConditions operator&&(FColumnBase Lhs, const FQueryConditions& Rhs)
+	{
+		FQueryConditions Result(Lhs);
+		Result.Tokens[Result.TokenCount++] = FQueryConditions::Token::And;
+		Result.Tokens[Result.TokenCount++] = FQueryConditions::Token::ScopeOpen;
+		FQueryConditions::AppendQuery(Result, Rhs);
+		Result.Tokens[Result.TokenCount++] = FQueryConditions::Token::ScopeClose;
+		return Result;
+	}
+
+	FQueryConditions operator||(const FQueryConditions& Lhs, FColumnBase Rhs)
+	{
+		FQueryConditions Result = Lhs;
+		Result.Columns[Result.ColumnCount++] = Rhs.TypeInfo;
+		Result.Tokens[Result.TokenCount++] = FQueryConditions::Token::Or;
+		return Result;
+	}
+
+	FQueryConditions operator||(const FQueryConditions& Lhs, const FQueryConditions& Rhs)
+	{
+		FQueryConditions Result;
+		Result.Tokens[Result.TokenCount++] = FQueryConditions::Token::ScopeOpen;
+		FQueryConditions::AppendQuery(Result, Lhs);
+		Result.Tokens[Result.TokenCount++] = FQueryConditions::Token::ScopeClose;
+
+		Result.Tokens[Result.TokenCount++] = FQueryConditions::Token::Or;
+
+		Result.Tokens[Result.TokenCount++] = FQueryConditions::Token::ScopeOpen;
+		FQueryConditions::AppendQuery(Result, Rhs);
+		Result.Tokens[Result.TokenCount++] = FQueryConditions::Token::ScopeClose;
+
+		return Result;
+	}
+
+	FQueryConditions operator||(FColumnBase Lhs, FColumnBase Rhs)
+	{
+		FQueryConditions Result(Lhs);
+		Result.Columns[Result.ColumnCount++] = Rhs.TypeInfo;
+		Result.Tokens[Result.TokenCount++] = FQueryConditions::Token::Or;
+		return Result;
+	}
+
+	FQueryConditions operator||(FColumnBase Lhs, const FQueryConditions& Rhs)
+	{
+		FQueryConditions Result(Lhs);
+		Result.Tokens[Result.TokenCount++] = FQueryConditions::Token::Or;
+		Result.Tokens[Result.TokenCount++] = FQueryConditions::Token::ScopeOpen;
+		FQueryConditions::AppendQuery(Result, Rhs);
+		Result.Tokens[Result.TokenCount++] = FQueryConditions::Token::ScopeClose;
+		return Result;
+	}
+
+
+	//
+	// DependsOn
+	//
 
 	FDependency::FDependency(ITypedElementDataStorageInterface::FQueryDescription* Query)
 		: Query(Query)
diff --git a/Engine/Source/Runtime/TypedElementFramework/Public/Elements/Framework/TypedElementQueryBuilder.h b/Engine/Source/Runtime/TypedElementFramework/Public/Elements/Framework/TypedElementQueryBuilder.h
--- a/Engine/Source/Runtime/TypedElementFramework/Public/Elements/Framework/TypedElementQueryBuilder.h
+++ b/Engine/Source/Runtime/TypedElementFramework/Public/Elements/Framework/TypedElementQueryBuilder.h
@@ -5,6 +5,7 @@
 #include "CoreMinimal.h"
 #include "Elements/Interfaces/TypedElementDataStorageInterface.h"
 #include "Templates/SubclassOf.h"
+#include "UObject/WeakObjectPtrTemplates.h"
 
 class UScriptStruct;
 class USubsystem;
@@ -109,6 +110,128 @@ namespace TypedElementQueryBuilder
 	TYPEDELEMENTFRAMEWORK_API const UScriptStruct* operator""_Type(const char* Name, std::size_t NameSize);
 	TYPEDELEMENTFRAMEWORK_API const UScriptStruct* operator""_TypeOptional(const char* Name, std::size_t NameSize);
 
+	struct FColumnBase
+	{
+		TWeakObjectPtr<const UScriptStruct> TypeInfo = nullptr;
+	
+	protected:
+		FColumnBase() = default;
+		TYPEDELEMENTFRAMEWORK_API constexpr explicit FColumnBase(TWeakObjectPtr<const UScriptStruct> ColumnTypeInfo);
+	};
+
+	template<typename T = void>
+	struct FColumn final : public FColumnBase
+	{
+		template<typename = std::enable_if_t<!std::is_same_v<T, void>>>
+		constexpr FColumn() : FColumnBase(T::StaticStruct()) {}
+		constexpr explicit FColumn(TWeakObjectPtr<const UScriptStruct> ColumnTypeInfo) : FColumnBase(ColumnTypeInfo) {};
+	};
+
+	/** 
+	 * Product of boolean combination of multiple columns. This can be used to verify if a collection of columns match 
+	 * the stored columns.
+	 */
+	class FQueryConditions final
+	{
+	public:
+		TYPEDELEMENTFRAMEWORK_API friend FQueryConditions operator&&(const FQueryConditions& Lhs, FColumnBase Rhs);
+		TYPEDELEMENTFRAMEWORK_API friend FQueryConditions operator&&(const FQueryConditions& Lhs, const FQueryConditions& Rhs);
+		TYPEDELEMENTFRAMEWORK_API friend FQueryConditions operator&&(FColumnBase Lhs, FColumnBase Rhs);
+		TYPEDELEMENTFRAMEWORK_API friend FQueryConditions operator&&(FColumnBase Lhs, const FQueryConditions& Rhs);
+		
+		TYPEDELEMENTFRAMEWORK_API friend FQueryConditions operator||(const FQueryConditions& Lhs, FColumnBase Rhs);
+		TYPEDELEMENTFRAMEWORK_API friend FQueryConditions operator||(const FQueryConditions& Lhs, const FQueryConditions& Rhs);
+		TYPEDELEMENTFRAMEWORK_API friend FQueryConditions operator||(FColumnBase Lhs, FColumnBase Rhs);
+		TYPEDELEMENTFRAMEWORK_API friend FQueryConditions operator||(FColumnBase Lhs, const FQueryConditions& Rhs);
+
+		FQueryConditions() = default;
+		// Not marked as "explicit" to allow conversion from a column. This means that conditions with a single
+		// argument can be written in the same way as ones that use combinations.
+		TYPEDELEMENTFRAMEWORK_API FQueryConditions(FColumnBase Column);
+		
+		/** Convert the conditions into a string and append them to the provided string. */
+		TYPEDELEMENTFRAMEWORK_API void AppendToString(FString& Output) const;
+
+		/** Runs the provided list of columns through the conditions and returns true if a valid combination of columns is found. */
+		TYPEDELEMENTFRAMEWORK_API bool Verify(TConstArrayView<FColumnBase> AvailableColumns) const;
+		/** 
+		 * Runs the provided list of columns through the conditions and returns true if a valid combination of columns is found.
+		 * This version returns a list of the columns that were used to match the condition.
+		 */
+		TYPEDELEMENTFRAMEWORK_API bool Verify(
+			TArray<TWeakObjectPtr<const UScriptStruct>>& MatchedColumns,
+			TConstArrayView<FColumnBase> AvailableColumns, 
+			bool AvailableColumnsAreSorted = false) const;
+		/** Runs the provided list of columns through the conditions and returns true if a valid combination of columns is found. */
+		TYPEDELEMENTFRAMEWORK_API bool Verify(TConstArrayView<TWeakObjectPtr<const UScriptStruct>> AvailableColumns, 
+			bool AvailableColumnsAreSorted = false) const;
+		/**
+		 * Runs the provided list of columns through the conditions and returns true if a valid combination of columns is found.
+		 * This version returns a list of the columns that were used to match the condition.
+		 */
+		TYPEDELEMENTFRAMEWORK_API bool Verify(
+			TArray<TWeakObjectPtr<const UScriptStruct>>& MatchedColumns,
+			TConstArrayView<TWeakObjectPtr<const UScriptStruct>> AvailableColumns, 
+			bool AvailableColumnsAreSorted = false) const;
+		/** Runs the provided list of columns through the conditions and returns true if a valid combination of columns is found. */
+		TYPEDELEMENTFRAMEWORK_API bool Verify(TSet<TWeakObjectPtr<const UScriptStruct>> AvailableColumns) const;
+
+		/** Returns the minimum number of columns needed for a successful match. */
+		TYPEDELEMENTFRAMEWORK_API uint8_t MinimumColumnMatchRequired() const;
+		/** Returns a list of all columns used. This can include duplicate columns. */
+		TYPEDELEMENTFRAMEWORK_API TConstArrayView<TWeakObjectPtr<const UScriptStruct>> GetColumns() const;
+		/** Whether or not there are any columns registered for operation.*/
+		TYPEDELEMENTFRAMEWORK_API bool IsEmpty() const;
+
+	private:
+		void AppendName(FString& Output, TWeakObjectPtr<const UScriptStruct> TypeInfo) const;
+		
+		bool EntersScopeNext(uint8_t Index) const;
+		bool EntersScope(uint8_t Index) const;
+		bool Contains(TWeakObjectPtr<const UScriptStruct> ColumnType, const TArray<FColumnBase>& AvailableColumns) const;
+		
+		template<typename ContainsCallback>
+		bool VerifyBootstrap(ContainsCallback&& Contains) const;
+		template<typename ContainsCallback>
+		bool VerifyRange(uint8_t& TokenIndex, uint8_t& ColumnIndex, ContainsCallback Contains) const;
+
+		void ConvertColumnBitToArray(TArray<TWeakObjectPtr<const UScriptStruct>>& MatchedColumns, uint64 ColumnBits) const;
+
+		uint8_t MinimumColumnMatchRequiredRange(uint8_t& Front) const;
+
+		static void AppendQuery(FQueryConditions& Target, const FQueryConditions& Source);
+
+		static constexpr SIZE_T MaxColumnCount = 32;
+		static constexpr SIZE_T MaxTokenCount = 64;
+
+		enum class Token : uint8_t
+		{
+			None,
+			And,
+			Or,
+			ScopeOpen,
+			ScopeClose
+		};
+		using ColumnArray = TWeakObjectPtr<const UScriptStruct>[MaxColumnCount];
+		using TokenArray = Token[MaxTokenCount];
+
+		ColumnArray Columns{ nullptr };
+		TokenArray Tokens{ Token::None };
+		uint8_t ColumnCount = 0;
+		uint8_t TokenCount = 0;
+
+	};
+
+	TYPEDELEMENTFRAMEWORK_API FQueryConditions operator&&(const FQueryConditions& Lhs, FColumnBase Rhs);
+	TYPEDELEMENTFRAMEWORK_API FQueryConditions operator&&(const FQueryConditions& Lhs, const FQueryConditions& Rhs);
+	TYPEDELEMENTFRAMEWORK_API FQueryConditions operator&&(FColumnBase Lhs, FColumnBase Rhs);
+	TYPEDELEMENTFRAMEWORK_API FQueryConditions operator&&(FColumnBase Lhs, const FQueryConditions& Rhs);
+	
+	TYPEDELEMENTFRAMEWORK_API FQueryConditions operator||(const FQueryConditions& Lhs, FColumnBase Rhs);
+	TYPEDELEMENTFRAMEWORK_API FQueryConditions operator||(const FQueryConditions& Lhs, const FQueryConditions& Rhs);
+	TYPEDELEMENTFRAMEWORK_API FQueryConditions operator||(FColumnBase Lhs, FColumnBase Rhs);
+	TYPEDELEMENTFRAMEWORK_API FQueryConditions operator||(FColumnBase Lhs, const FQueryConditions& Rhs);
+
 	class FDependency final
 	{
 		friend class Count;
diff --git a/Engine/Source/Runtime/TypedElementFramework/Tests/Elements/Framework/TypedElementQueryBuilderTests.cpp b/Engine/Source/Runtime/TypedElementFramework/Tests/Elements/Framework/TypedElementQueryBuilderTests.cpp
new file mode 100644
--- /dev/null
+++ b/Engine/Source/Runtime/TypedElementFramework/Tests/Elements/Framework/TypedElementQueryBuilderTests.cpp
@@ -0,0 +1,268 @@
+// Copyright Epic Games, Inc. All Rights Reserved.
+
+#if WITH_TESTS
+#include "TypedElementQueryBuilderTests.h"
+
+#include "Algo/Sort.h"
+#include "CoreTypes.h"
+#include "Containers/Array.h"
+#include "Containers/UnrealString.h"
+#include "Elements/Framework/TypedElementQueryBuilder.h"
+#include "Tests/TestHarnessAdapter.h"
+#include "UObject/Class.h"
+
+using namespace TypedElementQueryBuilder;
+
+
+static void AppendColumnName(FString& Output, TWeakObjectPtr<const UScriptStruct> TypeInfo)
+{
+#if WITH_EDITORONLY_DATA
+	static FName DisplayNameName(TEXT("DisplayName"));
+	if (const FString* Name = TypeInfo->FindMetaData(DisplayNameName))
+	{
+		Output += *Name;
+	}
+#else
+	Output += TEXT("<Unavailable>");
+#endif
+}
+
+static bool TestMatching(const FQueryConditions& TestQuery, const TArray<FColumnBase>& RequestedColumns, bool Expected, bool Sort = false)
+{
+	if (Sort)
+	{
+		Algo::SortBy(RequestedColumns,
+			[](const FColumnBase& Column) { return Column.TypeInfo.Get(); });
+	}
+
+	TArray<TWeakObjectPtr<const UScriptStruct>> Matches;
+	bool Result = TestQuery.Verify(Matches, RequestedColumns);
+
+	FString Description = (Result == Expected) ? TEXT("[Pass] ") : TEXT("[Fail] ");
+	TestQuery.AppendToString(Description);
+
+	{
+		Description += " -> { ";
+		auto It = RequestedColumns.begin();
+		AppendColumnName(Description, (*It).TypeInfo);
+		++It;
+		for (; It != RequestedColumns.end(); ++It)
+		{
+			Description += ", ";
+			AppendColumnName(Description, (*It).TypeInfo);
+		}
+		Description += " } ";
+	}
+
+	if (Expected)
+	{
+		if (!Matches.IsEmpty())
+		{
+			Description += " -> { ";
+			auto It = Matches.begin();
+			AppendColumnName(Description, *It);
+			++It;
+			for (It; It != Matches.end(); ++It)
+			{
+				Description += ", ";
+				AppendColumnName(Description, *It);
+			}
+			Description += " } ";
+		}
+
+		for (TWeakObjectPtr<const UScriptStruct> Match : Matches)
+		{
+			bool Found = false;
+			for (const FColumnBase& Requested : RequestedColumns)
+			{
+				if (Match == Requested.TypeInfo)
+				{
+					Found = true;
+					break;
+				}
+			}
+
+			if (!Found)
+			{
+				Result = false;
+				Description += " [Match failed]";
+				break;
+			}
+		}
+	}
+
+	FAutomationTestFramework::Get().GetCurrentTest()->AddInfo(MoveTemp(Description));
+	return (Result == Expected);
+}
+
+TEST_CASE_NAMED(FTypedElementQueryConditions_NoColumn, "TypedElementQueryBuilder::FTypedElementQueryConditions_NoColumn", "[ApplicationContextMask][EngineFilter]")
+{
+	FQueryConditions Example;
+	
+	CHECK(Example.MinimumColumnMatchRequired() == 0);
+	CHECK(TestMatching(Example, { FColumn<FTestColumnA>() }, false));
+}
+
+TEST_CASE_NAMED(FTypedElementQueryConditions_OneColumn, "TypedElementQueryBuilder::FTypedElementQueryConditions_OneColumn", "[ApplicationContextMask][EngineFilter]")
+{
+	FQueryConditions Example{ FColumn<FTestColumnA>() };
+
+	CHECK(Example.MinimumColumnMatchRequired() == 1);
+	CHECK(TestMatching(Example, { FColumn<FTestColumnA>() }, true));
+}
+
+TEST_CASE_NAMED(FTypedElementQueryConditions1, "TypedElementQueryBuilder::FTypedElementQueryConditions A && B && C", "[ApplicationContextMask][EngineFilter]")
+{
+	FColumn TestA(FTestColumnA::StaticStruct());
+	
+	FQueryConditions Example = FColumn<FTestColumnA>() && FColumn<FTestColumnB>() && FColumn<FTestColumnC>();
+	
+	CHECK(Example.MinimumColumnMatchRequired() == 3);
+	CHECK(TestMatching(Example, { FColumn<FTestColumnA>(), FColumn<FTestColumnB>(), FColumn<FTestColumnC>() }, true));
+	CHECK(TestMatching(Example, { FColumn<FTestColumnA>(), FColumn<FTestColumnB>(), FColumn<FTestColumnD>() }, false));
+}
+
+TEST_CASE_NAMED(FTypedElementQueryConditions2, "TypedElementQueryBuilder::FTypedElementQueryConditions A || B || C", "[ApplicationContextMask][EngineFilter]")
+{
+	FQueryConditions Example = FColumn<FTestColumnA>() || FColumn<FTestColumnB>() || FColumn<FTestColumnC>();
+	
+	CHECK(Example.MinimumColumnMatchRequired() == 1);
+	CHECK(TestMatching(Example, { FColumn<FTestColumnB>() }, true));
+	CHECK(TestMatching(Example, { FColumn<FTestColumnB>(), FColumn<FTestColumnC>() }, true));
+	CHECK(TestMatching(Example, { FColumn<FTestColumnD>() }, false));
+}
+
+TEST_CASE_NAMED(FTypedElementQueryConditions3, "TypedElementQueryBuilder::FTypedElementQueryConditions A && (B || C)", "[ApplicationContextMask][EngineFilter]")
+{
+	FQueryConditions Example = FColumn<FTestColumnA>() && (FColumn<FTestColumnB>() || FColumn<FTestColumnC>());
+	
+	CHECK(Example.MinimumColumnMatchRequired() == 2);
+	CHECK(TestMatching(Example, { FColumn<FTestColumnA>(), FColumn<FTestColumnB>() }, true));
+	CHECK(TestMatching(Example, { FColumn<FTestColumnA>(), FColumn<FTestColumnC>() }, true));
+	CHECK(TestMatching(Example, { FColumn<FTestColumnA>(), FColumn<FTestColumnD>() }, false));
+	CHECK(TestMatching(Example, { FColumn<FTestColumnD>(), FColumn<FTestColumnB>() }, false));
+}
+
+TEST_CASE_NAMED(FTypedElementQueryConditions4, "TypedElementQueryBuilder::FTypedElementQueryConditions A && (B || C) && (D || E)", "[ApplicationContextMask][EngineFilter]")
+{
+	FQueryConditions Example = 
+		FColumn<FTestColumnA>() && 
+		(FColumn<FTestColumnB>() || FColumn<FTestColumnC>()) &&
+		(FColumn<FTestColumnD>() || FColumn<FTestColumnE>());
+	
+	CHECK(Example.MinimumColumnMatchRequired() == 3);
+	
+	CHECK(TestMatching(Example, { FColumn<FTestColumnA>() }, false));
+	CHECK(TestMatching(Example, { FColumn<FTestColumnA>(), FColumn<FTestColumnB>() }, false));
+
+	CHECK(TestMatching(Example, { FColumn<FTestColumnA>(), FColumn<FTestColumnB>(), FColumn<FTestColumnD>() }, true));
+	CHECK(TestMatching(Example, { FColumn<FTestColumnA>(), FColumn<FTestColumnB>(), FColumn<FTestColumnE>() }, true));
+	CHECK(TestMatching(Example, { FColumn<FTestColumnA>(), FColumn<FTestColumnC>(), FColumn<FTestColumnD>() }, true));
+
+	CHECK(TestMatching(Example, { FColumn<FTestColumnA>(), FColumn<FTestColumnC>(), FColumn<FTestColumnF>() }, false));
+	CHECK(TestMatching(Example, { FColumn<FTestColumnA>(), FColumn<FTestColumnF>(), FColumn<FTestColumnD>() }, false));
+	CHECK(TestMatching(Example, { FColumn<FTestColumnB>(), FColumn<FTestColumnC>(), FColumn<FTestColumnD>() }, false));
+}
+
+TEST_CASE_NAMED(FTypedElementQueryConditions5, "TypedElementQueryBuilder::FTypedElementQueryConditions (A || B) && (C || D) && (E || F)", "[ApplicationContextMask][EngineFilter]")
+{
+	FQueryConditions Example =
+		(FColumn<FTestColumnA>() || FColumn<FTestColumnB>()) &&
+		(FColumn<FTestColumnC>() || FColumn<FTestColumnD>()) &&
+		(FColumn<FTestColumnE>() || FColumn<FTestColumnF>());
+	
+	CHECK(Example.MinimumColumnMatchRequired() == 3); 
+	
+	CHECK(TestMatching(Example, { FColumn<FTestColumnA>() }, false));
+	CHECK(TestMatching(Example, { FColumn<FTestColumnA>(), FColumn<FTestColumnC>() }, false));
+	
+	CHECK(TestMatching(Example, { FColumn<FTestColumnA>(), FColumn<FTestColumnC>(), FColumn<FTestColumnE>() }, true));
+	CHECK(TestMatching(Example, { FColumn<FTestColumnB>(), FColumn<FTestColumnC>(), FColumn<FTestColumnE>() }, true));
+	
+	CHECK(TestMatching(Example, { FColumn<FTestColumnA>(), FColumn<FTestColumnC>(), FColumn<FTestColumnG>() }, false));
+	CHECK(TestMatching(Example, { FColumn<FTestColumnA>(), FColumn<FTestColumnG>(), FColumn<FTestColumnD>() }, false));
+	CHECK(TestMatching(Example, { FColumn<FTestColumnG>(), FColumn<FTestColumnC>(), FColumn<FTestColumnD>() }, false));
+}
+
+TEST_CASE_NAMED(FTypedElementQueryConditions6, "TypedElementQueryBuilder::FTypedElementQueryConditions ((A || B) && (C || D)) || (E && F)", "[ApplicationContextMask][EngineFilter]")
+{
+	FQueryConditions Example =
+		(
+			(FColumn<FTestColumnA>() || FColumn<FTestColumnB>()) &&
+			(FColumn<FTestColumnC>() || FColumn<FTestColumnD>())
+		) ||
+		(FColumn<FTestColumnE>() && FColumn<FTestColumnF>());
+	
+	CHECK(Example.MinimumColumnMatchRequired() == 2);
+
+	CHECK(TestMatching(Example, { FColumn<FTestColumnA>() }, false));
+	CHECK(TestMatching(Example, { FColumn<FTestColumnA>(), FColumn<FTestColumnC>() }, true));
+
+	CHECK(TestMatching(Example, { FColumn<FTestColumnE>(), FColumn<FTestColumnF>() }, true));
+	CHECK(TestMatching(Example, { FColumn<FTestColumnG>() }, false));
+}
+
+TEST_CASE_NAMED(FTypedElementQueryConditions7, "TypedElementQueryBuilder::FTypedElementQueryConditions (A && B) || (C && D) || (E && F)", "[ApplicationContextMask][EngineFilter]")
+{
+	FQueryConditions Example =
+		(FColumn<FTestColumnA>() && FColumn<FTestColumnB>()) ||
+		(FColumn<FTestColumnC>() && FColumn<FTestColumnD>()) ||
+		(FColumn<FTestColumnE>() && FColumn<FTestColumnF>());
+
+	CHECK(Example.MinimumColumnMatchRequired() == 2);
+
+	CHECK(TestMatching(Example, { FColumn<FTestColumnA>() }, false));
+	CHECK(TestMatching(Example, { FColumn<FTestColumnA>(), FColumn<FTestColumnB>() }, true));
+	CHECK(TestMatching(Example, { FColumn<FTestColumnC>(), FColumn<FTestColumnD>() }, true));
+	CHECK(TestMatching(Example, { FColumn<FTestColumnE>(), FColumn<FTestColumnF>() }, true));
+
+	CHECK(TestMatching(Example, { FColumn<FTestColumnC>(), FColumn<FTestColumnD>(), FColumn<FTestColumnE>(), FColumn<FTestColumnF>() }, true));
+
+
+	CHECK(TestMatching(Example, { FColumn<FTestColumnE>() }, false));
+	CHECK(TestMatching(Example, { FColumn<FTestColumnG>() }, false));
+}
+
+TEST_CASE_NAMED(FTypedElementQueryConditions_MultiMatch, "TypedElementQueryBuilder::FTypedElementQueryConditions_MultiMatch", "[ApplicationContextMask][EngineFilter]")
+{
+	FQueryConditions Example =
+		(FColumn<FTestColumnA>() || FColumn<FTestColumnB>()) &&
+		(FColumn<FTestColumnC>() || FColumn<FTestColumnD>()) &&
+		(FColumn<FTestColumnE>() || FColumn<FTestColumnF>());
+
+	CHECK(TestMatching(Example, 
+		{ 
+			FColumn<FTestColumnA>(), 
+			FColumn<FTestColumnB>(), 
+			FColumn<FTestColumnC>(), 
+			FColumn<FTestColumnD>(),
+			FColumn<FTestColumnE>(),
+			FColumn<FTestColumnF>()
+		}, true));
+
+	CHECK(TestMatching(Example,
+		{
+			FColumn<FTestColumnA>(),
+			FColumn<FTestColumnC>(),
+			FColumn<FTestColumnE>(),
+			FColumn<FTestColumnG>()
+		}, true));
+}
+
+TEST_CASE_NAMED(FTypedElementQueryConditions_Sorted, "TypedElementQueryBuilder::FTypedElementQueryConditions_Sorted", "[ApplicationContextMask][EngineFilter]")
+{
+	FQueryConditions Example =
+		(FColumn<FTestColumnA>() && FColumn<FTestColumnB>()) ||
+		(FColumn<FTestColumnC>() && FColumn<FTestColumnD>()) ||
+		(FColumn<FTestColumnE>() && FColumn<FTestColumnF>());
+
+	CHECK(Example.MinimumColumnMatchRequired() == 2);
+
+	CHECK(TestMatching(Example, { FColumn<FTestColumnA>(), FColumn<FTestColumnB>() }, true, true));
+	CHECK(TestMatching(Example, { FColumn<FTestColumnC>(), FColumn<FTestColumnD>() }, true, true));
+	CHECK(TestMatching(Example, { FColumn<FTestColumnE>(), FColumn<FTestColumnF>() }, true, true));
+
+	CHECK(TestMatching(Example, { FColumn<FTestColumnC>(), FColumn<FTestColumnD>(), FColumn<FTestColumnE>(), FColumn<FTestColumnF>() }, true, true));
+}
+
+#endif // WITH_TESTS
diff --git a/Engine/Source/Runtime/TypedElementFramework/Tests/Elements/Framework/TypedElementQueryBuilderTests.h b/Engine/Source/Runtime/TypedElementFramework/Tests/Elements/Framework/TypedElementQueryBuilderTests.h
new file mode 100644
--- /dev/null
+++ b/Engine/Source/Runtime/TypedElementFramework/Tests/Elements/Framework/TypedElementQueryBuilderTests.h
@@ -0,0 +1,50 @@
+// Copyright Epic Games, Inc. All Rights Reserved.
+
+#pragma once
+
+#include "Elements/Interfaces/TypedElementDataStorageInterface.h"
+#include "UObject/ObjectMacros.h"
+
+#include "TypedElementQueryBuilderTests.generated.h"
+
+USTRUCT(meta = (DisplayName = "A"))
+struct FTestColumnA final : public FTypedElementDataStorageColumn
+{
+	GENERATED_BODY()
+};
+
+USTRUCT(meta = (DisplayName = "B"))
+struct FTestColumnB final : public FTypedElementDataStorageColumn
+{
+	GENERATED_BODY()
+};
+
+USTRUCT(meta = (DisplayName = "C"))
+struct FTestColumnC final : public FTypedElementDataStorageColumn
+{
+	GENERATED_BODY()
+};
+
+USTRUCT(meta = (DisplayName = "D"))
+struct FTestColumnD final : public FTypedElementDataStorageColumn
+{
+	GENERATED_BODY()
+};
+
+USTRUCT(meta = (DisplayName = "E"))
+struct FTestColumnE final : public FTypedElementDataStorageColumn
+{
+	GENERATED_BODY()
+};
+
+USTRUCT(meta = (DisplayName = "F"))
+struct FTestColumnF final : public FTypedElementDataStorageColumn
+{
+	GENERATED_BODY()
+};
+
+USTRUCT(meta = (DisplayName = "G"))
+struct FTestColumnG final : public FTypedElementDataStorageColumn
+{
+	GENERATED_BODY()
+};
