Fixed Clang compilation error.

#nojira

diff --git a/Engine/Source/Runtime/TypedElementFramework/Tests/Elements/Framework/TypedElementQueryBuilderTests.cpp b/Engine/Source/Runtime/TypedElementFramework/Tests/Elements/Framework/TypedElementQueryBuilderTests.cpp
--- a/Engine/Source/Runtime/TypedElementFramework/Tests/Elements/Framework/TypedElementQueryBuilderTests.cpp
+++ b/Engine/Source/Runtime/TypedElementFramework/Tests/Elements/Framework/TypedElementQueryBuilderTests.cpp
@@ -41,6 +41,7 @@ static bool TestMatching(const FQueryConditions& TestQuery, const TArray<FColumn
 	FString Description = (Result == Expected) ? TEXT("[Pass] ") : TEXT("[Fail] ");
 	TestQuery.AppendToString(Description);
 
+#if WITH_EDITORONLY_DATA
 	{
 		Description += " -> { ";
 		auto It = RequestedColumns.begin();
@@ -53,22 +54,25 @@ static bool TestMatching(const FQueryConditions& TestQuery, const TArray<FColumn
 		}
 		Description += " } ";
 	}
+#endif
 
 	if (Expected)
 	{
+#if WITH_EDITORONLY_DATA
 		if (!Matches.IsEmpty())
 		{
 			Description += " -> { ";
 			auto It = Matches.begin();
 			AppendColumnName(Description, *It);
 			++It;
-			for (It; It != Matches.end(); ++It)
+			for (; It != Matches.end(); ++It)
 			{
 				Description += ", ";
 				AppendColumnName(Description, *It);
 			}
 			Description += " } ";
 		}
+#endif
 
 		for (TWeakObjectPtr<const UScriptStruct> Match : Matches)
 		{
