Platform Clang specific compile error.

#nojira

diff --git a/Engine/Source/Runtime/TypedElementFramework/Public/Elements/Framework/TypedElementQueryBuilder.h b/Engine/Source/Runtime/TypedElementFramework/Public/Elements/Framework/TypedElementQueryBuilder.h
--- a/Engine/Source/Runtime/TypedElementFramework/Public/Elements/Framework/TypedElementQueryBuilder.h
+++ b/Engine/Source/Runtime/TypedElementFramework/Public/Elements/Framework/TypedElementQueryBuilder.h
@@ -122,7 +122,6 @@ namespace TypedElementQueryBuilder
 	template<typename T = void>
 	struct FColumn final : public FColumnBase
 	{
-		template<typename = std::enable_if_t<!std::is_same_v<T, void>>>
 		constexpr FColumn() : FColumnBase(T::StaticStruct()) {}
 		constexpr explicit FColumn(TWeakObjectPtr<const UScriptStruct> ColumnTypeInfo) : FColumnBase(ColumnTypeInfo) {};
 	};
