Addressed Mac compile error.

#nojira

diff --git a/Engine/Source/Runtime/TypedElementFramework/Private/Elements/Framework/TypedElementQueryBuilder.cpp b/Engine/Source/Runtime/TypedElementFramework/Private/Elements/Framework/TypedElementQueryBuilder.cpp
--- a/Engine/Source/Runtime/TypedElementFramework/Private/Elements/Framework/TypedElementQueryBuilder.cpp
+++ b/Engine/Source/Runtime/TypedElementFramework/Private/Elements/Framework/TypedElementQueryBuilder.cpp
@@ -32,16 +32,6 @@ namespace TypedElementQueryBuilder
 	}
 
 
-	//
-	// FColumnBase
-	//
-
-	constexpr FColumnBase::FColumnBase(TWeakObjectPtr<const UScriptStruct> ColumnTypeInfo)
-		: TypeInfo(ColumnTypeInfo)
-	{
-	}
-
-
 
 	//
 	// FQueryConditions
diff --git a/Engine/Source/Runtime/TypedElementFramework/Public/Elements/Framework/TypedElementQueryBuilder.h b/Engine/Source/Runtime/TypedElementFramework/Public/Elements/Framework/TypedElementQueryBuilder.h
--- a/Engine/Source/Runtime/TypedElementFramework/Public/Elements/Framework/TypedElementQueryBuilder.h
+++ b/Engine/Source/Runtime/TypedElementFramework/Public/Elements/Framework/TypedElementQueryBuilder.h
@@ -116,7 +116,7 @@ namespace TypedElementQueryBuilder
 	
 	protected:
 		FColumnBase() = default;
-		TYPEDELEMENTFRAMEWORK_API constexpr explicit FColumnBase(TWeakObjectPtr<const UScriptStruct> ColumnTypeInfo);
+		constexpr explicit FColumnBase(TWeakObjectPtr<const UScriptStruct> ColumnTypeInfo) : TypeInfo(ColumnTypeInfo) {}
 	};
 
 	template<typename T = void>
