Fix non-unity
#rnx

diff --git a/Engine/Source/Runtime/TypedElementFramework/Public/Elements/Framework/TypedElementMetaData.h b/Engine/Source/Runtime/TypedElementFramework/Public/Elements/Framework/TypedElementMetaData.h
--- a/Engine/Source/Runtime/TypedElementFramework/Public/Elements/Framework/TypedElementMetaData.h
+++ b/Engine/Source/Runtime/TypedElementFramework/Public/Elements/Framework/TypedElementMetaData.h
@@ -9,6 +9,7 @@
 #include "Misc/EnumClassFlags.h"
 #include "Misc/TVariant.h"
 #include "Templates/Tuple.h"
+#include "UObject/WeakObjectPtr.h"
 
 class UScriptStruct;
 
