Fix Uninitialized Field test

diff --git a/Engine/Source/Runtime/TypedElementFramework/Public/Elements/Columns/TypedElementRevisionControlColumns.h b/Engine/Source/Runtime/TypedElementFramework/Public/Elements/Columns/TypedElementRevisionControlColumns.h
--- a/Engine/Source/Runtime/TypedElementFramework/Public/Elements/Columns/TypedElementRevisionControlColumns.h
+++ b/Engine/Source/Runtime/TypedElementFramework/Public/Elements/Columns/TypedElementRevisionControlColumns.h
@@ -44,7 +44,7 @@ struct FSCCStatusColumn final : public FTypedElementDataStorageColumn
 {
 	GENERATED_BODY()
 
-	UPROPERTY()
+	UPROPERTY(meta = (IgnoreForMemberInitializationTest))
 	ESCCModification Modification;
 };
 
@@ -53,7 +53,7 @@ struct FSCCRevisionId
 {
 	GENERATED_BODY()
 
-	UPROPERTY()
+	UPROPERTY(meta = (IgnoreForMemberInitializationTest))
 	uint32 Id[5];
 };
 
