Fixed missing conditions for FColumn.

diff --git a/Engine/Source/Runtime/TypedElementFramework/Public/Elements/Framework/TypedElementQueryBuilder.h b/Engine/Source/Runtime/TypedElementFramework/Public/Elements/Framework/TypedElementQueryBuilder.h
--- a/Engine/Source/Runtime/TypedElementFramework/Public/Elements/Framework/TypedElementQueryBuilder.h
+++ b/Engine/Source/Runtime/TypedElementFramework/Public/Elements/Framework/TypedElementQueryBuilder.h
@@ -122,7 +122,9 @@ namespace TypedElementQueryBuilder
 	template<typename T = void>
 	struct FColumn final : public FColumnBase
 	{
+		template <typename U = T> requires (!std::is_same_v<U, void>)
 		constexpr FColumn() : FColumnBase(T::StaticStruct()) {}
+		template <typename U = T> requires (std::is_same_v<U, void>)
 		constexpr explicit FColumn(TWeakObjectPtr<const UScriptStruct> ColumnTypeInfo) : FColumnBase(ColumnTypeInfo) {};
 	};
 
