TypedElementDataStorageInterface Test Suite
* Added 29 tests for some basic test coverage of the DataStorageInterface.
* Not all tests pass yet, specifically tests surrounding invocation of observers on add/remove columns

diff --git a/Engine/Source/Runtime/TypedElementFramework/Tests/Elements/Framework/TypedElementDataStorageInterfaceTests.cpp b/Engine/Source/Runtime/TypedElementFramework/Tests/Elements/Framework/TypedElementDataStorageInterfaceTests.cpp
new file mode 100644
--- /dev/null
+++ b/Engine/Source/Runtime/TypedElementFramework/Tests/Elements/Framework/TypedElementDataStorageInterfaceTests.cpp
@@ -0,0 +1,580 @@
+﻿// Copyright Epic Games, Inc. All Rights Reserved.
+
+
+#if WITH_TESTS
+#include "TypedElementTestColumns.h"
+#include "Elements/Framework/TypedElementRegistry.h"
+#include "Misc/AutomationTest.h"
+
+namespace TypedElementDataStorageTests
+{	
+	TypedElementTableHandle RegisterTestTableA(ITypedElementDataStorageInterface* TedsInterface)
+	{
+		// RegisterTable is not idempotent, so this utility function makes it so
+		
+		const FName TestTableAName = TEXT("TestTableA");
+		TypedElementTableHandle Table = TedsInterface->FindTable(TestTableAName);
+		if (Table != TypedElementInvalidTableHandle)
+		{
+			return Table;
+		}
+		
+		return TedsInterface->RegisterTable(
+		{
+			FTestColumnA::StaticStruct(),
+			FTestColumnB::StaticStruct(),
+			FTestTagColumnA::StaticStruct(),
+			FTestTagColumnB::StaticStruct()
+		},
+		TEXT("TestTableA"));
+	}
+
+	bool IsEmptyQueryDescription(const TypedElementDataStorage::FQueryDescription& QueryDescription)
+	{
+		return
+			QueryDescription.Callback.MonitoredType == nullptr &&
+			QueryDescription.SelectionTypes.IsEmpty() &&
+			QueryDescription.SelectionAccessTypes.IsEmpty() &&
+			QueryDescription.SelectionMetaData.IsEmpty() &&
+			QueryDescription.ConditionTypes.IsEmpty() &&
+			QueryDescription.ConditionOperators.IsEmpty() &&
+			QueryDescription.DependencyTypes.IsEmpty() &&
+			QueryDescription.DependencyFlags.IsEmpty() &&
+			QueryDescription.CachedDependencies.IsEmpty() &&
+			QueryDescription.Subqueries.IsEmpty();
+	}
+
+	void IncrementThing(int64* QueryCallCountPtr)
+	{
+		++(*QueryCallCountPtr);
+	}
+
+	// Temporary wrapper to only register a single observer
+	// There is no way to unregister an observer at this time
+	template<typename ColumnType>
+	void RegisterTestObserver(
+		FName ObserverName,
+		ITypedElementDataStorageInterface* TedsInterface, TypedElementQueryBuilder::FObserver::EEvent EventType, int64* QueryCallCountPtr, TypedElementQueryHandle* QueryHandleInOut)
+	{
+		using namespace TypedElementQueryBuilder;
+		using DSI = ITypedElementDataStorageInterface;
+		TypedElementDataStorage::FQueryDescription QueryDescription = TedsInterface->GetQueryDescription(*QueryHandleInOut);
+
+		// Check if it is empty... If not then we need to register
+		// This is for an idempotent API to the test harness				
+		if (IsEmptyQueryDescription(QueryDescription))
+		{
+			*QueryHandleInOut = TedsInterface->RegisterQuery(
+				Select(
+					ObserverName,
+					FObserver(EventType, ColumnType::StaticStruct()),
+					[QueryCallCountPtr](DSI::IQueryContext& Context, TypedElementRowHandle Row)
+					{
+						++(*QueryCallCountPtr);
+					}).Compile());
+		}
+	}
+	
+	BEGIN_DEFINE_SPEC(TypedElementDataStorageTestsFixture, "TypedElementsDataStorage", EAutomationTestFlags::EditorContext | EAutomationTestFlags::EngineFilter)
+		ITypedElementDataStorageInterface* TedsInterface = nullptr;
+		TypedElementTableHandle TestTableHandleA = TypedElementInvalidTableHandle;
+
+		TArray<TypedElementRowHandle> CreatedRows;
+
+		TypedElementQueryHandle DataColumnA_AddObserverHandle = TypedElementInvalidQueryHandle;
+		TypedElementQueryHandle DataColumnC_AddObserverHandle = TypedElementInvalidQueryHandle;
+		TypedElementQueryHandle DataColumnB_RemoveObserverHandle = TypedElementInvalidQueryHandle;
+		TypedElementQueryHandle TagColumnA_AddObserverHandle = TypedElementInvalidQueryHandle;
+		TypedElementQueryHandle TagColumnC_AddObserverHandle = TypedElementInvalidQueryHandle;
+		TypedElementQueryHandle TagColumnB_RemoveObserverHandle = TypedElementInvalidQueryHandle;
+	
+		// Non-zero to catch problems in test where the difference is not being used
+		// Avoids need to reset this to 0 for each test
+		int64 DataColumnA_AddObserverCallCount = 10;
+		int64 DataColumnC_AddObserverCallCount = 10;
+		int64 DataColumnB_RemoveObserverCallCount = 10;
+		int64 TagColumnA_AddObserverCallCount = 10;
+		int64 TagColumnC_AddObserverCallCount = 10;
+		int64 TagColumnB_RemoveObserverCallCount = 10;
+	
+		void CreateTestRows(int32 Count)
+		{
+			CreatedRows.Reserve(Count);
+			for (int32 Index = 0; Index < Count; ++Index)
+			{
+				TypedElementRowHandle RowHandle = TedsInterface->AddRow(TestTableHandleA);
+				TestFalse("Expect valid row handle", RowHandle == TypedElementInvalidRowHandle);
+				CreatedRows.Add(RowHandle);
+			}
+		}
+
+		void CleanupTestRows()
+		{
+			for (TypedElementRowHandle RowHandle : CreatedRows)
+			{
+				TestTrue("Expected row to have been available", TedsInterface->IsRowAvailable(RowHandle));
+				TestTrue("Expected row to have been assigned to a table", TedsInterface->HasRowBeenAssigned(RowHandle));
+				TedsInterface->RemoveRow(RowHandle);
+				TestFalse("Expected row to have been unassigned a table", TedsInterface->HasRowBeenAssigned(RowHandle));
+				TestFalse("Expected row to be not available", TedsInterface->IsRowAvailable(RowHandle));
+			}
+			CreatedRows.Empty();
+		}
+	END_DEFINE_SPEC(TypedElementDataStorageTestsFixture)
+
+	void TypedElementDataStorageTestsFixture::Define()
+	{
+		BeforeEach([this]()
+		{
+			UTypedElementRegistry* TypedElementRegistry = UTypedElementRegistry::GetInstance();
+			TedsInterface = TypedElementRegistry->GetMutableDataStorage();
+			TestTrue("", TedsInterface != nullptr);
+		});
+
+		Describe("", [this]()
+		{
+			BeforeEach([this]()
+			{
+				TestTrue("Test requires valid TEDS Interface", TedsInterface != nullptr);
+			});
+			
+			Describe("", [this]()
+			{
+				Describe("RegisterTable(TConstArrayView<const UScriptStruct*>, const FName)", [this]()
+				{
+					It("should register a table and provide a valid handle", [this]()
+					{
+						// Note, this test doesn't really do anything a second time.
+						// The test fixture depends on global state which currently cannot be cleaned
+						// up.
+						TypedElementTableHandle Handle = RegisterTestTableA(TedsInterface);
+						TestNotEqual("Expecting valid table handle", Handle, TypedElementInvalidTableHandle);
+					});
+				});
+			});
+
+			Describe("", [this]()
+			{
+				BeforeEach([this]()
+				{
+					TestTableHandleA = RegisterTestTableA(TedsInterface);
+				});
+				
+				AfterEach([this]()
+				{
+					CleanupTestRows();
+				});
+				
+				Describe("AddRow(TypedElementTableHandle)", [this]()
+				{
+					for (int32 RowCount = 1; RowCount <= 2; ++RowCount)
+					{
+						It(FString::Printf(TEXT("should should be able to create %d rows when called %d times"), RowCount, RowCount), [this, RowCount]()
+						{
+							for (int32 CreatedRowIndex = 0; CreatedRowIndex < RowCount; ++CreatedRowIndex)
+							{
+								TypedElementRowHandle RowHandle = TedsInterface->AddRow(TestTableHandleA);
+								CreatedRows.Add(RowHandle);
+								TestNotEqual("Expecting valid row", RowHandle, TypedElementInvalidRowHandle);
+							}
+						});
+					}
+				});
+
+				Describe("BatchAddRow(TypedElementTableHandle, int32, TypedElementDataStorageCreationCallbackRef)", [this]()
+				{
+					for (int32 RowCount = 1; RowCount <= 2; ++RowCount)
+					{
+						It(FString::Printf(TEXT("should should be able to create %d rows when called once"), RowCount), [this, RowCount]()
+						{
+							auto RowInitializeCallback = [this](TypedElementRowHandle RowHandle)
+							{
+								CreatedRows.Add(RowHandle);
+								TestNotEqual("Expecting valid row", RowHandle, TypedElementInvalidRowHandle);
+							};
+											
+							TypedElementRowHandle RowHandle = TedsInterface->BatchAddRow(TestTableHandleA, RowCount, RowInitializeCallback);
+											
+							TestEqual("Unexpected number of created rows", CreatedRows.Num(), RowCount);
+						});
+					}
+				});
+
+			
+				Describe("", [this]()
+				{
+					/**
+					 * Prerequisites for this batch of tests are:
+					 * - A test table is registered
+					 * - An add and remove observer is registered for FTestInt64Column
+					 */
+					BeforeEach([this]()
+					{
+						TestTableHandleA = RegisterTestTableA(TedsInterface);
+						TestNotEqual("Expecting valid table handle", TestTableHandleA, TypedElementInvalidTableHandle);
+
+						RegisterTestObserver<FTestColumnA>(
+							TEXT("Increment CallCount when FTestColumnA added"),
+							TedsInterface,
+							TypedElementQueryBuilder::FObserver::EEvent::Add,
+							&DataColumnA_AddObserverCallCount,
+							&DataColumnA_AddObserverHandle);
+						TestTrue("Expect valid query observer handle", DataColumnA_AddObserverHandle != TypedElementInvalidQueryHandle);
+						
+						RegisterTestObserver<FTestColumnC>(
+							TEXT("Increment CallCount when FTestColumnC added"),
+							TedsInterface,
+							TypedElementQueryBuilder::FObserver::EEvent::Add,
+							&DataColumnC_AddObserverCallCount,
+							&DataColumnC_AddObserverHandle);
+						TestTrue("Expect valid query observer handle", DataColumnC_AddObserverHandle != TypedElementInvalidQueryHandle);
+
+						RegisterTestObserver<FTestColumnB>(
+							TEXT("Increment CallCount when FTestColumnB removed"),
+							TedsInterface,
+							TypedElementQueryBuilder::FObserver::EEvent::Remove,
+							&DataColumnB_RemoveObserverCallCount,
+							&DataColumnB_RemoveObserverHandle);
+						TestTrue("Expect valid query observer handle", DataColumnB_RemoveObserverHandle != TypedElementInvalidQueryHandle);
+
+						RegisterTestObserver<FTestTagColumnA>(
+							TEXT("Increment CallCount when FTestColumnA added"),
+							TedsInterface,
+							TypedElementQueryBuilder::FObserver::EEvent::Add,
+							&TagColumnA_AddObserverCallCount,
+							&TagColumnA_AddObserverHandle);
+						TestTrue("Expect valid query observer handle", TagColumnA_AddObserverHandle != TypedElementInvalidQueryHandle);
+												
+						RegisterTestObserver<FTestTagColumnC>(
+							TEXT("Increment CallCount when FTestTagColumnC added"),
+							TedsInterface,
+							TypedElementQueryBuilder::FObserver::EEvent::Add,
+							&TagColumnC_AddObserverCallCount, 
+							&TagColumnC_AddObserverHandle);
+						TestTrue("Expect valid query observer handle", TagColumnC_AddObserverHandle != TypedElementInvalidQueryHandle);
+
+						RegisterTestObserver<FTestTagColumnB>(
+							TEXT("Increment CallCount when FTestTagColumnB removed"),
+							TedsInterface, TypedElementQueryBuilder::FObserver::EEvent::Remove,
+							&TagColumnB_RemoveObserverCallCount,
+							&TagColumnB_RemoveObserverHandle);
+						TestTrue("Expect valid query observer handle", TagColumnB_RemoveObserverHandle != TypedElementInvalidQueryHandle);					
+					});
+					
+					AfterEach([this]()
+					{
+						CleanupTestRows();
+					});
+
+					Describe("AddRow(TypedElementTableHandle)", [this]()
+					{
+						for (int32 RowCount = 1; RowCount <= 2; ++RowCount)
+						{
+							It(FString::Printf(TEXT("should invoke Add observer %d times when called with table containing data column %d times"), RowCount, RowCount), [this, RowCount]()
+							{
+								int64 PreviousQueryCallCount = DataColumnA_AddObserverCallCount;
+								
+								for (int32 CreatedRowIndex = 0; CreatedRowIndex < RowCount; ++CreatedRowIndex)
+								{
+									TypedElementRowHandle RowHandle = TedsInterface->AddRow(TestTableHandleA);
+									CreatedRows.Add(RowHandle);
+								}
+
+								TestEqual(TEXT("Expect observer to be called correct number of times"), DataColumnA_AddObserverCallCount, PreviousQueryCallCount + RowCount);
+							});
+						}
+					});
+
+					Describe("BatchAddRow(TypedElementTableHandle, int32, TypedElementDataStorageCreationCallbackRef)", [this]()
+					{
+						for (int32 RowCount = 1; RowCount <= 2; ++RowCount)
+						{
+							It(FString::Printf(TEXT("should invoke Add observer %d times when called with table containing data column"), RowCount), [this, RowCount]()
+							{
+								int64 PreviousQueryCallCount = DataColumnA_AddObserverCallCount;
+								
+								auto RowInitializeCallback = [this](TypedElementRowHandle RowHandle)
+								{
+									CreatedRows.Add(RowHandle);
+								};
+																
+								TypedElementRowHandle RowHandle = TedsInterface->BatchAddRow(TestTableHandleA, RowCount, RowInitializeCallback);
+								
+								TestEqual(TEXT("Expect observer to be called correct number of times"), DataColumnA_AddObserverCallCount, PreviousQueryCallCount + RowCount);
+							});
+						}
+					});
+
+					Describe("", [this]()
+					{
+						BeforeEach([this]()
+						{
+							CreateTestRows(1);
+						});
+
+						Describe("RemoveRow(TypedElementTableHandle)", [this]()
+						{
+							It("should remove row when called", [this]()
+							{
+								TypedElementRowHandle Row = CreatedRows[0];
+								TestTrue("Row is available", TedsInterface->IsRowAvailable(Row));
+
+								TedsInterface->RemoveRow(Row);
+								
+
+								TestFalse("Row is not available", TedsInterface->IsRowAvailable(Row));
+
+								CreatedRows.RemoveAt(0);
+							});
+							
+							It("should call observer of removed data column when called on row that has the column", [this]()
+							{
+								int64 PreviousQueryCallCount = DataColumnB_RemoveObserverCallCount;
+								
+								TypedElementRowHandle Row = CreatedRows[0];
+								TedsInterface->RemoveRow(Row);
+
+								TestEqual("Expected observer called", DataColumnB_RemoveObserverCallCount, PreviousQueryCallCount + 1);
+
+								CreatedRows.RemoveAt(0);
+							});
+							
+							It("should call observer of removed tag column when called on row that has the column", [this]()
+							{
+								int64 PreviousQueryCallCount = TagColumnB_RemoveObserverCallCount;
+								
+								TypedElementRowHandle Row = CreatedRows[0];
+								TedsInterface->RemoveRow(Row);
+								
+								TestEqual("Expected observer called", TagColumnB_RemoveObserverCallCount, PreviousQueryCallCount + 1);
+
+								CreatedRows.RemoveAt(0);
+							});
+						});
+						
+						Describe("AddColumn(TypedElementRowHandle, const UScriptStruct*)", [this]()
+						{					
+							It("should add data column when called", [this]()
+							{
+								TypedElementRowHandle Row = CreatedRows[0];
+								TestFalse("Expect row to not have column about to be added",
+									TedsInterface->HasColumns(Row, MakeArrayView({FTestColumnC::StaticStruct()})));
+
+								TedsInterface->AddColumn(Row, FTestColumnC::StaticStruct());
+
+								TestTrue("Expect row to have column added",
+									TedsInterface->HasColumns(Row, MakeArrayView({FTestColumnC::StaticStruct()})));
+
+							});
+
+							It("should invoke registered observer when called with data column", [this]()
+							{
+								int64 PreviousQueryCallCount = DataColumnC_AddObserverCallCount;
+								TypedElementRowHandle Row = CreatedRows[0];
+								TedsInterface->AddColumn(Row, FTestColumnC::StaticStruct());
+
+								TestEqual("Expect observer to have been called one time", DataColumnC_AddObserverCallCount, PreviousQueryCallCount + 1);
+							});
+
+							It("should add tag column when called", [this]()
+							{
+								TypedElementRowHandle Row = CreatedRows[0];
+								TestFalse("Expect row to not have column about to be added",
+									TedsInterface->HasColumns(Row, MakeArrayView({FTestTagColumnC::StaticStruct()})));
+
+								TedsInterface->AddColumn(Row, FTestTagColumnC::StaticStruct());
+
+								TestTrue("Expect row to have column added",
+									TedsInterface->HasColumns(Row, MakeArrayView({FTestTagColumnC::StaticStruct()})));
+
+							});
+
+							It("should invoke registered observer when called with tag column", [this]()
+							{
+								int64 PreviousQueryCallCount = TagColumnC_AddObserverCallCount;
+								TypedElementRowHandle Row = CreatedRows[0];
+								TedsInterface->AddColumn(Row, FTestTagColumnC::StaticStruct());
+
+								TestEqual("Expect observer to have been called one time", TagColumnC_AddObserverCallCount, PreviousQueryCallCount + 1);
+							});
+						});
+						
+						Describe("RemoveColumn(TypedElementRowHandle, TConstArrayView<const UScriptStruct*>)", [this]()
+						{
+							It("should remove a single column when called", [this]()
+							{
+								TypedElementRowHandle Row = CreatedRows[0];
+								TestTrue("Expected to have column about to be removed",
+									TedsInterface->HasColumns(Row, MakeArrayView({FTestColumnB::StaticStruct()})));
+								
+								TedsInterface->RemoveColumn(Row, FTestColumnB::StaticStruct());
+
+								TestFalse("Expected to no longer have removed column",
+									TedsInterface->HasColumns(Row, MakeArrayView({FTestColumnB::StaticStruct()})));
+							});
+
+							It("should invoke registered data column observer when called", [this]()
+							{
+								int64 PreviousQueryCallCount = DataColumnB_RemoveObserverCallCount;
+								TypedElementRowHandle Row = CreatedRows[0];
+								
+								TedsInterface->RemoveColumn(Row, FTestColumnB::StaticStruct());
+								
+								TestEqual("Expect observer to have been called one time", DataColumnB_RemoveObserverCallCount, PreviousQueryCallCount + 1);
+							});
+
+							It("should remove a single tag column when called", [this]()
+							{
+								TypedElementRowHandle Row = CreatedRows[0];
+								TestTrue("Expected to have column about to be removed",
+									TedsInterface->HasColumns(Row, MakeArrayView({FTestTagColumnB::StaticStruct()})));
+														
+								TedsInterface->RemoveColumn(Row, FTestTagColumnB::StaticStruct());
+
+								TestFalse("Expected to no longer have removed column",
+									TedsInterface->HasColumns(Row, MakeArrayView({FTestTagColumnB::StaticStruct()})));
+							});
+
+							It("should invoke registered tag column observer when called", [this]()
+							{
+								int64 PreviousQueryCallCount = TagColumnB_RemoveObserverCallCount;
+								TypedElementRowHandle Row = CreatedRows[0];
+								
+								TedsInterface->RemoveColumn(Row, FTestTagColumnB::StaticStruct());
+								
+								TestEqual("Expect observer to have been called one time", TagColumnB_RemoveObserverCallCount, PreviousQueryCallCount + 1);
+							});
+						});
+
+						Describe("AddColumns(TypedElementRowHandle, TConstArrayView<const UScriptStruct*>)", [this]()
+						{
+							It("should add a single column when called with single data column", [this]()
+							{
+								TArray<const UScriptStruct*> Columns = {FTestColumnC::StaticStruct(), FTestColumnD::StaticStruct()};
+								
+								TypedElementRowHandle Row = CreatedRows[0];
+								TestFalse("Expected to not have column about to be added",
+									TedsInterface->HasColumns(Row, Columns));
+
+								TedsInterface->AddColumns(Row, Columns);
+
+								TestTrue("Expected to have column about to be added", 
+								TedsInterface->HasColumns(Row, Columns));
+							});
+
+							It("should add multiple column when called with multiple data columns", [this]()
+							{
+								TArray<const UScriptStruct*> Columns = {FTestColumnC::StaticStruct(), FTestColumnD::StaticStruct()};
+								
+								TypedElementRowHandle Row = CreatedRows[0];
+								TestFalse("Expected to not have column about to be added", 
+								TedsInterface->HasColumns(Row, Columns));
+
+								TedsInterface->AddColumns(Row, Columns);
+
+								TestTrue("Expected to have column about to be added",
+									TedsInterface->HasColumns(Row, Columns));
+							});
+
+							It("should invoke registered data column observer when called", [this]()
+							{
+								TArray<const UScriptStruct*> Columns = {FTestColumnC::StaticStruct()};
+								
+								int64 PreviousQueryCallCount = DataColumnC_AddObserverCallCount;
+								TypedElementRowHandle Row = CreatedRows[0];
+								
+								TedsInterface->AddColumns(Row, Columns);
+
+								TestEqual("Expect observer to have been called one time", DataColumnC_AddObserverCallCount, PreviousQueryCallCount + 1);
+							});
+
+							It("should add a single column when called with single tag column", [this]()
+							{
+								TArray<const UScriptStruct*> Columns = {FTestTagColumnC::StaticStruct(), FTestTagColumnD::StaticStruct()};
+														
+								TypedElementRowHandle Row = CreatedRows[0];
+								TestFalse("Expected to not have column about to be added",
+									TedsInterface->HasColumns(Row, Columns));
+
+								TedsInterface->AddColumns(Row, Columns);
+
+								TestTrue("Expected to have column about to be added", 
+								TedsInterface->HasColumns(Row, Columns));
+							});
+
+							It("should add multiple column when called with multiple tag columns", [this]()
+							{
+								TArray<const UScriptStruct*> Columns = {FTestTagColumnC::StaticStruct(), FTestTagColumnD::StaticStruct()};
+														
+								TypedElementRowHandle Row = CreatedRows[0];
+								TestFalse("Expected to not have column about to be added", 
+								TedsInterface->HasColumns(Row, Columns));
+
+								TedsInterface->AddColumns(Row, Columns);
+
+								TestTrue("Expected to have column about to be added",
+									TedsInterface->HasColumns(Row, Columns));
+							});
+
+							It("should invoke registered tag column added observer when called", [this]()
+							{
+								TArray<const UScriptStruct*> Columns = {FTestTagColumnC::StaticStruct()};
+														
+								int64 PreviousQueryCallCount = TagColumnC_AddObserverCallCount;
+								TypedElementRowHandle Row = CreatedRows[0];
+								TedsInterface->AddColumns(Row, Columns);
+
+								TestEqual("Expect observer to have been called one time", TagColumnC_AddObserverCallCount, PreviousQueryCallCount + 1);
+							});
+						});
+
+						Describe("RemoveColumns(TypedElementRowHandle, TConstArrayView<const UScriptStruct*>)", [this]()
+						{
+							It("should remove a single column when RemoveColumns called with single column", [this]()
+							{
+								TArray<const UScriptStruct*> Columns = {FTestColumnB::StaticStruct()};
+								
+								TypedElementRowHandle Row = CreatedRows[0];
+								TestTrue("Expected to not column about to be removed",
+									TedsInterface->HasColumns(Row, Columns));
+
+								TedsInterface->RemoveColumns(Row, Columns);
+
+								TestFalse("Expected to not have column that was removed", 
+								TedsInterface->HasColumns(Row, MakeArrayView({FTestColumnB::StaticStruct()})));
+							});
+
+							It("should remove all columns when RemoveColumns called with multiple columns", [this]()
+							{
+								TArray<const UScriptStruct*> Columns = {FTestColumnA::StaticStruct(), FTestColumnB::StaticStruct()};
+								
+								TypedElementRowHandle Row = CreatedRows[0];
+								TestTrue("Expected to have columns about to be removed", 
+								TedsInterface->HasColumns(Row, Columns));
+
+								TedsInterface->RemoveColumns(Row, Columns);
+
+								TestFalse("Expected to have column about to be added",
+									TedsInterface->HasColumns(Row, Columns));
+							});
+
+							It("should invoke registered observer when RemoveColumns called", [this]()
+							{
+								TArray<const UScriptStruct*> Columns = {FTestColumnB::StaticStruct()};
+								
+								int64 PreviousQueryCallCount = DataColumnB_RemoveObserverCallCount;
+								TypedElementRowHandle Row = CreatedRows[0];
+								TedsInterface->RemoveColumns(Row, Columns);
+
+								TestEqual("Expect observer to have been called one time", DataColumnB_RemoveObserverCallCount, PreviousQueryCallCount + 1);
+							});
+						});
+					});
+				});
+				
+			});
+		});
+	}
+}
+
+#endif
\ No newline at end of file
diff --git a/Engine/Source/Runtime/TypedElementFramework/Tests/Elements/Framework/TypedElementQueryBuilderTests.cpp b/Engine/Source/Runtime/TypedElementFramework/Tests/Elements/Framework/TypedElementQueryBuilderTests.cpp
--- a/Engine/Source/Runtime/TypedElementFramework/Tests/Elements/Framework/TypedElementQueryBuilderTests.cpp
+++ b/Engine/Source/Runtime/TypedElementFramework/Tests/Elements/Framework/TypedElementQueryBuilderTests.cpp
@@ -1,7 +1,7 @@
 // Copyright Epic Games, Inc. All Rights Reserved.
 
 #if WITH_TESTS
-#include "TypedElementQueryBuilderTests.h"
+#include "TypedElementTestColumns.h"
 
 #include "Algo/Sort.h"
 #include "CoreTypes.h"
diff --git a/Engine/Source/Runtime/TypedElementFramework/Tests/Elements/Framework/TypedElementQueryBuilderTests.h b/Engine/Source/Runtime/TypedElementFramework/Tests/Elements/Framework/TypedElementQueryBuilderTests.h
deleted file mode 100644
diff --git a/Engine/Source/Runtime/TypedElementFramework/Tests/Elements/Framework/TypedElementTestColumns.h b/Engine/Source/Runtime/TypedElementFramework/Tests/Elements/Framework/TypedElementTestColumns.h
new file mode 100644
--- /dev/null
+++ b/Engine/Source/Runtime/TypedElementFramework/Tests/Elements/Framework/TypedElementTestColumns.h
@@ -0,0 +1,74 @@
+// Copyright Epic Games, Inc. All Rights Reserved.
+
+#pragma once
+
+#include "Elements/Interfaces/TypedElementDataStorageInterface.h"
+#include "UObject/ObjectMacros.h"
+
+#include "TypedElementTestColumns.generated.h"
+
+USTRUCT(meta = (DisplayName = "ColumnA"))
+struct FTestColumnA final : public FTypedElementDataStorageColumn
+{
+	GENERATED_BODY()
+};
+
+USTRUCT(meta = (DisplayName = "ColumnB"))
+struct FTestColumnB final : public FTypedElementDataStorageColumn
+{
+	GENERATED_BODY()
+};
+
+USTRUCT(meta = (DisplayName = "ColumnC"))
+struct FTestColumnC final : public FTypedElementDataStorageColumn
+{
+	GENERATED_BODY()
+};
+
+USTRUCT(meta = (DisplayName = "ColumnD"))
+struct FTestColumnD final : public FTypedElementDataStorageColumn
+{
+	GENERATED_BODY()
+};
+
+USTRUCT(meta = (DisplayName = "ColumnE"))
+struct FTestColumnE final : public FTypedElementDataStorageColumn
+{
+	GENERATED_BODY()
+};
+
+USTRUCT(meta = (DisplayName = "ColumnF"))
+struct FTestColumnF final : public FTypedElementDataStorageColumn
+{
+	GENERATED_BODY()
+};
+
+USTRUCT(meta = (DisplayName = "ColumnG"))
+struct FTestColumnG final : public FTypedElementDataStorageColumn
+{
+	GENERATED_BODY()
+};
+
+USTRUCT(meta = (DisplayName = "TagA"))
+struct FTestTagColumnA final : public FTypedElementDataStorageTag
+{
+	GENERATED_BODY()
+};
+
+USTRUCT(meta = (DisplayName = "TagB"))
+struct FTestTagColumnB final : public FTypedElementDataStorageTag
+{
+	GENERATED_BODY()
+};
+
+USTRUCT(meta = (DisplayName = "TagC"))
+struct FTestTagColumnC final : public FTypedElementDataStorageTag
+{
+	GENERATED_BODY()
+};
+
+USTRUCT(meta = (DisplayName = "TagD"))
+struct FTestTagColumnD final : public FTypedElementDataStorageTag
+{
+	GENERATED_BODY()
+};
\ No newline at end of file
