Fix uninitialized scrpit struct member warnings that came out of the woodwork with async GFPs.

diff --git a/Engine/Plugins/Experimental/Chooser/Source/Chooser/Internal/AnimNode_ChooserPlayer.h b/Engine/Plugins/Experimental/Chooser/Source/Chooser/Internal/AnimNode_ChooserPlayer.h
--- a/Engine/Plugins/Experimental/Chooser/Source/Chooser/Internal/AnimNode_ChooserPlayer.h
+++ b/Engine/Plugins/Experimental/Chooser/Source/Chooser/Internal/AnimNode_ChooserPlayer.h
@@ -30,7 +30,7 @@ struct FAnimCurveOverride
 	FName CurveName;
 	// Value to set to the curve
 	UPROPERTY(EditAnywhere, BlueprintReadWrite, Category = "CurveValue")
-	float CurveValue;
+	float CurveValue = 0.0f;
 };
 
 template <>
diff --git a/Engine/Source/Runtime/TypedElementFramework/Public/Elements/Columns/TypedElementRevisionControlColumns.h b/Engine/Source/Runtime/TypedElementFramework/Public/Elements/Columns/TypedElementRevisionControlColumns.h
--- a/Engine/Source/Runtime/TypedElementFramework/Public/Elements/Columns/TypedElementRevisionControlColumns.h
+++ b/Engine/Source/Runtime/TypedElementFramework/Public/Elements/Columns/TypedElementRevisionControlColumns.h
@@ -1,4 +1,4 @@
-﻿// Copyright Epic Games, Inc. All Rights Reserved.
+// Copyright Epic Games, Inc. All Rights Reserved.
 
 #pragma once
 
@@ -45,7 +45,7 @@ struct FSCCStatusColumn final : public FTypedElementDataStorageColumn
 	GENERATED_BODY()
 
 	UPROPERTY(meta = (IgnoreForMemberInitializationTest))
-	ESCCModification Modification;
+	ESCCModification Modification = ESCCModification::Modified;
 };
 
 USTRUCT()
@@ -54,7 +54,7 @@ struct FSCCRevisionId
 	GENERATED_BODY()
 
 	UPROPERTY(meta = (IgnoreForMemberInitializationTest))
-	uint32 Id[5];
+	uint32 Id[5] = {0};
 };
 
 USTRUCT(meta = (DisplayName = "Revision ID"))
