Extended the index hash to support a few additional types.

This change also removes the specialization for UObject as the GetUniqueID() call is not unique enough for the purposes of TEDS.

diff --git a/Engine/Source/Runtime/TypedElementFramework/Public/Elements/Framework/TypedElementIndexHasher.h b/Engine/Source/Runtime/TypedElementFramework/Public/Elements/Framework/TypedElementIndexHasher.h
--- a/Engine/Source/Runtime/TypedElementFramework/Public/Elements/Framework/TypedElementIndexHasher.h
+++ b/Engine/Source/Runtime/TypedElementFramework/Public/Elements/Framework/TypedElementIndexHasher.h
@@ -2,14 +2,27 @@
 
 #pragma once
 
+#include "Containers/UnrealString.h"
 #include "Elements/Interfaces/TypedElementQueryStorageInterfaces.h"
+#include "UObject/NameTypes.h"
+#include "UObject/ObjectPtr.h"
+#include "UObject/StrongObjectPtr.h"
+#include "UObject/WeakObjectPtrTemplates.h"
 
 namespace TypedElementDataStorage
 {
 	template<typename T>
 	IndexHash GenerateIndexHash(const T* Object);
+	
 	template<typename T>
-	IndexHash GenerateIndexHash(T* Object);
+	IndexHash GenerateIndexHash(const TWeakObjectPtr<T>& Object);
+	template<typename T>
+	IndexHash GenerateIndexHash(const TObjectPtr<T>& Object);
+	template<typename T>
+	IndexHash GenerateIndexHash(const TStrongObjectPtr<T>& Object);
+
+	inline IndexHash GenerateIndexHash(const FString& Object);
+	inline IndexHash GenerateIndexHash(FName Object);
 } // namespace TypedElementDataStorage
 
 #include "Elements/Framework/TypedElementIndexHasher.inl"
\ No newline at end of file
diff --git a/Engine/Source/Runtime/TypedElementFramework/Public/Elements/Framework/TypedElementIndexHasher.inl b/Engine/Source/Runtime/TypedElementFramework/Public/Elements/Framework/TypedElementIndexHasher.inl
--- a/Engine/Source/Runtime/TypedElementFramework/Public/Elements/Framework/TypedElementIndexHasher.inl
+++ b/Engine/Source/Runtime/TypedElementFramework/Public/Elements/Framework/TypedElementIndexHasher.inl
@@ -2,27 +2,43 @@
 
 #pragma once
 
-#include <type_traits>
-#include "UObject/Object.h"
+#include "Hash/CityHash.h"
 
 namespace TypedElementDataStorage
 {
 	template<typename T>
 	IndexHash GenerateIndexHash(const T* Object)
 	{
-		if constexpr (std::is_base_of_v<UObject, T>)
-		{
-			return Object->GetUniqueID();
-		}
-		else
-		{
-			return reinterpret_cast<IndexHash>(Object);
-		}
+		return reinterpret_cast<IndexHash>(Object);
 	}
 
 	template<typename T>
-	IndexHash GenerateIndexHash(T* Object)
+	IndexHash GenerateIndexHash(const TWeakObjectPtr<T>& Object)
 	{
-		return GenerateIndexHash(const_cast<const T*>(Object));
+		return GenerateIndexHash(Object.Get());
+	}
+
+	template<typename T>
+	IndexHash GenerateIndexHash(const TObjectPtr<T>& Object)
+	{
+		return GenerateIndexHash(Object.Get());
+	}
+	
+	template<typename T>
+	IndexHash GenerateIndexHash(const TStrongObjectPtr<T>& Object)
+	{
+		return GenerateIndexHash(Object.Get());
+	}
+
+	IndexHash GenerateIndexHash(const FString& Object)
+	{
+		return CityHash64(reinterpret_cast<const char*>(*Object), Object.Len() * sizeof(**Object));
+	}
+
+	IndexHash GenerateIndexHash(FName Object)
+	{
+		constexpr static const char SeedName[] = "FName";
+		static uint64 Seed = CityHash64(SeedName, sizeof(SeedName) - 1);
+		return CityHash128to64({ Seed, Object.GetComparisonIndex().ToUnstableInt() });
 	}
 } // namespace TypedElementDataStorage
