Add const interface to GetColumn from TEDS

diff --git a/Engine/Plugins/Experimental/TypedElementsDataStorage/Source/TypedElementsDataStorage/Private/TypedElementDatabase.cpp b/Engine/Plugins/Experimental/TypedElementsDataStorage/Source/TypedElementsDataStorage/Private/TypedElementDatabase.cpp
--- a/Engine/Plugins/Experimental/TypedElementsDataStorage/Source/TypedElementsDataStorage/Private/TypedElementDatabase.cpp
+++ b/Engine/Plugins/Experimental/TypedElementsDataStorage/Source/TypedElementsDataStorage/Private/TypedElementDatabase.cpp
@@ -373,7 +373,7 @@ ColumnDataResult UTypedElementDatabase::AddOrGetColumnData(TypedElementRowHandle
 	}
 }
 
-void* UTypedElementDatabase::GetColumnData(TypedElementRowHandle Row, const UScriptStruct* ColumnType)
+const void* UTypedElementDatabase::GetColumnData(TypedElementRowHandle Row, const UScriptStruct* ColumnType) const
 {
 	FMassEntityHandle Entity = FMassEntityHandle::FromNumber(Row);
 	if (ActiveEditorEntityManager && ActiveEditorEntityManager->IsEntityActive(Entity) &&
@@ -388,6 +388,11 @@ void* UTypedElementDatabase::GetColumnData(TypedElementRowHandle Row, const UScr
 	return nullptr;
 }
 
+void* UTypedElementDatabase::GetColumnData(TypedElementRowHandle Row, const UScriptStruct* ColumnType)
+{
+	return const_cast<void*>(static_cast<const UTypedElementDatabase*>(this)->GetColumnData(Row, ColumnType));
+}
+
 ColumnDataResult UTypedElementDatabase::GetColumnData(TypedElementRowHandle Row, FTopLevelAssetPath ColumnName)
 {
 	FMassEntityHandle Entity = FMassEntityHandle::FromNumber(Row);
diff --git a/Engine/Plugins/Experimental/TypedElementsDataStorage/Source/TypedElementsDataStorage/Private/TypedElementDatabase.h b/Engine/Plugins/Experimental/TypedElementsDataStorage/Source/TypedElementsDataStorage/Private/TypedElementDatabase.h
--- a/Engine/Plugins/Experimental/TypedElementsDataStorage/Source/TypedElementsDataStorage/Private/TypedElementDatabase.h
+++ b/Engine/Plugins/Experimental/TypedElementsDataStorage/Source/TypedElementsDataStorage/Private/TypedElementDatabase.h
@@ -63,6 +63,7 @@ public:
 	void* AddOrGetColumnData(TypedElementRowHandle Row, const UScriptStruct* ColumnType) override;
 	ColumnDataResult AddOrGetColumnData(TypedElementRowHandle Row, FTopLevelAssetPath ColumnName) override;
 	void* GetColumnData(TypedElementRowHandle Row, const UScriptStruct* ColumnType) override;
+	const void* GetColumnData(TypedElementRowHandle Row, const UScriptStruct* ColumnType) const override;
 	ColumnDataResult AddOrGetColumnData(TypedElementRowHandle Row, FTopLevelAssetPath ColumnName,
 		TConstArrayView<TypedElement::ColumnUtils::Argument> Arguments) override;
 	bool AddColumns(TypedElementRowHandle Row, TConstArrayView<const UScriptStruct*> Columns) override;
diff --git a/Engine/Source/Runtime/TypedElementFramework/Public/Elements/Interfaces/TypedElementDataStorageInterface.h b/Engine/Source/Runtime/TypedElementFramework/Public/Elements/Interfaces/TypedElementDataStorageInterface.h
--- a/Engine/Source/Runtime/TypedElementFramework/Public/Elements/Interfaces/TypedElementDataStorageInterface.h
+++ b/Engine/Source/Runtime/TypedElementFramework/Public/Elements/Interfaces/TypedElementDataStorageInterface.h
@@ -197,6 +197,7 @@ public:
 	
 	/** Retrieves a pointer to the column of the given row or a nullptr if not found or if the column type is a tag. */
 	virtual void* GetColumnData(TypedElementRowHandle Row, const UScriptStruct* ColumnType) = 0;
+	virtual const void* GetColumnData(TypedElementRowHandle Row, const UScriptStruct* ColumnType) const = 0;
 	virtual ColumnDataResult GetColumnData(TypedElementRowHandle Row, FTopLevelAssetPath ColumnName) = 0;
 
 	/** Determines if the provided row contains the collection of columns and tags. */
@@ -348,6 +349,10 @@ public:
 	template<typename ColumnType>
 	ColumnType* GetColumn(TypedElementRowHandle Row);
 
+	/** Returns a pointer to the column of the given row or a nullptr if the type couldn't be found or the row doesn't exist. */
+	template<typename ColumnType>
+	const ColumnType* GetColumn(TypedElementRowHandle Row) const;
+
 	template<typename... ColumnTypes>
 	bool HasColumns(TypedElementRowHandle Row) const;
 
@@ -432,6 +437,12 @@ ColumnType* ITypedElementDataStorageInterface::GetColumn(TypedElementRowHandle R
 	return reinterpret_cast<ColumnType*>(GetColumnData(Row, ColumnType::StaticStruct()));
 }
 
+template<typename ColumnType>
+const ColumnType* ITypedElementDataStorageInterface::GetColumn(TypedElementRowHandle Row) const
+{
+	return reinterpret_cast<const ColumnType*>(GetColumnData(Row, ColumnType::StaticStruct()));
+}
+
 template<typename... ColumnType>
 bool ITypedElementDataStorageInterface::HasColumns(TypedElementRowHandle Row) const
 {
