Add TEDS PivotOffset Column definition

diff --git a/Engine/Source/Runtime/TypedElementFramework/Public/Elements/Columns/TypedElementPivotOffsetColumns.h b/Engine/Source/Runtime/TypedElementFramework/Public/Elements/Columns/TypedElementPivotOffsetColumns.h
new file mode 100644
--- /dev/null
+++ b/Engine/Source/Runtime/TypedElementFramework/Public/Elements/Columns/TypedElementPivotOffsetColumns.h
@@ -0,0 +1,16 @@
+﻿// Copyright Epic Games, Inc. All Rights Reserved.
+
+#pragma once
+
+#include "Elements/Interfaces/TypedElementDataStorageInterface.h"
+
+#include "TypedElementPivotOffsetColumns.generated.h"
+
+USTRUCT()
+struct FTypedElementPivotOffset final : public FTypedElementDataStorageColumn
+{
+	GENERATED_BODY()
+
+	UPROPERTY()
+	FVector Offset;
+};
\ No newline at end of file
