Allow FTypedElementPivotOffset to have uninitialized member

diff --git a/Engine/Source/Runtime/TypedElementFramework/Public/Elements/Columns/TypedElementPivotOffsetColumns.h b/Engine/Source/Runtime/TypedElementFramework/Public/Elements/Columns/TypedElementPivotOffsetColumns.h
--- a/Engine/Source/Runtime/TypedElementFramework/Public/Elements/Columns/TypedElementPivotOffsetColumns.h
+++ b/Engine/Source/Runtime/TypedElementFramework/Public/Elements/Columns/TypedElementPivotOffsetColumns.h
@@ -11,6 +11,6 @@ struct FTypedElementPivotOffset final : public FTypedElementDataStorageColumn
 {
 	GENERATED_BODY()
 
-	UPROPERTY()
+	UPROPERTY(meta = (IgnoreForMemberInitializationTest))
 	FVector Offset;
 };
\ No newline at end of file
