Expose Instanced TEDS factories API to DataStorageInterface

diff --git a/Engine/Plugins/Experimental/TypedElementsDataStorage/Source/TypedElementsDataStorage/Private/TypedElementDatabase.h b/Engine/Plugins/Experimental/TypedElementsDataStorage/Source/TypedElementsDataStorage/Private/TypedElementDatabase.h
--- a/Engine/Plugins/Experimental/TypedElementsDataStorage/Source/TypedElementsDataStorage/Private/TypedElementDatabase.h
+++ b/Engine/Plugins/Experimental/TypedElementsDataStorage/Source/TypedElementsDataStorage/Private/TypedElementDatabase.h
@@ -66,7 +66,7 @@ public:
 	FactoryConstIterator CreateFactoryIterator() const;
 
 	/** Returns factory instance given the type of factory */
-	const UTypedElementDataStorageFactory* FindFactory(const UClass* FactoryType) const;
+	const UTypedElementDataStorageFactory* FindFactory(const UClass* FactoryType) const override;
 	/** Helper for FindFactory(const UClass*) */
 	template<typename FactoryTypeT>
 	const FactoryTypeT* FindFactory() const;
diff --git a/Engine/Source/Runtime/TypedElementFramework/Public/Elements/Interfaces/TypedElementDataStorageInterface.h b/Engine/Source/Runtime/TypedElementFramework/Public/Elements/Interfaces/TypedElementDataStorageInterface.h
--- a/Engine/Source/Runtime/TypedElementFramework/Public/Elements/Interfaces/TypedElementDataStorageInterface.h
+++ b/Engine/Source/Runtime/TypedElementFramework/Public/Elements/Interfaces/TypedElementDataStorageInterface.h
@@ -22,6 +22,7 @@
 class UClass;
 class USubsystem;
 class UScriptStruct;
+class UTypedElementDataStorageFactory;
 
 struct ColumnDataResult
 {
@@ -87,6 +88,20 @@ class ITypedElementDataStorageInterface
 	GENERATED_BODY()
 
 public:
+	/**
+	 * @section Factories
+	 *
+	 * @description
+	 * Factories are an automated way to register tables, queries and other information with TEDS.
+	 */
+
+	/** Finds a factory instance registered with TEDS */
+	virtual const UTypedElementDataStorageFactory* FindFactory(const UClass* FactoryType) const = 0;
+
+	/** Convenience function for FindFactory */
+	template<typename FactoryT>
+	const FactoryT* FindFactory() const;
+	
 	/**
 	 * @section Table management
 	 * 
@@ -365,6 +380,12 @@ public:
 
 // Implementations
 
+template <typename FactoryT>
+const FactoryT* ITypedElementDataStorageInterface::FindFactory() const
+{
+	return GetFactory(FactoryT::StaticClass());
+}
+
 template<typename Column>
 bool ITypedElementDataStorageInterface::AddColumn(TypedElementRowHandle Row)
 {
