Fix malformed helper template

diff --git a/Engine/Source/Runtime/TypedElementFramework/Public/Elements/Interfaces/TypedElementDataStorageInterface.h b/Engine/Source/Runtime/TypedElementFramework/Public/Elements/Interfaces/TypedElementDataStorageInterface.h
--- a/Engine/Source/Runtime/TypedElementFramework/Public/Elements/Interfaces/TypedElementDataStorageInterface.h
+++ b/Engine/Source/Runtime/TypedElementFramework/Public/Elements/Interfaces/TypedElementDataStorageInterface.h
@@ -386,7 +386,7 @@ public:
 template <typename FactoryT>
 const FactoryT* ITypedElementDataStorageInterface::FindFactory() const
 {
-	return GetFactory(FactoryT::StaticClass());
+	return static_cast<const FactoryT*>(FindFactory(FactoryT::StaticClass()));
 }
 
 template<typename Column>
