Fixed a crash when enable the new SCC integration with TEDS then dragging an actor into the world.

The following changes were made:
- Removed the intermediate context in the SCC widget. Instead the UpdateStatusWidget now takes a few extra arguments. This helps long term maintenance and readability.
- Tagged the SCC widget with a tag to narrow down the widgets that it considers for updating. In the original setup widgets came in that were not related to the SCC widget and caused issues. This also removed the need for an early out, although a check remains.
- Removed explicitly initializing FNames to none, as this is their default.
- Removed an unused subquery in the SCC widget's callback.
- The `ReadOnly` option in `Select` can now be marked as "optional". This means the column can be accessed if it's present, but it's not considered when matching queries to tables.
- Used C++ concepts to restrict the types on functions in the query builder that can only take columns or even more narrow if they only take data columns.
    - Since FActorFragment is still used, some updates had to be made to use the explicit form as the templated form now no longer passes.

diff --git a/Engine/Plugins/Experimental/TypedElementsDataStorage/Source/TypedElementsDataStorage/Private/Compatibility/TEDSTypedElementActorHandleFactory.cpp b/Engine/Plugins/Experimental/TypedElementsDataStorage/Source/TypedElementsDataStorage/Private/Compatibility/TEDSTypedElementActorHandleFactory.cpp
--- a/Engine/Plugins/Experimental/TypedElementsDataStorage/Source/TypedElementsDataStorage/Private/Compatibility/TEDSTypedElementActorHandleFactory.cpp
+++ b/Engine/Plugins/Experimental/TypedElementsDataStorage/Source/TypedElementsDataStorage/Private/Compatibility/TEDSTypedElementActorHandleFactory.cpp
@@ -1,4 +1,4 @@
-﻿// Copyright Epic Games, Inc. All Rights Reserved.
+// Copyright Epic Games, Inc. All Rights Reserved.
 
 #include "TEDSTypedElementActorHandleFactory.h"
 
@@ -36,7 +36,7 @@ void UTEDSTypedElementActorHandleFactory::RegisterQueries(ITypedElementDataStora
 	using namespace TypedElementDataStorage;
 	GetAllActorsQuery = DataStorage.RegisterQuery(
 	Select()
-		.ReadOnly<FMassActorFragment>()
+		.ReadOnly(FMassActorFragment::StaticStruct())
 	.Compile());
 }
 
@@ -67,7 +67,7 @@ void UTEDSTypedElementActorHandleFactory::RegisterQuery_ActorHandlePopulate(ITyp
 		}
 	)
 	.Where()
-		.All<FMassActorFragment>()
+		.All(FMassActorFragment::StaticStruct())
 	.Compile());
 }
 
diff --git a/Engine/Plugins/Experimental/TypedElementsDataStorage/Source/TypedElementsDataStorage/Private/Debug/TypedElementDatabaseDebugCommands.cpp b/Engine/Plugins/Experimental/TypedElementsDataStorage/Source/TypedElementsDataStorage/Private/Debug/TypedElementDatabaseDebugCommands.cpp
--- a/Engine/Plugins/Experimental/TypedElementsDataStorage/Source/TypedElementsDataStorage/Private/Debug/TypedElementDatabaseDebugCommands.cpp
+++ b/Engine/Plugins/Experimental/TypedElementsDataStorage/Source/TypedElementsDataStorage/Private/Debug/TypedElementDatabaseDebugCommands.cpp
@@ -76,7 +76,7 @@ namespace Private
 						Select()
 							.ReadOnly<FTypedElementUObjectColumn, FTypedElementLabelColumn>()
 						.Where()
-							.All<Conditions...>()
+							.All(Conditions::StaticStruct()...)
 						.Compile());
 				}
 				else
@@ -184,7 +184,8 @@ FAutoConsoleCommandWithOutputDevice CheckActorAndObjectAddressCommand(
 				{
 					return DataStorage->RegisterQuery(
 						Select()
-							.ReadOnly<FTypedElementUObjectColumn, FMassActorFragment>()
+							.ReadOnly<FTypedElementUObjectColumn>()
+							.ReadOnly(FMassActorFragment::StaticStruct())
 						.Compile());
 				}();
 				
diff --git a/Engine/Plugins/Experimental/TypedElementsDataStorage/Source/TypedElementsDataStorage/Private/Processors/TypedElementProcessorAdaptors.cpp b/Engine/Plugins/Experimental/TypedElementsDataStorage/Source/TypedElementsDataStorage/Private/Processors/TypedElementProcessorAdaptors.cpp
--- a/Engine/Plugins/Experimental/TypedElementsDataStorage/Source/TypedElementsDataStorage/Private/Processors/TypedElementProcessorAdaptors.cpp
+++ b/Engine/Plugins/Experimental/TypedElementsDataStorage/Source/TypedElementsDataStorage/Private/Processors/TypedElementProcessorAdaptors.cpp
@@ -57,9 +57,9 @@ struct FMassContextCommon : public T
 		for (int32 Index = 0; Index < TypeCount; ++Index)
 		{
 			checkf(ColumnTypes->IsValid(), TEXT("Attempting to retrieve a column that is not available."));
-			*RetrievedAddresses = *AccessTypes == ITypedElementDataStorageInterface::EQueryAccessType::ReadOnly
-				? const_cast<char*>(reinterpret_cast<const char*>(Context.GetFragmentView(ColumnTypes->Get()).GetData()))
-				: reinterpret_cast<char*>(Context.GetMutableFragmentView(ColumnTypes->Get()).GetData());
+			*RetrievedAddresses = *AccessTypes == ITypedElementDataStorageInterface::EQueryAccessType::ReadWrite
+				? reinterpret_cast<char*>(Context.GetMutableFragmentView(ColumnTypes->Get()).GetData())
+				: const_cast<char*>(reinterpret_cast<const char*>(Context.GetFragmentView(ColumnTypes->Get()).GetData()));
 
 			++RetrievedAddresses;
 			++ColumnTypes;
@@ -431,9 +431,9 @@ public:
 		for (int32 Index = 0; Index < SubsystemCount; ++Index)
 		{
 			checkf(DependencyTypes->IsValid(), TEXT("Attempting to retrieve a subsystem that's no longer valid."));
-			*RetrievedAddresses = *AccessTypes == ITypedElementDataStorageInterface::EQueryAccessType::ReadOnly
-				? const_cast<USubsystem*>(Context.GetSubsystem<USubsystem>(const_cast<UClass*>(DependencyTypes->Get())))
-				: Context.GetMutableSubsystem<USubsystem>(const_cast<UClass*>(DependencyTypes->Get()));
+			*RetrievedAddresses = *AccessTypes == ITypedElementDataStorageInterface::EQueryAccessType::ReadWrite
+				? Context.GetMutableSubsystem<USubsystem>(const_cast<UClass*>(DependencyTypes->Get()))
+				: const_cast<USubsystem*>(Context.GetSubsystem<USubsystem>(const_cast<UClass*>(DependencyTypes->Get())));
 
 			++RetrievedAddresses;
 			++DependencyTypes;
@@ -902,6 +902,30 @@ bool UTypedElementQueryProcessorCallbackAdapterProcessorWith4Subqueries::Configu
 	return ConfigureQueryCallbackData(Query, QueryHandle, QueryStore, Environment, NativeSubqueries);
 }
 
+bool UTypedElementQueryProcessorCallbackAdapterProcessorWith5Subqueries::ConfigureQueryCallback(
+	FTypedElementExtendedQuery& Query, FTypedElementExtendedQueryStore::Handle QueryHandle, FTypedElementExtendedQueryStore& QueryStore, FTypedElementDatabaseEnvironment& Environment)
+{
+	return ConfigureQueryCallbackData(Query, QueryHandle, QueryStore, Environment, NativeSubqueries);
+}
+
+bool UTypedElementQueryProcessorCallbackAdapterProcessorWith6Subqueries::ConfigureQueryCallback(
+	FTypedElementExtendedQuery& Query, FTypedElementExtendedQueryStore::Handle QueryHandle, FTypedElementExtendedQueryStore& QueryStore, FTypedElementDatabaseEnvironment& Environment)
+{
+	return ConfigureQueryCallbackData(Query, QueryHandle, QueryStore, Environment, NativeSubqueries);
+}
+
+bool UTypedElementQueryProcessorCallbackAdapterProcessorWith7Subqueries::ConfigureQueryCallback(
+	FTypedElementExtendedQuery& Query, FTypedElementExtendedQueryStore::Handle QueryHandle, FTypedElementExtendedQueryStore& QueryStore, FTypedElementDatabaseEnvironment& Environment)
+{
+	return ConfigureQueryCallbackData(Query, QueryHandle, QueryStore, Environment, NativeSubqueries);
+}
+
+bool UTypedElementQueryProcessorCallbackAdapterProcessorWith8Subqueries::ConfigureQueryCallback(
+	FTypedElementExtendedQuery& Query, FTypedElementExtendedQueryStore::Handle QueryHandle, FTypedElementExtendedQueryStore& QueryStore, FTypedElementDatabaseEnvironment& Environment)
+{
+	return ConfigureQueryCallbackData(Query, QueryHandle, QueryStore, Environment, NativeSubqueries);
+}
+
 
 /**
  * UTypedElementQueryObserverCallbackAdapterProcessor
@@ -1028,3 +1052,27 @@ bool UTypedElementQueryObserverCallbackAdapterProcessorWith4Subqueries::Configur
 {
 	return ConfigureQueryCallbackData(Query, QueryHandle, QueryStore, Environment, NativeSubqueries);
 }
+
+bool UTypedElementQueryObserverCallbackAdapterProcessorWith5Subqueries::ConfigureQueryCallback(
+	FTypedElementExtendedQuery& Query, FTypedElementExtendedQueryStore::Handle QueryHandle, FTypedElementExtendedQueryStore& QueryStore, FTypedElementDatabaseEnvironment& Environment)
+{
+	return ConfigureQueryCallbackData(Query, QueryHandle, QueryStore, Environment, NativeSubqueries);
+}
+
+bool UTypedElementQueryObserverCallbackAdapterProcessorWith6Subqueries::ConfigureQueryCallback(
+	FTypedElementExtendedQuery& Query, FTypedElementExtendedQueryStore::Handle QueryHandle, FTypedElementExtendedQueryStore& QueryStore, FTypedElementDatabaseEnvironment& Environment)
+{
+	return ConfigureQueryCallbackData(Query, QueryHandle, QueryStore, Environment, NativeSubqueries);
+}
+
+bool UTypedElementQueryObserverCallbackAdapterProcessorWith7Subqueries::ConfigureQueryCallback(
+	FTypedElementExtendedQuery& Query, FTypedElementExtendedQueryStore::Handle QueryHandle, FTypedElementExtendedQueryStore& QueryStore, FTypedElementDatabaseEnvironment& Environment)
+{
+	return ConfigureQueryCallbackData(Query, QueryHandle, QueryStore, Environment, NativeSubqueries);
+}
+
+bool UTypedElementQueryObserverCallbackAdapterProcessorWith8Subqueries::ConfigureQueryCallback(
+	FTypedElementExtendedQuery& Query, FTypedElementExtendedQueryStore::Handle QueryHandle, FTypedElementExtendedQueryStore& QueryStore, FTypedElementDatabaseEnvironment& Environment)
+{
+	return ConfigureQueryCallbackData(Query, QueryHandle, QueryStore, Environment, NativeSubqueries);
+}
diff --git a/Engine/Plugins/Experimental/TypedElementsDataStorage/Source/TypedElementsDataStorage/Private/Processors/TypedElementProcessorAdaptors.h b/Engine/Plugins/Experimental/TypedElementsDataStorage/Source/TypedElementsDataStorage/Private/Processors/TypedElementProcessorAdaptors.h
--- a/Engine/Plugins/Experimental/TypedElementsDataStorage/Source/TypedElementsDataStorage/Private/Processors/TypedElementProcessorAdaptors.h
+++ b/Engine/Plugins/Experimental/TypedElementsDataStorage/Source/TypedElementsDataStorage/Private/Processors/TypedElementProcessorAdaptors.h
@@ -202,6 +202,74 @@ private:
 	FMassEntityQuery NativeSubqueries[4];
 };
 
+UCLASS()
+class UTypedElementQueryProcessorCallbackAdapterProcessorWith5Subqueries final : public UTypedElementQueryProcessorCallbackAdapterProcessorBase
+{
+	GENERATED_BODY()
+
+public:
+	bool ConfigureQueryCallback(
+		FTypedElementExtendedQuery& Query,
+		FTypedElementExtendedQueryStore::Handle QueryHandle,
+		FTypedElementExtendedQueryStore& QueryStore,
+		FTypedElementDatabaseEnvironment& Environment) override;
+
+private:
+	UPROPERTY(transient)
+		FMassEntityQuery NativeSubqueries[5];
+};
+
+UCLASS()
+class UTypedElementQueryProcessorCallbackAdapterProcessorWith6Subqueries final : public UTypedElementQueryProcessorCallbackAdapterProcessorBase
+{
+	GENERATED_BODY()
+
+public:
+	bool ConfigureQueryCallback(
+		FTypedElementExtendedQuery& Query,
+		FTypedElementExtendedQueryStore::Handle QueryHandle,
+		FTypedElementExtendedQueryStore& QueryStore,
+		FTypedElementDatabaseEnvironment& Environment) override;
+
+private:
+	UPROPERTY(transient)
+		FMassEntityQuery NativeSubqueries[6];
+};
+
+UCLASS()
+class UTypedElementQueryProcessorCallbackAdapterProcessorWith7Subqueries final : public UTypedElementQueryProcessorCallbackAdapterProcessorBase
+{
+	GENERATED_BODY()
+
+public:
+	bool ConfigureQueryCallback(
+		FTypedElementExtendedQuery& Query,
+		FTypedElementExtendedQueryStore::Handle QueryHandle,
+		FTypedElementExtendedQueryStore& QueryStore,
+		FTypedElementDatabaseEnvironment& Environment) override;
+
+private:
+	UPROPERTY(transient)
+		FMassEntityQuery NativeSubqueries[7];
+};
+
+UCLASS()
+class UTypedElementQueryProcessorCallbackAdapterProcessorWith8Subqueries final : public UTypedElementQueryProcessorCallbackAdapterProcessorBase
+{
+	GENERATED_BODY()
+
+public:
+	bool ConfigureQueryCallback(
+		FTypedElementExtendedQuery& Query,
+		FTypedElementExtendedQueryStore::Handle QueryHandle,
+		FTypedElementExtendedQueryStore& QueryStore,
+		FTypedElementDatabaseEnvironment& Environment) override;
+
+private:
+	UPROPERTY(transient)
+		FMassEntityQuery NativeSubqueries[8];
+};
+
 /**
  * Adapts observer queries callback for MASS.
  */
@@ -315,3 +383,74 @@ private:
 	UPROPERTY(transient)
 	FMassEntityQuery NativeSubqueries[4];
 };
+
+UCLASS()
+class UTypedElementQueryObserverCallbackAdapterProcessorWith5Subqueries final : public UTypedElementQueryObserverCallbackAdapterProcessorBase
+{
+	GENERATED_BODY()
+
+public:
+	bool ConfigureQueryCallback(
+		FTypedElementExtendedQuery& Query,
+		FTypedElementExtendedQueryStore::Handle QueryHandle,
+		FTypedElementExtendedQueryStore& QueryStore,
+		FTypedElementDatabaseEnvironment& Environment) override;
+
+private:
+	UPROPERTY(transient)
+		FMassEntityQuery NativeSubqueries[5];
+};
+
+UCLASS()
+class UTypedElementQueryObserverCallbackAdapterProcessorWith6Subqueries final : public UTypedElementQueryObserverCallbackAdapterProcessorBase
+{
+	GENERATED_BODY()
+
+public:
+	bool ConfigureQueryCallback(
+		FTypedElementExtendedQuery& Query,
+		FTypedElementExtendedQueryStore::Handle QueryHandle,
+		FTypedElementExtendedQueryStore& QueryStore,
+		FTypedElementDatabaseEnvironment& Environment) override;
+
+private:
+	UPROPERTY(transient)
+		FMassEntityQuery NativeSubqueries[6];
+};
+
+UCLASS()
+class UTypedElementQueryObserverCallbackAdapterProcessorWith7Subqueries final : public UTypedElementQueryObserverCallbackAdapterProcessorBase
+{
+	GENERATED_BODY()
+
+public:
+	bool ConfigureQueryCallback(
+		FTypedElementExtendedQuery& Query,
+		FTypedElementExtendedQueryStore::Handle QueryHandle,
+		FTypedElementExtendedQueryStore& QueryStore,
+		FTypedElementDatabaseEnvironment& Environment) override;
+
+private:
+	UPROPERTY(transient)
+		FMassEntityQuery NativeSubqueries[7];
+};
+
+UCLASS()
+class UTypedElementQueryObserverCallbackAdapterProcessorWith8Subqueries final : public UTypedElementQueryObserverCallbackAdapterProcessorBase
+{
+	GENERATED_BODY()
+
+public:
+	bool ConfigureQueryCallback(
+		FTypedElementExtendedQuery& Query,
+		FTypedElementExtendedQueryStore::Handle QueryHandle,
+		FTypedElementExtendedQueryStore& QueryStore,
+		FTypedElementDatabaseEnvironment& Environment) override;
+
+private:
+	UPROPERTY(transient)
+		FMassEntityQuery NativeSubqueries[8];
+};
+
+
+
diff --git a/Engine/Plugins/Experimental/TypedElementsDataStorage/Source/TypedElementsDataStorage/Private/Queries/TypedElementExtendedQueryStore.cpp b/Engine/Plugins/Experimental/TypedElementsDataStorage/Source/TypedElementsDataStorage/Private/Queries/TypedElementExtendedQueryStore.cpp
--- a/Engine/Plugins/Experimental/TypedElementsDataStorage/Source/TypedElementsDataStorage/Private/Queries/TypedElementExtendedQueryStore.cpp
+++ b/Engine/Plugins/Experimental/TypedElementsDataStorage/Source/TypedElementsDataStorage/Private/Queries/TypedElementExtendedQueryStore.cpp
@@ -413,6 +413,18 @@ FMassEntityQuery& FTypedElementExtendedQueryStore::SetupNativeQuery(
 			case 4:
 				Processor = NewObject<UTypedElementQueryProcessorCallbackAdapterProcessorWith4Subqueries>();
 				break;
+			case 5:
+				Processor = NewObject<UTypedElementQueryProcessorCallbackAdapterProcessorWith5Subqueries>();
+				break;
+			case 6:
+				Processor = NewObject<UTypedElementQueryProcessorCallbackAdapterProcessorWith6Subqueries>();
+				break;
+			case 7:
+				Processor = NewObject<UTypedElementQueryProcessorCallbackAdapterProcessorWith7Subqueries>();
+				break;
+			case 8:
+				Processor = NewObject<UTypedElementQueryProcessorCallbackAdapterProcessorWith8Subqueries>();
+				break;
 			default:
 				checkf(false, TEXT("The current Typed Elements Data Storage backend doesn't support %i subqueries per processor query."), 
 					Query.Subqueries.Num());
@@ -443,6 +455,18 @@ FMassEntityQuery& FTypedElementExtendedQueryStore::SetupNativeQuery(
 			case 4:
 				Observer = NewObject<UTypedElementQueryObserverCallbackAdapterProcessorWith4Subqueries>();
 				break;
+			case 5:
+				Observer = NewObject<UTypedElementQueryObserverCallbackAdapterProcessorWith5Subqueries>();
+				break;
+			case 6:
+				Observer = NewObject<UTypedElementQueryObserverCallbackAdapterProcessorWith6Subqueries>();
+				break;
+			case 7:
+				Observer = NewObject<UTypedElementQueryObserverCallbackAdapterProcessorWith7Subqueries>();
+				break;
+			case 8:
+				Observer = NewObject<UTypedElementQueryObserverCallbackAdapterProcessorWith8Subqueries>();
+				break;
 			default:
 				checkf(false, TEXT("The current Typed Elements Data Storage backend doesn't support %i subqueries per observer query."),
 					Query.Subqueries.Num());
@@ -482,7 +506,7 @@ bool FTypedElementExtendedQueryStore::SetupSelectedColumns(
 			for (int SelectionIndex = 0; SelectionIndex < SelectionCount; ++SelectionIndex)
 			{
 				TWeakObjectPtr<const UScriptStruct>& Type = Query.SelectionTypes[SelectionIndex];
-				
+				DSI::EQueryAccessType AccessType = Query.SelectionAccessTypes[SelectionIndex];
 				if (ensureMsgf(Type.IsValid(), TEXT("Provided query selection type can not be null.")) &&
 					ensureMsgf(
 						Type->IsChildOf(FTypedElementDataStorageColumn::StaticStruct()) ||
@@ -490,7 +514,7 @@ bool FTypedElementExtendedQueryStore::SetupSelectedColumns(
 						TEXT("Provided query selection type '%s' is not based on FTypedElementDataStorageColumn or another supported base type."),
 						*Type->GetStructPathName().ToString()))
 				{
-					NativeQuery.AddRequirement(Type.Get(), ConvertToNativeAccessType(Query.SelectionAccessTypes[SelectionIndex]));
+					NativeQuery.AddRequirement(Type.Get(), ConvertToNativeAccessType(AccessType), ConvertToNativePresenceType(AccessType));
 				}
 				else
 				{
@@ -695,6 +719,8 @@ EMassFragmentAccess FTypedElementExtendedQueryStore::ConvertToNativeAccessType(I
 	switch (AccessType)
 	{
 	case ITypedElementDataStorageInterface::EQueryAccessType::ReadOnly:
+		// Fall through
+	case ITypedElementDataStorageInterface::EQueryAccessType::OptionalReadOnly:
 		return EMassFragmentAccess::ReadOnly;
 	case ITypedElementDataStorageInterface::EQueryAccessType::ReadWrite:
 		return EMassFragmentAccess::ReadWrite;
@@ -704,6 +730,22 @@ EMassFragmentAccess FTypedElementExtendedQueryStore::ConvertToNativeAccessType(I
 	}
 }
 
+EMassFragmentPresence FTypedElementExtendedQueryStore::ConvertToNativePresenceType(ITypedElementDataStorageInterface::EQueryAccessType AccessType)
+{
+	switch (AccessType)
+	{
+	case ITypedElementDataStorageInterface::EQueryAccessType::ReadOnly:
+		return EMassFragmentPresence::All;
+	case ITypedElementDataStorageInterface::EQueryAccessType::OptionalReadOnly:
+		return EMassFragmentPresence::Optional;
+	case ITypedElementDataStorageInterface::EQueryAccessType::ReadWrite:
+		return EMassFragmentPresence::All;
+	default:
+		checkf(false, TEXT("Invalid query access type: %i."), static_cast<uint32>(AccessType));
+		return EMassFragmentPresence::MAX;
+	}
+}
+
 void FTypedElementExtendedQueryStore::RegisterPreambleQuery(ITypedElementDataStorageInterface::EQueryTickPhase Phase, Handle Query)
 {
 	PhasePreparationQueries[static_cast<QueryTickPhaseType>(Phase)].Add(Query);
diff --git a/Engine/Plugins/Experimental/TypedElementsDataStorage/Source/TypedElementsDataStorage/Private/Queries/TypedElementExtendedQueryStore.h b/Engine/Plugins/Experimental/TypedElementsDataStorage/Source/TypedElementsDataStorage/Private/Queries/TypedElementExtendedQueryStore.h
--- a/Engine/Plugins/Experimental/TypedElementsDataStorage/Source/TypedElementsDataStorage/Private/Queries/TypedElementExtendedQueryStore.h
+++ b/Engine/Plugins/Experimental/TypedElementsDataStorage/Source/TypedElementsDataStorage/Private/Queries/TypedElementExtendedQueryStore.h
@@ -169,6 +169,7 @@ private:
 	bool SetupActivatable(Handle QueryHandle, ITypedElementDataStorageInterface::FQueryDescription& Query);
 
 	EMassFragmentAccess ConvertToNativeAccessType(ITypedElementDataStorageInterface::EQueryAccessType AccessType);
+	EMassFragmentPresence ConvertToNativePresenceType(ITypedElementDataStorageInterface::EQueryAccessType AccessType);
 
 	void RegisterPreambleQuery(ITypedElementDataStorageInterface::EQueryTickPhase Phase, Handle Query);
 	void RegisterPostambleQuery(ITypedElementDataStorageInterface::EQueryTickPhase Phase, Handle Query);
diff --git a/Engine/Source/Runtime/TypedElementFramework/Private/Elements/Framework/TypedElementQueryBuilder.cpp b/Engine/Source/Runtime/TypedElementFramework/Private/Elements/Framework/TypedElementQueryBuilder.cpp
--- a/Engine/Source/Runtime/TypedElementFramework/Private/Elements/Framework/TypedElementQueryBuilder.cpp
+++ b/Engine/Source/Runtime/TypedElementFramework/Private/Elements/Framework/TypedElementQueryBuilder.cpp
@@ -358,6 +358,32 @@ namespace TypedElementQueryBuilder
 		return *this;
 	}
 
+	Select& Select::ReadOnly(const UScriptStruct* Target, EOptional Optional)
+	{
+		checkf(Target, TEXT("The Select section in the Typed Elements query builder doesn't support nullptrs as Read-Only input."));
+		Query.SelectionTypes.Emplace(Target);
+		Query.SelectionAccessTypes.Emplace(Optional == EOptional::Yes
+			? ITypedElementDataStorageInterface::EQueryAccessType::OptionalReadOnly
+			: ITypedElementDataStorageInterface::EQueryAccessType::ReadOnly);
+		Query.SelectionMetaData.Emplace(Target, TypedElementDataStorage::FColumnMetaData::EFlags::None);
+
+		return *this;
+	}
+
+	Select& Select::ReadOnly(TConstArrayView<const UScriptStruct*> Targets, EOptional Optional)
+	{
+		int32 NewCount = Query.SelectionTypes.Num() + Targets.Num();
+		Query.SelectionTypes.Reserve(NewCount);
+		Query.SelectionAccessTypes.Reserve(NewCount);
+		Query.SelectionMetaData.Reserve(NewCount);
+
+		for (const UScriptStruct* Target : Targets)
+		{
+			ReadOnly(Target, Optional);
+		}
+		return *this;
+	}
+
 	Select& Select::ReadWrite(const UScriptStruct* Target)
 	{
 		checkf(Target, TEXT("The Select section in the Typed Elements query builder doesn't support nullptrs as Read/Write input."));
diff --git a/Engine/Source/Runtime/TypedElementFramework/Public/Elements/Common/TypedElementQueryTypes.h b/Engine/Source/Runtime/TypedElementFramework/Public/Elements/Common/TypedElementQueryTypes.h
--- a/Engine/Source/Runtime/TypedElementFramework/Public/Elements/Common/TypedElementQueryTypes.h
+++ b/Engine/Source/Runtime/TypedElementFramework/Public/Elements/Common/TypedElementQueryTypes.h
@@ -75,9 +75,13 @@ namespace TypedElementDataStorage
 		Max //< Value indicating the maximum value in this enum. Not to be used as an enum value.
 	};
 
-	enum class EQueryAccessType : bool
+	enum class EQueryAccessType : uint8
 	{
+		/** Provide read-only access. */
 		ReadOnly,
+		/** Provide read-only access, but do not include when finding matching tables. */
+		OptionalReadOnly,
+		/** Provide read and write access. */
 		ReadWrite
 	};
 
diff --git a/Engine/Source/Runtime/TypedElementFramework/Public/Elements/Framework/TypedElementQueryBuilder.h b/Engine/Source/Runtime/TypedElementFramework/Public/Elements/Framework/TypedElementQueryBuilder.h
--- a/Engine/Source/Runtime/TypedElementFramework/Public/Elements/Framework/TypedElementQueryBuilder.h
+++ b/Engine/Source/Runtime/TypedElementFramework/Public/Elements/Framework/TypedElementQueryBuilder.h
@@ -10,6 +10,23 @@
 class UScriptStruct;
 class USubsystem;
 
+namespace UE
+{
+	// Work around missing header/implementations on some platforms
+	namespace detail
+	{
+		template<typename T, typename U>
+		concept SameHelper = std::is_same_v<T, U>;
+	}
+	template<typename T, typename U>
+	concept same_as = detail::SameHelper<T, U> && detail::SameHelper<U, T>;
+
+	template<typename From, typename To>
+	concept convertible_to = std::is_convertible_v<From, To> && requires { static_cast<To>(std::declval<From>()); };
+
+	template<typename Derived, typename Base>
+	concept derived_from = std::is_base_of_v<Base, Derived> && std::is_convertible_v<const volatile Derived*, const volatile Base*>;
+}
 /**
  * The TypedElementQueryBuilder allows for the construction of queries for use by the Typed Element Data Storage.
  * There are two types of queries, simple and normal. Simple queries are guaranteed to be supported by the data
@@ -110,6 +127,21 @@ namespace TypedElementQueryBuilder
 	TYPEDELEMENTFRAMEWORK_API const UScriptStruct* operator""_Type(const char* Name, std::size_t NameSize);
 	TYPEDELEMENTFRAMEWORK_API const UScriptStruct* operator""_TypeOptional(const char* Name, std::size_t NameSize);
 
+	template<typename T>
+	concept TDataColumnType = UE::derived_from<T, FTypedElementDataStorageColumn>;
+
+	template<typename T>
+	concept TTagColumnType = UE::derived_from<T, FTypedElementDataStorageTag>;
+
+	template<typename T>
+	concept TColumnType = TDataColumnType<T> || TTagColumnType<T>;
+
+	enum class EOptional
+	{
+		No,
+		Yes
+	};
+
 	class FDependency final
 	{
 		friend class Count;
@@ -145,15 +177,15 @@ namespace TypedElementQueryBuilder
 		TYPEDELEMENTFRAMEWORK_API FDependency DependsOn();
 		TYPEDELEMENTFRAMEWORK_API ITypedElementDataStorageInterface::FQueryDescription&& Compile();
 
-		template<typename... TargetTypes>
+		template<TColumnType... TargetTypes>
 		FSimpleQuery& All();
 		TYPEDELEMENTFRAMEWORK_API FSimpleQuery& All(const UScriptStruct* Target);
 		TYPEDELEMENTFRAMEWORK_API FSimpleQuery& All(TConstArrayView<const UScriptStruct*> Targets);
-		template<typename... TargetTypes>
+		template<TColumnType... TargetTypes>
 		FSimpleQuery& Any();
 		TYPEDELEMENTFRAMEWORK_API FSimpleQuery& Any(const UScriptStruct* Target);
 		TYPEDELEMENTFRAMEWORK_API FSimpleQuery& Any(TConstArrayView<const UScriptStruct*> Targets);
-		template<typename... TargetTypes>
+		template<TColumnType... TargetTypes>
 		FSimpleQuery& None();
 		TYPEDELEMENTFRAMEWORK_API FSimpleQuery& None(const UScriptStruct* Target);
 		TYPEDELEMENTFRAMEWORK_API FSimpleQuery& None(TConstArrayView<const UScriptStruct*> Targets);
@@ -194,14 +226,14 @@ namespace TypedElementQueryBuilder
 
 		TYPEDELEMENTFRAMEWORK_API FObserver(EEvent MonitorForEvent, const UScriptStruct* MonitoredColumn);
 
-		template<typename ColumnType>
+		template<TColumnType ColumnType>
 		static FObserver OnAdd();
-		template<typename ColumnType>
+		template<TColumnType ColumnType>
 		static FObserver OnRemove();
 
 		TYPEDELEMENTFRAMEWORK_API FObserver& SetEvent(EEvent MonitorForEvent);
 		TYPEDELEMENTFRAMEWORK_API FObserver& SetMonitoredColumn(const UScriptStruct* MonitoredColumn);
-		template<typename ColumnType>
+		template<TColumnType ColumnType>
 		FObserver& SetMonitoredColumn();
 		TYPEDELEMENTFRAMEWORK_API FObserver& ForceToGameThread(bool bForce);
 		TYPEDELEMENTFRAMEWORK_API FObserver& MakeActivatable(FName Name);
@@ -296,18 +328,42 @@ namespace TypedElementQueryBuilder
 	public:
 		TYPEDELEMENTFRAMEWORK_API Select();
 
+		/** Select the columns to operate one, giving read - only or read / write access. */
 		template<typename CallbackType, typename Function>
 		Select(FName Name, const CallbackType& Type, Function&& Callback);
+		/** Select the columns to operate one, giving read - only or read / write access. */
 		template<typename CallbackType, typename Class, typename Function>
 		Select(FName Name, const CallbackType& Type, Class* Instance, Function&& Callback);
 
-		template<typename... TargetTypes>
+		/** Request read-only access to the listed columns. */
+		template<TDataColumnType... TargetTypes>
 		Select& ReadOnly();
+		/** Request read-only access to the listed columns. */
 		TYPEDELEMENTFRAMEWORK_API Select& ReadOnly(const UScriptStruct* Target);
+		/** Request read-only access to the listed columns. */
 		TYPEDELEMENTFRAMEWORK_API Select& ReadOnly(TConstArrayView<const UScriptStruct*> Targets);
-		template<typename... TargetTypes>
+		/** 
+		 * Request read-only access to the listed columns. If optional is true read access will be given if the column is in the table but
+		 * it will not be used for finding matching tables. Columns bound with optional can not be bound to a query callback argument.
+		 */
+		template<TDataColumnType... TargetTypes>
+		Select& ReadOnly(EOptional Optional);
+		/**
+		 * Request read-only access to the listed columns. If optional is true read access will be given if the column is in the table but
+		 * it will not be used for finding matching tables. Columns bound with optional can not be bound to a query callback argument.
+		 */
+		TYPEDELEMENTFRAMEWORK_API Select& ReadOnly(const UScriptStruct* Target, EOptional Optional);
+		/**
+		 * Request read-only access to the listed columns. If optional is true read access will be given if the column is in the table but
+		 * it will not be used for finding matching tables. Columns bound with optional can not be bound to a query callback argument.
+		 */
+		TYPEDELEMENTFRAMEWORK_API Select& ReadOnly(TConstArrayView<const UScriptStruct*> Targets, EOptional Optionall);
+		/** Request read and write access to the listed columns. */
+		template<TDataColumnType... TargetTypes>
 		Select& ReadWrite();
+		/** Request read and write access to the listed columns. */
 		TYPEDELEMENTFRAMEWORK_API Select& ReadWrite(const UScriptStruct* Target);
+		/** Request read and write access to the listed columns. */
 		TYPEDELEMENTFRAMEWORK_API Select& ReadWrite(TConstArrayView<const UScriptStruct*> Targets);
 
 		TYPEDELEMENTFRAMEWORK_API ITypedElementDataStorageInterface::FQueryDescription&& Compile();
diff --git a/Engine/Source/Runtime/TypedElementFramework/Public/Elements/Framework/TypedElementQueryBuilder.inl b/Engine/Source/Runtime/TypedElementFramework/Public/Elements/Framework/TypedElementQueryBuilder.inl
--- a/Engine/Source/Runtime/TypedElementFramework/Public/Elements/Framework/TypedElementQueryBuilder.inl
+++ b/Engine/Source/Runtime/TypedElementFramework/Public/Elements/Framework/TypedElementQueryBuilder.inl
@@ -10,24 +10,6 @@
 #include "Elements/Common/TypedElementQueryDescription.h"
 #include "Elements/Common/TypedElementQueryTypes.h"
 
-// Work around missing header/implementations on some platforms
-namespace UE
-{
-	namespace detail
-	{
-		template< class T, class U >
-		concept SameHelper = std::is_same_v<T, U>;
-	}
-	template< class T, class U >
-	concept same_as = detail::SameHelper<T, U> && detail::SameHelper<U, T>;
-
-	template <class From, class To>
-	concept convertible_to = std::is_convertible_v<From, To> && requires { static_cast<To>(std::declval<From>()); };
-
-	template< class Derived, class Base >
-	concept derived_from = std::is_base_of_v<Base, Derived> && std::is_convertible_v<const volatile Derived*, const volatile Base*>;
-}
-
 namespace TypedElementQueryBuilder
 {
 	namespace Internal
@@ -108,19 +90,19 @@ namespace TypedElementQueryBuilder
 	// FObserver
 	//
 	
-	template<typename ColumnType>
+	template<TColumnType ColumnType>
 	FObserver FObserver::OnAdd()
 	{
 		return FObserver(FObserver::EEvent::Add, ColumnType::StaticStruct());
 	}
 
-	template<typename ColumnType>
+	template<TColumnType ColumnType>
 	FObserver FObserver::OnRemove()
 	{
 		return FObserver(FObserver::EEvent::Remove, ColumnType::StaticStruct());
 	}
 
-	template<typename ColumnType>
+	template<TColumnType ColumnType>
 	FObserver& FObserver::SetMonitoredColumn()
 	{
 		return SetMonitoredColumn(ColumnType::StaticStruct());
@@ -1135,14 +1117,21 @@ e.g. void(FCachedQueryContext<Subsystem1, const Subsystem2>& Context, TypedEleme
 		Internal::BindQueryFunction<TypedElementDataStorage::IQueryContext, ValidateColumns>(Query.Callback.Function, Instance, Callback);
 	}
 
-	template<typename... TargetTypes>
+	template<TDataColumnType... TargetTypes>
 	Select& Select::ReadOnly()
 	{
 		ReadOnly({ TargetTypes::StaticStruct()... });
 		return *this;
 	}
 
-	template<typename... TargetTypes>
+	template<TDataColumnType... TargetTypes>
+	Select& Select::ReadOnly(EOptional Optional)
+	{
+		ReadOnly({ TargetTypes::StaticStruct()... }, Optional);
+		return *this;
+	}
+
+	template<TDataColumnType... TargetTypes>
 	Select& Select::ReadWrite()
 	{
 		ReadWrite({ TargetTypes::StaticStruct()... });
@@ -1154,21 +1143,21 @@ e.g. void(FCachedQueryContext<Subsystem1, const Subsystem2>& Context, TypedEleme
 	// FSimpleQuery
 	//
 
-	template<typename... TargetTypes>
+	template<TColumnType... TargetTypes>
 	FSimpleQuery& FSimpleQuery::All()
 	{
 		All({ TargetTypes::StaticStruct()... });
 		return *this;
 	}
 
-	template<typename... TargetTypes>
+	template<TColumnType... TargetTypes>
 	FSimpleQuery& FSimpleQuery::Any()
 	{
 		Any({ TargetTypes::StaticStruct()... });
 		return *this;
 	}
 
-	template<typename... TargetTypes>
+	template<TColumnType... TargetTypes>
 	FSimpleQuery& FSimpleQuery::None()
 	{
 		None({ TargetTypes::StaticStruct()... });
