Add a StringView overload to the GenerateIndexHash function in the TypedElementIndexHasher namespace.

diff --git a/Engine/Source/Runtime/TypedElementFramework/Public/Elements/Framework/TypedElementIndexHasher.h b/Engine/Source/Runtime/TypedElementFramework/Public/Elements/Framework/TypedElementIndexHasher.h
--- a/Engine/Source/Runtime/TypedElementFramework/Public/Elements/Framework/TypedElementIndexHasher.h
+++ b/Engine/Source/Runtime/TypedElementFramework/Public/Elements/Framework/TypedElementIndexHasher.h
@@ -2,6 +2,7 @@
 
 #pragma once
 
+#include "Containers/StringFwd.h"
 #include "Containers/UnrealString.h"
 #include "Elements/Interfaces/TypedElementQueryStorageInterfaces.h"
 #include "UObject/NameTypes.h"
@@ -22,6 +23,7 @@ namespace TypedElementDataStorage
 	IndexHash GenerateIndexHash(const TStrongObjectPtr<T>& Object);
 
 	inline IndexHash GenerateIndexHash(const FString& Object);
+	inline IndexHash GenerateIndexHash(FStringView Object);
 	inline IndexHash GenerateIndexHash(FName Object);
 } // namespace TypedElementDataStorage
 
diff --git a/Engine/Source/Runtime/TypedElementFramework/Public/Elements/Framework/TypedElementIndexHasher.inl b/Engine/Source/Runtime/TypedElementFramework/Public/Elements/Framework/TypedElementIndexHasher.inl
--- a/Engine/Source/Runtime/TypedElementFramework/Public/Elements/Framework/TypedElementIndexHasher.inl
+++ b/Engine/Source/Runtime/TypedElementFramework/Public/Elements/Framework/TypedElementIndexHasher.inl
@@ -35,6 +35,11 @@ namespace TypedElementDataStorage
 		return CityHash64(reinterpret_cast<const char*>(*Object), Object.Len() * sizeof(**Object));
 	}
 
+	IndexHash GenerateIndexHash(FStringView Object)
+	{
+		return CityHash64(reinterpret_cast<const char*>(Object.GetData()), Object.Len() * sizeof(*Object.GetData()));
+	}
+
 	IndexHash GenerateIndexHash(FName Object)
 	{
 		constexpr static const char SeedName[] = "FName";
