Remove `AddOrGetColumn` from TEDS.

The `AddOrGetColumn` functions wasn't being used to get columns. It also added confusion about what it actually did when data was passed in for an already created column. A previous change already moved all `AddOrGetColumn` calls over to `AddColumn`, so this change removes the last remaining pieces. It also completes the pass to have all templates that require columns use concepts for additional type safety.

diff --git a/Engine/Plugins/Experimental/TypedElementsDataStorage/Source/TypedElementsDataStorage/Private/TypedElementDatabase.cpp b/Engine/Plugins/Experimental/TypedElementsDataStorage/Source/TypedElementsDataStorage/Private/TypedElementDatabase.cpp
--- a/Engine/Plugins/Experimental/TypedElementsDataStorage/Source/TypedElementsDataStorage/Private/TypedElementDatabase.cpp
+++ b/Engine/Plugins/Experimental/TypedElementsDataStorage/Source/TypedElementsDataStorage/Private/TypedElementDatabase.cpp
@@ -428,56 +428,42 @@ void UTypedElementDatabase::AddColumnData(TypedElementRowHandle Row, const UScri
 	const TypedElementDataStorage::ColumnCreationCallbackRef& Initializer,
 	TypedElementDataStorage::ColumnCopyOrMoveCallback Relocator)
 {
-	AddOrGetColumnData(Row, ColumnType, Initializer, Relocator);
-}
-
-void UTypedElementDatabase::RemoveColumn(TypedElementRowHandle Row, const UScriptStruct* ColumnType)
-{
-	if (ColumnType && ActiveEditorEntityManager)
+	if (ActiveEditorEntityManager && ColumnType && ColumnType->IsChildOf(FMassFragment::StaticStruct()))
 	{
 		if (HasRowBeenAssigned(Row))
 		{
-			FTypedElementDatabaseCommandBuffer::Execute_RemoveColumnCommand(*ActiveEditorEntityManager, Row, ColumnType);
+			FMassEntityHandle Entity = FMassEntityHandle::FromNumber(Row);
+			FStructView Column = ActiveEditorEntityManager->GetFragmentDataStruct(Entity, ColumnType);
+			if (!Column.IsValid())
+			{
+				ActiveEditorEntityManager->AddFragmentToEntity(Entity, ColumnType, Initializer);
+			}
+			else
+			{
+				Initializer(Column.GetMemory(), *ColumnType);
+			}
 		}
 		else
 		{
-			Environment->GetDirectDeferredCommands().Queue_RemoveColumnCommand(Row, ColumnType);
+			void* Column = Environment->GetDirectDeferredCommands().Queue_AddDataColumnCommandUnitialized(Row, ColumnType, Relocator);
+			Initializer(Column, *ColumnType);
 		}
 	}
 }
 
-void* UTypedElementDatabase::AddOrGetColumnData(TypedElementRowHandle Row, const UScriptStruct* ColumnType,
-	const TypedElementDataStorage::ColumnCreationCallbackRef& Initializer, TypedElementDataStorage::ColumnCopyOrMoveCallback Relocator)
+void UTypedElementDatabase::RemoveColumn(TypedElementRowHandle Row, const UScriptStruct* ColumnType)
 {
-	FMassEntityHandle Entity = FMassEntityHandle::FromNumber(Row);
-	if (ActiveEditorEntityManager && ColumnType && ColumnType->IsChildOf(FMassFragment::StaticStruct()))
+	if (ColumnType && ActiveEditorEntityManager)
 	{
 		if (HasRowBeenAssigned(Row))
 		{
-			FStructView Column = ActiveEditorEntityManager->GetFragmentDataStruct(Entity, ColumnType);
-			if (!Column.IsValid())
-			{
-				ActiveEditorEntityManager->AddFragmentToEntity(Entity, ColumnType, Initializer);
-				Column = ActiveEditorEntityManager->GetFragmentDataStruct(Entity, ColumnType);
-				checkf(Column.IsValid(), TEXT("Added a new column to the Typed Element's data store, but it couldn't be retrieved."));
-				return Column.GetMemory();
-
-			}
-			else
-			{
-				void* Result = Column.GetMemory();
-				Initializer(Result, *ColumnType);
-				return Result;
-			}
+			FTypedElementDatabaseCommandBuffer::Execute_RemoveColumnCommand(*ActiveEditorEntityManager, Row, ColumnType);
 		}
 		else
 		{
-			void* Result = Environment->GetDirectDeferredCommands().Queue_AddDataColumnCommandUnitialized(Row, ColumnType, Relocator);
-			Initializer(Result, *ColumnType);
-			return Result;
+			Environment->GetDirectDeferredCommands().Queue_RemoveColumnCommand(Row, ColumnType);
 		}
 	}
-	return nullptr;
 }
 
 const void* UTypedElementDatabase::GetColumnData(TypedElementRowHandle Row, const UScriptStruct* ColumnType) const
diff --git a/Engine/Plugins/Experimental/TypedElementsDataStorage/Source/TypedElementsDataStorage/Private/TypedElementDatabase.h b/Engine/Plugins/Experimental/TypedElementsDataStorage/Source/TypedElementsDataStorage/Private/TypedElementDatabase.h
--- a/Engine/Plugins/Experimental/TypedElementsDataStorage/Source/TypedElementsDataStorage/Private/TypedElementDatabase.h
+++ b/Engine/Plugins/Experimental/TypedElementsDataStorage/Source/TypedElementsDataStorage/Private/TypedElementDatabase.h
@@ -100,9 +100,6 @@ public:
 		const TypedElementDataStorage::ColumnCreationCallbackRef& Initializer,
 		TypedElementDataStorage::ColumnCopyOrMoveCallback Relocator) override;
 	void RemoveColumn(TypedElementRowHandle Row, const UScriptStruct* ColumnType) override;
-	void* AddOrGetColumnData(TypedElementRowHandle Row, const UScriptStruct* ColumnType,
-		const TypedElementDataStorage::ColumnCreationCallbackRef& Initializer,
-		TypedElementDataStorage::ColumnCopyOrMoveCallback Relocator) override;
 	void* GetColumnData(TypedElementRowHandle Row, const UScriptStruct* ColumnType) override;
 	const void* GetColumnData(TypedElementRowHandle Row, const UScriptStruct* ColumnType) const override;
 	void AddColumns(TypedElementRowHandle Row, TConstArrayView<const UScriptStruct*> Columns) override;
diff --git a/Engine/Plugins/Experimental/TypedElementsDataStorage/Source/TypedElementsDataStorage/Private/TypedElementDatabaseCompatibility.cpp b/Engine/Plugins/Experimental/TypedElementsDataStorage/Source/TypedElementsDataStorage/Private/TypedElementDatabaseCompatibility.cpp
--- a/Engine/Plugins/Experimental/TypedElementsDataStorage/Source/TypedElementsDataStorage/Private/TypedElementDatabaseCompatibility.cpp
+++ b/Engine/Plugins/Experimental/TypedElementsDataStorage/Source/TypedElementsDataStorage/Private/TypedElementDatabaseCompatibility.cpp
@@ -618,11 +618,11 @@ void UTypedElementDatabaseCompatibility::TickPendingExternalObjectRegistration()
 		ExternalObjectsPendingRegistration.ProcessEntries(*Storage, *this,
 			[this](TypedElementRowHandle Row, const ExternalObjectRegistration& Object)
 			{
-				Storage->AddOrGetColumn<FTypedElementExternalObjectColumn>(Row, FTypedElementExternalObjectColumn{ .Object = Object.Object });
-				Storage->AddOrGetColumn<FTypedElementScriptStructTypeInfoColumn>(Row, FTypedElementScriptStructTypeInfoColumn{ .TypeInfo = Object.TypeInfo });
+				ITypedElementDataStorageInterface* Interface = Storage;
+				Interface->AddColumn(Row, FTypedElementExternalObjectColumn{ .Object = Object.Object });
+				Interface->AddColumn(Row, FTypedElementScriptStructTypeInfoColumn{ .TypeInfo = Object.TypeInfo });
 				// Make sure the new row is tagged for update.
-				Storage->AddColumn<FTypedElementSyncFromWorldTag>(Row);
-
+				Interface->AddColumn<FTypedElementSyncFromWorldTag>(Row);
 				OnObjectAdded(Object.Object, Object.TypeInfo.Get(), Row);
 			});
 
diff --git a/Engine/Source/Runtime/TypedElementFramework/Public/Elements/Interfaces/TypedElementDataStorageInterface.h b/Engine/Source/Runtime/TypedElementFramework/Public/Elements/Interfaces/TypedElementDataStorageInterface.h
--- a/Engine/Source/Runtime/TypedElementFramework/Public/Elements/Interfaces/TypedElementDataStorageInterface.h
+++ b/Engine/Source/Runtime/TypedElementFramework/Public/Elements/Interfaces/TypedElementDataStorageInterface.h
@@ -50,7 +50,7 @@ class UTypedElementDataStorageInterface : public UInterface
  * multiple column lists are used. Note that the returned array view is only available while
  * this object is constructed, so care must be taken with functions that return a const array view.
  */
-template<typename... Columns>
+template<TypedElementDataStorage::TColumnType... Columns>
 struct TTypedElementColumnTypeList
 {
 	const UScriptStruct* ColumnTypes[sizeof...(Columns)] = { Columns::StaticStruct()... };
@@ -162,19 +162,19 @@ public:
 	 * at a time.
 	 */
 	virtual void AddColumns(TypedElementRowHandle Row, TConstArrayView<const UScriptStruct*> Columns) = 0;
-	template<typename... Columns>
+	template<TypedElementDataStorage::TColumnType... Columns>
 	void AddColumns(TypedElementRowHandle Row);
 
 	/** Removes a column from a row or does nothing if already removed. */
 	virtual void RemoveColumn(TypedElementRowHandle Row, const UScriptStruct* ColumnType) = 0;
-	template<typename Column>
+	template<TypedElementDataStorage::TColumnType Column>
 	void RemoveColumn(TypedElementRowHandle Row);
 	/**
 	 * Removes multiple columns from a row. This is typically more efficient than adding columns one
 	 * at a time.
 	 */
 	virtual void RemoveColumns(TypedElementRowHandle Row, TConstArrayView<const UScriptStruct*> Columns) = 0;
-	template<typename... Columns>
+	template<TypedElementDataStorage::TColumnType... Columns>
 	void RemoveColumns(TypedElementRowHandle Row);
 
 	/** 
@@ -190,39 +190,20 @@ public:
 		TConstArrayView<TypedElementRowHandle> Rows,
 		TConstArrayView<const UScriptStruct*> ColumnsToAdd,
 		TConstArrayView<const UScriptStruct*> ColumnsToRemove) = 0;
-
-	/**
-	 * Adds a new column to a row. If the column already exists it will be returned instead. If the column couldn't
-	 * be added or the column type points to a tag an nullptr will be returned.
-	 * This version allows a new object to be initialized before OnAdd observers are called as well as provide a
-	 * bespoke call to move the object from temporary locations to the final table location.
-	 */
-	virtual void* AddOrGetColumnData(TypedElementRowHandle Row, const UScriptStruct* ColumnType,
-		const TypedElementDataStorage::ColumnCreationCallbackRef& Initializer,
-		TypedElementDataStorage::ColumnCopyOrMoveCallback Relocator) = 0;
-	/**
-	 * Returns a pointer to the column of the given row or creates a new one if not found.
-	 * Enables type deduction of ColumnType from Column argument.
-	 *
-	 * For example, FTransformColumn added and deduced from second argument:
-	 * StorageInterface->AddOrGetColumn(Row, FTransformColumn{.Transform = Transform});
-	 */
-	template<typename ColumnType>
-	ColumnType* AddOrGetColumn(TypedElementRowHandle Row, ColumnType&& Column);
 	
 	/** Retrieves a pointer to the column of the given row or a nullptr if not found or if the column type is a tag. */
 	virtual void* GetColumnData(TypedElementRowHandle Row, const UScriptStruct* ColumnType) = 0;
 	virtual const void* GetColumnData(TypedElementRowHandle Row, const UScriptStruct* ColumnType) const = 0;
 	/** Returns a pointer to the column of the given row or a nullptr if the type couldn't be found or the row doesn't exist. */
-	template<typename ColumnType>
+	template<TypedElementDataStorage::TDataColumnType ColumnType>
 	ColumnType* GetColumn(TypedElementRowHandle Row);
-	template<typename ColumnType>
+	template<TypedElementDataStorage::TDataColumnType ColumnType>
 	const ColumnType* GetColumn(TypedElementRowHandle Row) const;
 	
 	/** Determines if the provided row contains the collection of columns and tags. */
 	virtual bool HasColumns(TypedElementRowHandle Row, TConstArrayView<const UScriptStruct*> ColumnTypes) const = 0;
 	virtual bool HasColumns(TypedElementRowHandle Row, TConstArrayView<TWeakObjectPtr<const UScriptStruct>> ColumnTypes) const = 0;
-	template<typename... ColumnTypes>
+	template<TypedElementDataStorage::TColumnType... ColumnTypes>
 	bool HasColumns(TypedElementRowHandle Row) const;
 
 	/** Determines if the columns in the row match the query conditions. */
@@ -365,19 +346,19 @@ void ITypedElementDataStorageInterface::AddColumn(TypedElementRowHandle Row)
 	AddColumn(Row, Column::StaticStruct());
 }
 
-template<typename Column>
+template<TypedElementDataStorage::TColumnType Column>
 void ITypedElementDataStorageInterface::RemoveColumn(TypedElementRowHandle Row)
 {
 	RemoveColumn(Row, Column::StaticStruct());
 }
 
-template<typename... Columns>
+template<TypedElementDataStorage::TColumnType... Columns>
 void ITypedElementDataStorageInterface::AddColumns(TypedElementRowHandle Row)
 {
 	AddColumns(Row, { Columns::StaticStruct()...});
 }
 
-template<typename... Columns>
+template<TypedElementDataStorage::TColumnType... Columns>
 void ITypedElementDataStorageInterface::RemoveColumns(TypedElementRowHandle Row)
 {
 	RemoveColumns(Row, { Columns::StaticStruct()...});
@@ -386,13 +367,7 @@ void ITypedElementDataStorageInterface::RemoveColumns(TypedElementRowHandle Row)
 template<TypedElementDataStorage::TDataColumnType ColumnType>
 void ITypedElementDataStorageInterface::AddColumn(TypedElementRowHandle Row, ColumnType&& Column)
 {
-	AddOrGetColumn<ColumnType>(Row, Forward<ColumnType>(Column));
-}
-
-template<typename ColumnType>
-ColumnType* ITypedElementDataStorageInterface::AddOrGetColumn(TypedElementRowHandle Row, ColumnType&& Column)
-{
-	return reinterpret_cast<ColumnType*>(AddOrGetColumnData(Row, ColumnType::StaticStruct(),
+	AddColumnData(Row, ColumnType::StaticStruct(),
 		[&Column](void* ColumnData, const UScriptStruct&)
 		{
 			if constexpr (std::is_move_constructible_v<ColumnType>)
@@ -414,22 +389,22 @@ ColumnType* ITypedElementDataStorageInterface::AddOrGetColumn(TypedElementRowHan
 			{
 				*reinterpret_cast<ColumnType*>(Destination) = *reinterpret_cast<ColumnType*>(Source);
 			}
-		}));
+		});
 }
 
-template<typename ColumnType>
+template<TypedElementDataStorage::TDataColumnType ColumnType>
 ColumnType* ITypedElementDataStorageInterface::GetColumn(TypedElementRowHandle Row)
 {
 	return reinterpret_cast<ColumnType*>(GetColumnData(Row, ColumnType::StaticStruct()));
 }
 
-template<typename ColumnType>
+template<TypedElementDataStorage::TDataColumnType ColumnType>
 const ColumnType* ITypedElementDataStorageInterface::GetColumn(TypedElementRowHandle Row) const
 {
 	return reinterpret_cast<const ColumnType*>(GetColumnData(Row, ColumnType::StaticStruct()));
 }
 
-template<typename... ColumnType>
+template<TypedElementDataStorage::TColumnType... ColumnType>
 bool ITypedElementDataStorageInterface::HasColumns(TypedElementRowHandle Row) const
 {
 	return HasColumns(Row, TConstArrayView<const UScriptStruct*>({ ColumnType::StaticStruct()... }));
