ForceInit Viewport OverlayColor to resolve uninitialized data issue

diff --git a/Engine/Source/Runtime/Engine/Public/PrimitiveSceneProxy.h b/Engine/Source/Runtime/Engine/Public/PrimitiveSceneProxy.h
--- a/Engine/Source/Runtime/Engine/Public/PrimitiveSceneProxy.h
+++ b/Engine/Source/Runtime/Engine/Public/PrimitiveSceneProxy.h
@@ -1571,7 +1571,7 @@ private:
 	uint64 HiddenEditorViews;
 
 	/** Color to blend over the object as an overlay in the viewport */
-	FColor OverlayColor;
+	FColor OverlayColor = FColor(EForceInit::ForceInitToZero);
 
 	/** Index of the color to use for the object's outline */
 	uint32 SelectionOutlineColorIndex : 8;
diff --git a/Engine/Source/Runtime/Engine/Public/PrimitiveSceneProxyDesc.h b/Engine/Source/Runtime/Engine/Public/PrimitiveSceneProxyDesc.h
--- a/Engine/Source/Runtime/Engine/Public/PrimitiveSceneProxyDesc.h
+++ b/Engine/Source/Runtime/Engine/Public/PrimitiveSceneProxyDesc.h
@@ -183,7 +183,7 @@ struct FPrimitiveSceneProxyDesc
 	uint64 HiddenEditorViews = 0;	
 
 #if WITH_EDITOR
-	FColor OverlayColor;
+	FColor OverlayColor = FColor(EForceInit::ForceInitToZero);
 #endif
 #if MESH_DRAW_COMMAND_STATS
 	FName MeshDrawCommandStatsCategory;
diff --git a/Engine/Source/Runtime/TypedElementFramework/Public/Elements/Columns/TypedElementViewportColumns.h b/Engine/Source/Runtime/TypedElementFramework/Public/Elements/Columns/TypedElementViewportColumns.h
--- a/Engine/Source/Runtime/TypedElementFramework/Public/Elements/Columns/TypedElementViewportColumns.h
+++ b/Engine/Source/Runtime/TypedElementFramework/Public/Elements/Columns/TypedElementViewportColumns.h
@@ -1,4 +1,4 @@
-﻿// Copyright Epic Games, Inc. All Rights Reserved.
+// Copyright Epic Games, Inc. All Rights Reserved.
 
 #pragma once
 
@@ -28,5 +28,5 @@ struct FTypedElementViewportOverlayColorColumn final : public FTypedElementDataS
 	GENERATED_BODY()
 
 	UPROPERTY()
-	FColor OverlayColor{};
+	FColor OverlayColor = FColor(EForceInit::ForceInitToZero);
 };
