Log External References to TypedElementHandles on destruction instead of crashing with fatal error

diff --git a/Engine/Source/Runtime/TypedElementFramework/Public/Elements/Framework/TypedElementData.h b/Engine/Source/Runtime/TypedElementFramework/Public/Elements/Framework/TypedElementData.h
--- a/Engine/Source/Runtime/TypedElementFramework/Public/Elements/Framework/TypedElementData.h
+++ b/Engine/Source/Runtime/TypedElementFramework/Public/Elements/Framework/TypedElementData.h
@@ -240,22 +240,6 @@ public:
 #endif	// UE_TYPED_ELEMENT_HAS_REFCOUNTING
 	}
 
-	void LogReferences() const
-	{
-#if UE_TYPED_ELEMENT_HAS_REFTRACKING
-		if (References)
-		{
-			References->LogReferences();
-		}
-		else
-		{
-			UE_LOG(LogCore, Error, TEXT("CVar 'TypedElements.EnableReferenceTracking' is disabled. Enable it to see reference tracking."));
-		}
-#else	// UE_TYPED_ELEMENT_HAS_REFTRACKING
-		UE_LOG(LogCore, Error, TEXT("UE_TYPED_ELEMENT_HAS_REFTRACKING is disabled. Enable it and recompile to see reference tracking."));
-#endif	// UE_TYPED_ELEMENT_HAS_REFTRACKING
-	}
-
 	void StoreDestructionRequestCallstack() const
 	{
 #if UE_TYPED_ELEMENT_HAS_REFTRACKING
@@ -266,14 +250,26 @@ public:
 #endif	// UE_TYPED_ELEMENT_HAS_REFTRACKING
 	}
 
-	void CheckNoExternalReferencesOnDestruction() const
+	void LogExternalReferencesOnDestruction() const
 	{
 #if DO_CHECK
 		const FTypedElementRefCount LocalRefCount = GetRefCount();
 		if (LocalRefCount > 1)
 		{
-			LogReferences();
-			UE_LOG(LogCore, Fatal, TEXT("Element '%s' is still externally referenced when being destroyed! Ref-count: %d; see above for reference information (if available)."), *GetDebugId(), LocalRefCount);
+#if UE_TYPED_ELEMENT_HAS_REFTRACKING
+			UE_LOG(LogCore, Warning, TEXT("Element '%s' is still externally referenced when being destroyed. Ref-count: %d;"), *GetDebugId(), LocalRefCount);
+			if (References)
+			{
+				References->LogReferences();
+			}
+			else
+			{
+				UE_LOG(LogCore, Warning, TEXT("Enable CVar 'TypedElements.EnableReferenceTracking' for Element reference tracking."));
+			}
+#else	// UE_TYPED_ELEMENT_HAS_REFTRACKING
+			UE_LOG(LogCore, Warning, TEXT("Element '%s' is still externally referenced when being destroyed. Ref-count: %d"), *GetDebugId(), LocalRefCount);
+			UE_LOG(LogCore, Warning, TEXT("UE_TYPED_ELEMENT_HAS_REFTRACKING is disabled. Enable it, recompile and enable CVar 'TypedElements.EnableReferenceTracking' for reference tracking."));
+#endif	// UE_TYPED_ELEMENT_HAS_REFTRACKING
 		}
 #endif	// DO_CHECK
 	}
@@ -560,7 +556,7 @@ public:
 		
 		TTypedElementInternalData<ElementDataType>& InternalData = InternalDataArray[InElementId];
 		checkf(InExpectedDataPtr == &InternalData, TEXT("Internal data pointer did not match the expected value! Does this handle belong to a different element registry?"));
-		InternalData.CheckNoExternalReferencesOnDestruction();
+		InternalData.LogExternalReferencesOnDestruction();
 		InternalData.Reset();
 		InternalDataFreeIndices.Add(InElementId);
 	}
@@ -655,7 +651,7 @@ public:
 
 		TTypedElementInternalData<void>& InternalData = InternalDataArray[InElementId];
 		checkf(InExpectedDataPtr == &InternalData, TEXT("Internal data pointer did not match the expected value! Does this handle belong to a different element registry?"));
-		InternalData.CheckNoExternalReferencesOnDestruction();
+		InternalData.LogExternalReferencesOnDestruction();
 		InternalData.Reset();
 		InternalDataFreeIndices.Add(InternalDataArrayIndex);
 	}
