Fixed an issue with the Index hashing of the FNames not being different when comparing the hash of "/AFolder" and "/AFolder_1"

Added a overload to be able generate an index hash from a SoftObjectPath.

diff --git a/Engine/Source/Runtime/TypedElementFramework/Public/Elements/Framework/TypedElementIndexHasher.h b/Engine/Source/Runtime/TypedElementFramework/Public/Elements/Framework/TypedElementIndexHasher.h
--- a/Engine/Source/Runtime/TypedElementFramework/Public/Elements/Framework/TypedElementIndexHasher.h
+++ b/Engine/Source/Runtime/TypedElementFramework/Public/Elements/Framework/TypedElementIndexHasher.h
@@ -7,6 +7,7 @@
 #include "Elements/Interfaces/TypedElementQueryStorageInterfaces.h"
 #include "UObject/NameTypes.h"
 #include "UObject/ObjectPtr.h"
+#include "UObject/SoftObjectPath.h"
 #include "UObject/StrongObjectPtr.h"
 #include "UObject/WeakObjectPtrTemplates.h"
 
@@ -25,6 +26,7 @@ namespace TypedElementDataStorage
 	inline IndexHash GenerateIndexHash(const FString& Object);
 	inline IndexHash GenerateIndexHash(FStringView Object);
 	inline IndexHash GenerateIndexHash(FName Object);
+	inline IndexHash GenerateIndexHash(const FSoftObjectPath& ObjectPath);
 } // namespace TypedElementDataStorage
 
 #include "Elements/Framework/TypedElementIndexHasher.inl"
\ No newline at end of file
diff --git a/Engine/Source/Runtime/TypedElementFramework/Public/Elements/Framework/TypedElementIndexHasher.inl b/Engine/Source/Runtime/TypedElementFramework/Public/Elements/Framework/TypedElementIndexHasher.inl
--- a/Engine/Source/Runtime/TypedElementFramework/Public/Elements/Framework/TypedElementIndexHasher.inl
+++ b/Engine/Source/Runtime/TypedElementFramework/Public/Elements/Framework/TypedElementIndexHasher.inl
@@ -3,6 +3,7 @@
 #pragma once
 
 #include "Hash/CityHash.h"
+#include "UObject/TopLevelAssetPath.h"
 
 namespace TypedElementDataStorage
 {
@@ -44,6 +45,17 @@ namespace TypedElementDataStorage
 	{
 		constexpr static const char SeedName[] = "FName";
 		static uint64 Seed = CityHash64(SeedName, sizeof(SeedName) - 1);
-		return CityHash128to64({ Seed, Object.GetComparisonIndex().ToUnstableInt() });
+		return CityHash128to64({ Seed, Object.ToUnstableInt() });
+	}
+
+	IndexHash GenerateIndexHash(const FSoftObjectPath& ObjectPath)
+	{
+		constexpr static const char SeedName[] = "FSoftObjectPath";
+		static uint64 Seed = CityHash64(SeedName, sizeof(SeedName) - 1);
+
+		FTopLevelAssetPath TopLevelAssetPath = ObjectPath.GetAssetPath();
+		uint64 Hash = CityHash128to64({ Seed, GenerateIndexHash(TopLevelAssetPath.GetPackageName()) });
+		Hash = CityHash128to64({ Hash, GenerateIndexHash(TopLevelAssetPath.GetAssetName()) });
+		return CityHash128to64({ Hash, GenerateIndexHash(ObjectPath.GetSubPathString()) });
 	}
 } // namespace TypedElementDataStorage
