Integrate serialized placeholder-typed entity-owned objects into the scene outliner database.

diff --git a/Engine/Source/Runtime/TypedElementFramework/Public/Elements/Columns/TypedElementCompatibilityColumns.h b/Engine/Source/Runtime/TypedElementFramework/Public/Elements/Columns/TypedElementCompatibilityColumns.h
--- a/Engine/Source/Runtime/TypedElementFramework/Public/Elements/Columns/TypedElementCompatibilityColumns.h
+++ b/Engine/Source/Runtime/TypedElementFramework/Public/Elements/Columns/TypedElementCompatibilityColumns.h
@@ -56,3 +56,28 @@ struct FTypedElementWorldColumn final : public FTypedElementDataStorageColumn
 
 	TWeakObjectPtr<UWorld> World;
 };
+
+/**
+ * Tag to signal that a row represents a property bag "placeholder-typed" object reference.
+ * This object has an unknown base type and will generally be associated with a "property
+ * bag" containing any serialized data that was loaded for it as a set of "loose" properties.
+ */
+USTRUCT(meta = (DisplayName = "Property bag placeholder"))
+struct FTypedElementPropertyBagPlaceholderTag final : public FTypedElementDataStorageTag
+{
+	GENERATED_BODY()
+};
+
+/**
+ * A column that stores an inferred base type for placeholder-typed object references. Note that
+ * this will not be the same as the placeholder object type (stored in the class type info column).
+ * This column can be used to query for a "base type" determined from the serialization context,
+ * for systems that need to look/behave differently based on an inferred base type context (e.g. UI).
+ */
+USTRUCT(meta = (DisplayName = "Placeholder type info"))
+struct FTypedElementPropertyBagPlaceholderTypeInfoColumn final : public FTypedElementDataStorageColumn
+{
+	GENERATED_BODY()
+
+	TWeakObjectPtr<const UStruct> InferredBaseType;
+};
\ No newline at end of file
