TEDS UI: Attribute Binding
Add functionality to bind TEDS row, column pairs to TAttributes directly and by providing a conversion function
Add some unit tests for attribute binders

diff --git a/Engine/Plugins/Experimental/TypedElementsDataStorage/Source/TypedElementsDataStorageUI/Private/TypedElementsDataStorageUI.cpp b/Engine/Plugins/Experimental/TypedElementsDataStorage/Source/TypedElementsDataStorageUI/Private/TypedElementsDataStorageUI.cpp
--- a/Engine/Plugins/Experimental/TypedElementsDataStorage/Source/TypedElementsDataStorageUI/Private/TypedElementsDataStorageUI.cpp
+++ b/Engine/Plugins/Experimental/TypedElementsDataStorage/Source/TypedElementsDataStorageUI/Private/TypedElementsDataStorageUI.cpp
@@ -10,7 +10,9 @@
 
 #define LOCTEXT_NAMESPACE "FTypedElementsDataStorageModule"
 
-FAutoConsoleCommandWithOutputDevice PrintWidgetPurposesConsoleCommand(
+namespace UE::TypedElementsDataStorageUI::Private
+{
+	FAutoConsoleCommandWithOutputDevice PrintWidgetPurposesConsoleCommand(
 	TEXT("TEDS.UI.PrintWidgetPurposes"),
 	TEXT("Prints a list of all the known widget purposes."),
 	FConsoleCommandWithOutputDeviceDelegate::CreateLambda([](FOutputDevice& Output)
@@ -28,6 +30,16 @@ FAutoConsoleCommandWithOutputDevice PrintWidgetPurposesConsoleCommand(
 			}
 		}));
 
+	static bool bUseNewTEDSUIWidgets = false;
+	FAutoConsoleVariableRef UseNewWidgetsCvar(
+		TEXT("TEDS.UI.UseNewWidgets"),
+		bUseNewTEDSUIWidgets,
+		TEXT("If true, TEDS UI will use new attribute binding driven widgets (needs to be set at startup)")
+	);
+}
+
+
+
 void FTypedElementsDataStorageUiModule::StartupModule()
 {
 }
diff --git a/Engine/Source/Runtime/TypedElementFramework/Private/Elements/Framework/TypedElementAttributeBinding.cpp b/Engine/Source/Runtime/TypedElementFramework/Private/Elements/Framework/TypedElementAttributeBinding.cpp
new file mode 100644
--- /dev/null
+++ b/Engine/Source/Runtime/TypedElementFramework/Private/Elements/Framework/TypedElementAttributeBinding.cpp
@@ -0,0 +1,18 @@
+// Copyright Epic Games, Inc. All Rights Reserved.
+
+#include "Elements/Framework/TypedElementAttributeBinding.h"
+
+#include "Elements/Framework/TypedElementRegistry.h"
+
+namespace UE::EditorDataStorage
+{
+	FAttributeBinder::FAttributeBinder(TypedElementDataStorage::RowHandle InTargetRow)
+		: TargetRow(InTargetRow)
+	{
+		if (ensureMsgf(UTypedElementRegistry::GetInstance()->AreDataStorageInterfacesSet(), 
+				TEXT("The TypedElementsDataStorage plugin needs to be enabled to use attribute bindings.")))
+		{
+			DataStorage = UTypedElementRegistry::GetInstance()->GetMutableDataStorage();
+		}
+	}
+}
diff --git a/Engine/Source/Runtime/TypedElementFramework/Public/Elements/Framework/TypedElementAttributeBinding.h b/Engine/Source/Runtime/TypedElementFramework/Public/Elements/Framework/TypedElementAttributeBinding.h
new file mode 100644
--- /dev/null
+++ b/Engine/Source/Runtime/TypedElementFramework/Public/Elements/Framework/TypedElementAttributeBinding.h
@@ -0,0 +1,89 @@
+// Copyright Epic Games, Inc. All Rights Reserved.
+
+#pragma once
+
+#include "Elements/Common/TypedElementHandles.h"
+#include "Elements/Interfaces/TypedElementDataStorageInterface.h"
+#include "Misc/Attribute.h"
+
+class ITypedElementDataStorageInterface;
+
+namespace UE::EditorDataStorage
+{
+	// Concept for a conversion function used by the attribute binder to bind a column data member to an attribute of a different type
+	template <typename FunctionType, typename ArgumentType>
+	concept AttributeBinderInvocable = std::is_invocable_v<std::decay_t<FunctionType>, const ArgumentType&>;
+
+	/**
+	 * Builder class that can be used as a shorthand to bind data inside a TEDS row, column pair to a TAttribute so the attribute updates if the data
+	 * in the column is changed.
+	 * 
+	 * Usage Example:
+	 * 
+	 * TypedElementDataStorage::FAttributeBinder Binder(RowHandle);
+	 * TAttribute<int> TestAttribute(Binder.BindData(&FTestColumnInt::TestInt))
+	 */
+	class FAttributeBinder
+	{
+	public:
+
+
+
+		// Create an attribute binder for a given row
+		TYPEDELEMENTFRAMEWORK_API FAttributeBinder(TypedElementDataStorage::RowHandle InTargetRow);
+
+		/**
+		 * Bind a specific data member inside a TEDS column to an attribute of the same type as the data
+		 *
+		 * @param Variable The data member inside a column to be bound
+		 * @return A TAttribute bound to the row, column pair specified
+		 *
+		 * Example:
+		 * TypedElementDataStorage::FAttributeBinder Binder(RowHandle);
+		 * TAttribute<FString> TestAttribute(Binder.BindData(&FTypedElementLabelColumn::Label))
+		 */
+		template <typename AttributeType, TypedElementDataStorage::TDataColumnType ColumnType>
+		TAttribute<AttributeType> BindData(AttributeType ColumnType::* Variable);
+
+		/**
+		 * Bind a specific data member inside a TEDS column to an attribute of a different type than the data by providing a conversion function
+		 *
+		 * @param Variable The data member inside a column to be bound
+		 * @param Converter Conversion function to convert from DataType -> AttributeType
+		 * @return A TAttribute bound to the row, column pair specified
+		 *
+		 * Example:
+		 * TypedElementDataStorage::FAttributeBinder Binder(RowHandle);
+		 * TAttribute<FText> TestAttribute(Binder.BindData(&FTypedElementLabelColumn::Label),
+		 *                                 [](const FString& Data)
+		 *                                   {
+		 *                                      return FText::FromString(Data);
+		 *                                   }
+		 *                                 ));
+		 */
+		template <typename AttributeType, typename DataType, TypedElementDataStorage::TDataColumnType ColumnType>
+		TAttribute<AttributeType> BindData(DataType ColumnType::* Variable, const TFunction<AttributeType(const DataType&)>& Converter);
+
+		/**
+		 * Overload for the conversion binder to accept lambdas instead of TFunctions
+		 *
+		 * @param Variable The data member inside a column to be bound
+		 * @param Converter Conversion function to convert from DataType -> AttributeType
+		 * @return A TAttribute bound to the row, column pair specified
+		 */
+		template <typename DataType, TypedElementDataStorage::TDataColumnType ColumnType, typename FunctionType>
+			requires AttributeBinderInvocable<FunctionType, DataType>
+		auto BindData(DataType ColumnType::* Variable, FunctionType Converter);
+		
+	private:
+
+		// The target row for this binder
+		TypedElementDataStorage::RowHandle TargetRow;
+
+		// A ptr to the data storage for quick access
+		ITypedElementDataStorageInterface* DataStorage;
+	};
+}
+
+
+#include "TypedElementAttributeBinding.inl"
\ No newline at end of file
diff --git a/Engine/Source/Runtime/TypedElementFramework/Public/Elements/Framework/TypedElementAttributeBinding.inl b/Engine/Source/Runtime/TypedElementFramework/Public/Elements/Framework/TypedElementAttributeBinding.inl
new file mode 100644
--- /dev/null
+++ b/Engine/Source/Runtime/TypedElementFramework/Public/Elements/Framework/TypedElementAttributeBinding.inl
@@ -0,0 +1,181 @@
+// Copyright Epic Games, Inc. All Rights Reserved.
+
+#pragma once
+
+#include "TypedElementRegistry.h"
+
+namespace UE::EditorDataStorage
+{
+	// A property that can be directly be accessed from an object
+	template <typename PropertyType>
+	struct DirectProperty final
+	{
+		// Offset to a particular data member inside the object
+		size_t Offset;
+
+		PropertyType* Get(void* Object) const
+		{ 
+			return reinterpret_cast<PropertyType*>((reinterpret_cast<char*>(Object) + Offset));
+		}
+		const PropertyType* Get(const void* Object) const
+		{ 
+			return reinterpret_cast<const PropertyType*>((reinterpret_cast<const char*>(Object) + Offset)); 
+		}
+	};
+
+	// A property that goes through a conversion function before being accessed from the object
+	template <typename PropertyType>
+	struct ConvertibleProperty final
+	{
+		// Conversion function
+		TFunction<PropertyType(const void*)> Converter;
+
+		// Cache to avoid a copy when returning if not necessary
+		mutable PropertyType Cache;
+
+		PropertyType* Get(void* Object) const
+		{ 
+			Cache = Converter(Object);
+			return &Cache;
+		}
+		const PropertyType* Get(const void* Object) const
+		{ 
+			Cache = Converter(Object);
+			return &Cache;
+		}
+	};
+
+	template <typename PropertyType>
+	class Property
+	{
+	public:
+		
+		// Bind this property directly
+		template <TypedElementDataStorage::TDataColumnType ObjectType>
+		void Bind(PropertyType ObjectType::* Variable)
+		{
+			DirectProperty<PropertyType> Result;
+			Result.Offset = reinterpret_cast<size_t>(&(reinterpret_cast<const ObjectType*>(0)->*Variable));
+			
+			ObjectTypeInfo = ObjectType::StaticStruct();
+			InternalProperty = InternalPropertyType(TInPlaceType<DirectProperty<PropertyType>>(), Result);;
+		}
+
+		// Bind this property using a conversion function
+		template <typename InputType, TypedElementDataStorage::TDataColumnType ObjectType>
+		void Bind(InputType ObjectType::* Variable, TFunction<PropertyType(const InputType&)> Converter)
+		{
+			ConvertibleProperty<PropertyType> Result;
+			Result.Converter = [Converter, Variable](const void* Object)
+			{
+				return Converter(reinterpret_cast<const ObjectType*>(Object)->*Variable);
+			};
+			
+			ObjectTypeInfo = ObjectType::StaticStruct();
+			InternalProperty = InternalPropertyType(TInPlaceType<ConvertibleProperty<PropertyType>>(), Result);;
+		}
+
+		// Get the bound property for the specified object
+		template <TypedElementDataStorage::TDataColumnType ObjectType>
+		PropertyType& Get(ObjectType& Object) const
+		{
+			return Get(&Object, ObjectType::StaticStruct());
+		}
+
+		// Get the bound property for the specified object
+		template <TypedElementDataStorage::TDataColumnType ObjectType>
+		const PropertyType& Get(const ObjectType& Object) const
+		{
+			return Get(&Object, ObjectType::StaticStruct());
+		}
+		
+		// Get the bound property for this specified object ptr by providing type information about the object
+		PropertyType& Get(void* Object, const UScriptStruct* Type) const
+		{
+			ensureMsgf(Type == ObjectTypeInfo, TEXT("Provided object type (%s) did not match bound object type (%s)."),
+				*Type->GetFName().ToString(), *ObjectTypeInfo->GetFName().ToString());
+
+			return *Visit([Object](auto&& Prop) { return Prop.Get(Object); }, InternalProperty);
+		}
+		
+		// Get the bound property for this specified object ptr by providing type information about the object
+		const PropertyType& Get(const void* Object, const UScriptStruct* Type) const
+		{
+			ensureMsgf(Type == ObjectTypeInfo, TEXT("Provided object type (%s) did not match bound object type (%s)."),
+				*Type->GetFName().ToString(), *ObjectTypeInfo->GetFName().ToString());
+
+			return *Visit([Object](auto&& Prop) { return Prop.Get(Object); }, InternalProperty);
+		}
+
+	private:
+
+		// Internally we could be storing a direct property or a convertible property
+		using InternalPropertyType = TVariant<DirectProperty<PropertyType>, ConvertibleProperty<PropertyType>>;
+
+		// The actual property
+		InternalPropertyType InternalProperty;
+
+		// The type of the object we are bound to
+		UScriptStruct* ObjectTypeInfo = nullptr;
+	};
+	
+	template <typename AttributeType, TypedElementDataStorage::TDataColumnType ColumnType>
+	TAttribute<AttributeType> FAttributeBinder::BindData(AttributeType ColumnType::* Variable)
+	{
+		if(!DataStorage)
+		{
+			return TAttribute<AttributeType>();
+		}
+		
+		// Create a direct property and bind it to the given variable
+		Property<AttributeType> Prop;
+		Prop.Bind(Variable);
+
+		// We don't want any references to this in the lambda because binders are designed to be used and destructed on the stack
+		return TAttribute<AttributeType>::CreateLambda([Property = Prop, Storage = DataStorage, Row = TargetRow]()
+		{
+			// Get the column from the given row and use that to return the stored property
+			if(ColumnType* Column = Storage->GetColumn<ColumnType>(Row))
+			{
+				return Property.Get(Column, ColumnType::StaticStruct());
+			}
+			return AttributeType();
+		});
+	}
+
+	template <typename AttributeType, typename DataType, TypedElementDataStorage::TDataColumnType ColumnType>
+	TAttribute<AttributeType> FAttributeBinder::BindData(DataType ColumnType::* Variable, const TFunction<AttributeType(const DataType&)>& Converter)
+	{
+		if(!DataStorage)
+		{
+			return TAttribute<AttributeType>();
+		}
+		
+		// Create a convertible property and bind it to the given variable
+		Property<AttributeType> Prop;
+		Prop.Bind(Variable, Converter);
+	
+		return TAttribute<AttributeType>::CreateLambda([Property = Prop, Storage = DataStorage, Row = TargetRow]()
+		{
+			if(ColumnType* Column = Storage->GetColumn<ColumnType>(Row))
+			{
+				return Property.Get(Column, ColumnType::StaticStruct());
+			}
+
+			return AttributeType();
+		});
+	}
+	
+	template <typename DataType, TypedElementDataStorage::TDataColumnType ColumnType, typename FunctionType>
+		requires AttributeBinderInvocable<FunctionType, DataType>
+	auto FAttributeBinder::BindData(DataType ColumnType::* Variable, FunctionType Converter)
+	{
+		// Deduce the attribute type from the return value of the converter function
+		using AttributeType = decltype(Converter(std::declval<DataType>()));
+		
+		return BindData<AttributeType, DataType>(Variable, TFunction<AttributeType(const DataType&)>(Converter));
+	}
+	
+}
+
+	
\ No newline at end of file
diff --git a/Engine/Source/Runtime/TypedElementFramework/Tests/Elements/Framework/TypedElementAttributeBindingTests.cpp b/Engine/Source/Runtime/TypedElementFramework/Tests/Elements/Framework/TypedElementAttributeBindingTests.cpp
new file mode 100644
--- /dev/null
+++ b/Engine/Source/Runtime/TypedElementFramework/Tests/Elements/Framework/TypedElementAttributeBindingTests.cpp
@@ -0,0 +1,199 @@
+// Copyright Epic Games, Inc. All Rights Reserved.
+
+
+#include "Elements/Framework/TypedElementAttributeBinding.h"
+#if WITH_TESTS
+#include "TypedElementTestColumns.h"
+#include "Elements/Framework/TypedElementRegistry.h"
+#include "Misc/AutomationTest.h"
+
+namespace TypedElementDataStorageTests
+{
+	BEGIN_DEFINE_SPEC(TypedElementAttributeBindingTestsFixture, "TypedElementsDataStorage.AttributeBinding", EAutomationTestFlags::EditorContext | EAutomationTestFlags::EngineFilter)
+
+		ITypedElementDataStorageInterface* TedsInterface = nullptr;
+		const FName TestTableName = TEXT("TestTable_AttributeBinding");
+		TypedElementDataStorage::TableHandle TestTableHandle = TypedElementDataStorage::InvalidTableHandle;
+		TypedElementDataStorage::RowHandle TestRowHandle = TypedElementDataStorage::InvalidRowHandle;
+
+		TypedElementDataStorage::TableHandle RegisterTestTable() const
+		{
+			const TypedElementDataStorage::TableHandle Table = TedsInterface->FindTable(TestTableName);
+			
+			if (Table != TypedElementDataStorage::InvalidTableHandle)
+			{
+				return Table;
+			}
+			
+			return TedsInterface->RegisterTable(
+			{
+				FTestColumnInt::StaticStruct(),
+				FTestColumnString::StaticStruct()
+			},
+			TestTableName);
+		}
+	
+		TypedElementDataStorage::RowHandle CreateTestRow(TypedElementDataStorage::TableHandle InTableHandle) const
+		{
+			const TypedElementDataStorage::RowHandle RowHandle = TedsInterface->AddRow(InTableHandle);
+			return RowHandle;
+		}
+	
+		void CleanupTestRow(TypedElementDataStorage::RowHandle InRowHandle) const
+		{
+			TedsInterface->RemoveRow(InRowHandle);
+		}
+	
+	END_DEFINE_SPEC(TypedElementAttributeBindingTestsFixture)
+
+	void TypedElementAttributeBindingTestsFixture::Define()
+	{
+		BeforeEach([this]()
+		{
+			UTypedElementRegistry* TypedElementRegistry = UTypedElementRegistry::GetInstance();
+			TedsInterface = TypedElementRegistry->GetMutableDataStorage();
+			TestTrue("", TedsInterface != nullptr);
+			
+			TestTableHandle = RegisterTestTable();
+			TestNotEqual("Expecting valid table handle", TestTableHandle, TypedElementDataStorage::InvalidTableHandle);
+
+			TestRowHandle = CreateTestRow(TestTableHandle);
+			TestFalse("Expect valid row handle", TestRowHandle == TypedElementDataStorage::InvalidRowHandle);
+		});
+
+		Describe("", [this]()
+		{
+			Describe("Integer Attribute", [this]()
+			{
+				Describe("Direct integer attribute", [this]()
+				{
+					It("Direct attribute should update on updating column value", [this]()
+					{
+						constexpr int InitialValue = 10;
+						constexpr int UpdatedValue = 20;
+						
+						// Add the test int column to the test row
+						TedsInterface->AddColumn(TestRowHandle, FTestColumnInt{.TestInt =  InitialValue});
+						FTestColumnInt* TestColumnInt = TedsInterface->GetColumn<FTestColumnInt>(TestRowHandle);
+
+						TestNotNull("Expecting Valid Column", TestColumnInt);
+
+						// Create an int attribute and bind it
+						UE::EditorDataStorage::FAttributeBinder Binder(TestRowHandle);
+						const TAttribute TestAttribute(Binder.BindData(&FTestColumnInt::TestInt));
+
+						TestEqual("Expecting attribute value to match column value before modification", TestAttribute.Get(), TestColumnInt->TestInt);
+
+						// Modify the value in the column
+						TestColumnInt->TestInt = UpdatedValue;
+
+						TestEqual("Expecting attribute value to update after modification", TestAttribute.Get(), UpdatedValue);
+						TestEqual("Expecting attribute value to match column value after modification", TestAttribute.Get(), TestColumnInt->TestInt);
+					});
+				});
+
+				Describe("Float attribute bound to integer column data", [this]()
+				{
+					It("Converted attribute should update on updating column value", [this]()
+					{
+						constexpr int InitialValue = 10;
+						constexpr int UpdatedValue = 20;
+
+						// Add the test int column to the test row
+						TedsInterface->AddColumn(TestRowHandle, FTestColumnInt{.TestInt =  InitialValue});
+						FTestColumnInt* TestColumnInt = TedsInterface->GetColumn<FTestColumnInt>(TestRowHandle);
+
+						TestNotNull("Expecting valid column", TestColumnInt);
+
+						// Create a float attribute and bind it by providing a conversion function
+						UE::EditorDataStorage::FAttributeBinder Binder(TestRowHandle);
+						const TAttribute<float> TestAttribute(Binder.BindData(&FTestColumnInt::TestInt,
+							[](const int& Data)
+							{
+								return static_cast<float>(Data);
+							}));
+
+						TestEqual("Expecting attribute value to match column value before modification", static_cast<int>(TestAttribute.Get()), TestColumnInt->TestInt);
+
+						// Modify the value in the column
+						TestColumnInt->TestInt = UpdatedValue;
+
+						TestEqual("Expecting attribute value to update after modification", static_cast<int>(TestAttribute.Get()), UpdatedValue);
+						TestEqual("Expecting attribute value to match column value after modification", static_cast<int>(TestAttribute.Get()), TestColumnInt->TestInt);
+					});
+				});
+			});
+
+			Describe("String Attribute", [this]()
+			{
+				Describe("Direct string attribute", [this]()
+				{
+					It("Direct attribute should update on updating column value", [this]()
+					{
+						const FString InitialValue(TEXT("Test String"));
+						const FString UpdatedValue(TEXT("Test string after modification"));
+						
+						// Add the test int column to the test row
+						TedsInterface->AddColumn(TestRowHandle, FTestColumnString{.TestString = InitialValue});
+						FTestColumnString* TestColumnString = TedsInterface->GetColumn<FTestColumnString>(TestRowHandle);
+
+						TestNotNull("Expecting valid column", TestColumnString);
+
+						// Create an int attribute and bind it
+						UE::EditorDataStorage::FAttributeBinder Binder(TestRowHandle);
+						const TAttribute TestAttribute(Binder.BindData(&FTestColumnString::TestString));
+
+						TestEqual("Expecting attribute value to match column value before modification", TestAttribute.Get(), TestColumnString->TestString);
+
+						// Modify the value in the column
+						TestColumnString->TestString = UpdatedValue;
+						
+						TestEqual("Expecting attribute value to update after modification", TestAttribute.Get(), UpdatedValue);
+						TestEqual("Expecting attribute value to match column value after modification", TestAttribute.Get(), TestColumnString->TestString);
+					});
+				});
+
+				Describe("Text attribute bound to string column data", [this]()
+				{
+					It("Converted attribute should update on updating column value", [this]()
+					{
+						const FString InitialValue(TEXT("Test String"));
+						const FString UpdatedValue(TEXT("Test string after modification"));
+
+						// Add the test int column to the test row
+						TedsInterface->AddColumn(TestRowHandle, FTestColumnString{.TestString = InitialValue});
+						FTestColumnString* TestColumnString = TedsInterface->GetColumn<FTestColumnString>(TestRowHandle);
+
+						TestNotNull("Expecting valid column", TestColumnString);
+
+						// Create an int attribute and bind it
+						UE::EditorDataStorage::FAttributeBinder Binder(TestRowHandle);
+						const TAttribute<FText> TestAttribute(Binder.BindData(&FTestColumnString::TestString,
+							[](const FString& Data)
+							{
+								return FText::FromString(Data);
+							}));
+
+						TestEqual("Expecting attribute value to match column value before modification", TestAttribute.Get().ToString(), TestColumnString->TestString);
+
+						// Modify the value in the column
+						TestColumnString->TestString = UpdatedValue;
+
+						TestEqual("Expecting attribute value to update after modification", TestAttribute.Get().ToString(), UpdatedValue);
+						TestEqual("Expecting attribute value to match column value after modification", TestAttribute.Get().ToString(), TestColumnString->TestString);
+					});
+				});
+			});
+		});
+
+		AfterEach([this]()
+		{
+			CleanupTestRow(TestRowHandle);
+			TestRowHandle = TypedElementDataStorage::InvalidRowHandle;
+			TestTableHandle = TypedElementDataStorage::InvalidTableHandle;
+			TedsInterface = nullptr;
+		});
+	}
+}
+
+#endif // WITH_TESTS
\ No newline at end of file
diff --git a/Engine/Source/Runtime/TypedElementFramework/Tests/Elements/Framework/TypedElementTestColumns.h b/Engine/Source/Runtime/TypedElementFramework/Tests/Elements/Framework/TypedElementTestColumns.h
--- a/Engine/Source/Runtime/TypedElementFramework/Tests/Elements/Framework/TypedElementTestColumns.h
+++ b/Engine/Source/Runtime/TypedElementFramework/Tests/Elements/Framework/TypedElementTestColumns.h
@@ -49,6 +49,25 @@ struct FTestColumnG final : public FTypedElementDataStorageColumn
 	GENERATED_BODY()
 };
 
+USTRUCT(meta = (DisplayName = "ColumnInt"))
+struct FTestColumnInt final : public FTypedElementDataStorageColumn
+{
+	GENERATED_BODY()
+	
+	UPROPERTY()
+	int TestInt = 0;
+};
+
+USTRUCT(meta = (DisplayName = "ColumnString"))
+struct FTestColumnString final : public FTypedElementDataStorageColumn
+{
+	GENERATED_BODY()
+	
+	UPROPERTY()
+    FString TestString;
+};
+
+
 USTRUCT(meta = (DisplayName = "TagA"))
 struct FTestTagColumnA final : public FTypedElementDataStorageTag
 {
