Fixed an incorrect test for TEDS' Query Condition tests.

diff --git a/Engine/Source/Runtime/TypedElementFramework/Tests/Elements/Framework/TypedElementQueryBuilderTests.cpp b/Engine/Source/Runtime/TypedElementFramework/Tests/Elements/Framework/TypedElementQueryBuilderTests.cpp
--- a/Engine/Source/Runtime/TypedElementFramework/Tests/Elements/Framework/TypedElementQueryBuilderTests.cpp
+++ b/Engine/Source/Runtime/TypedElementFramework/Tests/Elements/Framework/TypedElementQueryBuilderTests.cpp
@@ -104,7 +104,8 @@ TEST_CASE_NAMED(FTypedElementQueryConditions_NoColumn, "TypedElementQueryBuilder
 	FQueryConditions Example;
 	
 	CHECK(Example.MinimumColumnMatchRequired() == 0);
-	CHECK(TestMatching(Example, { FColumn<FTestColumnA>() }, false));
+	// Since there are no restrictions provided in the query, all input passes.
+	CHECK(TestMatching(Example, { FColumn<FTestColumnA>() }, true));
 }
 
 TEST_CASE_NAMED(FTypedElementQueryConditions_OneColumn, "TypedElementQueryBuilder::FTypedElementQueryConditions_OneColumn", "[ApplicationContextMask][EngineFilter]")
