[Typed Elements] Add additional checks and messaging around a suspected intermittent PIE Test crash

#rnx

diff --git a/Engine/Source/Runtime/TypedElementFramework/Public/Elements/Framework/TypedElementData.h b/Engine/Source/Runtime/TypedElementFramework/Public/Elements/Framework/TypedElementData.h
--- a/Engine/Source/Runtime/TypedElementFramework/Public/Elements/Framework/TypedElementData.h
+++ b/Engine/Source/Runtime/TypedElementFramework/Public/Elements/Framework/TypedElementData.h
@@ -523,6 +523,8 @@ public:
 			? InternalDataFreeIndices.Pop(EAllowShrinking::No)
 			: InternalDataArray.Add();
 
+		checkf(InternalDataArray.IsValidIndex(InOutElementId), TEXT("Expected a valid index to use. Index %d is out of bounds of array size %d"), InOutElementId, InternalDataArray.Num());
+
 		TTypedElementInternalData<ElementDataType>& InternalData = InternalDataArray[InOutElementId];
 		InternalData.Initialize(InTypeId, InOutElementId);
 		return InternalData;
@@ -552,13 +554,14 @@ public:
 	{
 		FWriteScopeLock InternalDataLock(InternalDataRW);
 
-		checkSlow(InternalDataArray.IsValidIndex(InElementId));
+		check(InternalDataArray.IsValidIndex(InElementId));
 		
 		TTypedElementInternalData<ElementDataType>& InternalData = InternalDataArray[InElementId];
 		checkf(InExpectedDataPtr == &InternalData, TEXT("Internal data pointer did not match the expected value! Does this handle belong to a different element registry?"));
 		InternalData.LogExternalReferencesOnDestruction();
 		InternalData.Reset();
-		InternalDataFreeIndices.Add(InElementId);
+		const TArray<int32>::SizeType FreeIndicesIndex = InternalDataFreeIndices.AddUnique(InElementId);
+		checkf(FreeIndicesIndex == InternalDataFreeIndices.Num() - 1, TEXT("Double free deteced when removing element handle at index %d"), FreeIndicesIndex);
 	}
 
 	const TTypedElementInternalData<ElementDataType>& GetDataForElement(const FTypedHandleElementId InElementId) const
@@ -607,12 +610,14 @@ public:
 		FWriteScopeLock InternalDataLock(InternalDataRW);
 
 		checkSlow(InOutElementId >= 0);
-		checkSlow(!ElementIdToArrayIndex.Contains(InOutElementId));
+		check(!ElementIdToArrayIndex.Contains(InOutElementId));
 
 		const int32 InternalDataArrayIndex = InternalDataFreeIndices.Num() > 0
 			? InternalDataFreeIndices.Pop(EAllowShrinking::No)
 			: InternalDataArray.Add();
 
+		checkf(InternalDataArray.IsValidIndex(InternalDataArrayIndex), TEXT("Expected a valid index to reuse. Index %d is out of bounds of array size %d"), InternalDataArrayIndex, InternalDataArray.Num());
+
 		ElementIdToArrayIndex.Add(InOutElementId, InternalDataArrayIndex);
 
 		TTypedElementInternalData<void>& InternalData = InternalDataArray[InOutElementId];
@@ -653,7 +658,8 @@ public:
 		checkf(InExpectedDataPtr == &InternalData, TEXT("Internal data pointer did not match the expected value! Does this handle belong to a different element registry?"));
 		InternalData.LogExternalReferencesOnDestruction();
 		InternalData.Reset();
-		InternalDataFreeIndices.Add(InternalDataArrayIndex);
+		const TArray<int32>::SizeType FreeIndicesIndex = InternalDataFreeIndices.AddUnique(InternalDataArrayIndex);
+		checkf(FreeIndicesIndex == InternalDataFreeIndices.Num() - 1, TEXT("Double free deteced when removing element handle at index %d"), FreeIndicesIndex);
 	}
 
 	const TTypedElementInternalData<void>& GetDataForElement(const FTypedHandleElementId InElementId) const
