TEDS UI: Add functionality to provide a default value for attribute bindings which is used when the column isn't present

diff --git a/Engine/Plugins/Experimental/TypedElementsDataStorage/Source/TypedElementsDataStorage/Tests/TypedElementAttributeBindingTests.cpp b/Engine/Plugins/Experimental/TypedElementsDataStorage/Source/TypedElementsDataStorage/Tests/TypedElementAttributeBindingTests.cpp
--- a/Engine/Plugins/Experimental/TypedElementsDataStorage/Source/TypedElementsDataStorage/Tests/TypedElementAttributeBindingTests.cpp
+++ b/Engine/Plugins/Experimental/TypedElementsDataStorage/Source/TypedElementsDataStorage/Tests/TypedElementAttributeBindingTests.cpp
@@ -9,7 +9,7 @@
 
 namespace TypedElementDataStorageTests
 {
-	BEGIN_DEFINE_SPEC(TypedElementAttributeBindingTestsFixture, "TypedElementsDataStorage.AttributeBinding", EAutomationTestFlags::EditorContext | EAutomationTestFlags::EngineFilter)
+	BEGIN_DEFINE_SPEC(TypedElementAttributeBindingTestsFixture, "EditorDataStorage.AttributeBinding", EAutomationTestFlags::EditorContext | EAutomationTestFlags::EngineFilter)
 
 		ITypedElementDataStorageInterface* TedsInterface = nullptr;
 		const FName TestTableName = TEXT("TestTable_AttributeBinding");
@@ -184,6 +184,34 @@ namespace TypedElementDataStorageTests
 					});
 				});
 			});
+
+			Describe("Default Value", [this]()
+			{
+				It("Default value should be used when column isn't present", [this]()
+				{
+					constexpr int DefaultValue = 10;
+					
+					UE::EditorDataStorage::FAttributeBinder Binder(TestRowHandle);
+
+					// Create an int attribute and directly bind it
+					const TAttribute TestIntAttribute(Binder.BindData(&FTestColumnInt::TestInt, DefaultValue));
+
+					// Create a float attribute and bind it by providing a conversion function
+					const TAttribute<float> TestFloatAttribute(Binder.BindData(&FTestColumnInt::TestInt,
+						[](const int& Data)
+						{
+							return static_cast<float>(Data);
+						}, DefaultValue));
+
+
+					// Remove FTestColumnInt from TestRowHandle so the default value is used
+					TedsInterface->RemoveColumn(TestRowHandle, FTestColumnInt::StaticStruct());
+
+					TestEqual("Expecting int attribute value to match default value", TestIntAttribute.Get(), DefaultValue);
+					TestEqual("Expecting float attribute value to match default value", static_cast<int>(TestFloatAttribute.Get()), DefaultValue);
+
+				});
+			});
 		});
 
 		AfterEach([this]()
diff --git a/Engine/Source/Runtime/TypedElementFramework/Public/Elements/Framework/TypedElementAttributeBinding.h b/Engine/Source/Runtime/TypedElementFramework/Public/Elements/Framework/TypedElementAttributeBinding.h
--- a/Engine/Source/Runtime/TypedElementFramework/Public/Elements/Framework/TypedElementAttributeBinding.h
+++ b/Engine/Source/Runtime/TypedElementFramework/Public/Elements/Framework/TypedElementAttributeBinding.h
@@ -27,15 +27,14 @@ namespace UE::EditorDataStorage
 	{
 	public:
 
-
-
 		// Create an attribute binder for a given row
 		TYPEDELEMENTFRAMEWORK_API FAttributeBinder(TypedElementDataStorage::RowHandle InTargetRow);
 
 		/**
 		 * Bind a specific data member inside a TEDS column to an attribute of the same type as the data
-		 *
-		 * @param Variable The data member inside a column to be bound
+		 * 
+		 * @param InVariable The data member inside a column to be bound
+		 * @param InDefaultValue The default value to be used when the column isn't present on a row
 		 * @return A TAttribute bound to the row, column pair specified
 		 *
 		 * Example:
@@ -43,13 +42,15 @@ namespace UE::EditorDataStorage
 		 * TAttribute<FString> TestAttribute(Binder.BindData(&FTypedElementLabelColumn::Label))
 		 */
 		template <typename AttributeType, TypedElementDataStorage::TDataColumnType ColumnType>
-		TAttribute<AttributeType> BindData(AttributeType ColumnType::* Variable);
+		TAttribute<AttributeType> BindData(AttributeType ColumnType::* InVariable, const AttributeType& InDefaultValue = AttributeType());
 
 		/**
 		 * Bind a specific data member inside a TEDS column to an attribute of a different type than the data by providing a conversion function
+		 * NOTE: the default value is not the actual attribute type but rather the data type in the column and it gets passed to the conversion function
 		 *
-		 * @param Variable The data member inside a column to be bound
-		 * @param Converter Conversion function to convert from DataType -> AttributeType
+		 * @param InVariable The data member inside a column to be bound
+		 * @param InConverter Conversion function to convert from DataType -> AttributeType
+		 * @param InDefaultValue The default value to be used when the column isn't present on a row 
 		 * @return A TAttribute bound to the row, column pair specified
 		 *
 		 * Example:
@@ -62,18 +63,19 @@ namespace UE::EditorDataStorage
 		 *                                 ));
 		 */
 		template <typename AttributeType, typename DataType, TypedElementDataStorage::TDataColumnType ColumnType>
-		TAttribute<AttributeType> BindData(DataType ColumnType::* Variable, const TFunction<AttributeType(const DataType&)>& Converter);
+		TAttribute<AttributeType> BindData(DataType ColumnType::* InVariable, const TFunction<AttributeType(const DataType&)>& InConverter, const DataType& InDefaultValue = DataType());
 
 		/**
 		 * Overload for the conversion binder to accept lambdas instead of TFunctions
 		 *
-		 * @param Variable The data member inside a column to be bound
-		 * @param Converter Conversion function to convert from DataType -> AttributeType
+		 * @param InVariable The data member inside a column to be bound
+		 * @param InConverter Conversion function to convert from DataType -> AttributeType
+		 * @param InDefaultValue The default value to be used when the column isn't present on a row
 		 * @return A TAttribute bound to the row, column pair specified
 		 */
 		template <typename DataType, TypedElementDataStorage::TDataColumnType ColumnType, typename FunctionType>
 			requires AttributeBinderInvocable<FunctionType, DataType>
-		auto BindData(DataType ColumnType::* Variable, FunctionType Converter);
+		auto BindData(DataType ColumnType::* InVariable, FunctionType InConverter, const DataType& InDefaultValue = DataType());
 		
 	private:
 
diff --git a/Engine/Source/Runtime/TypedElementFramework/Public/Elements/Framework/TypedElementAttributeBinding.inl b/Engine/Source/Runtime/TypedElementFramework/Public/Elements/Framework/TypedElementAttributeBinding.inl
--- a/Engine/Source/Runtime/TypedElementFramework/Public/Elements/Framework/TypedElementAttributeBinding.inl
+++ b/Engine/Source/Runtime/TypedElementFramework/Public/Elements/Framework/TypedElementAttributeBinding.inl
@@ -120,7 +120,7 @@ namespace UE::EditorDataStorage
 	};
 	
 	template <typename AttributeType, TypedElementDataStorage::TDataColumnType ColumnType>
-	TAttribute<AttributeType> FAttributeBinder::BindData(AttributeType ColumnType::* Variable)
+	TAttribute<AttributeType> FAttributeBinder::BindData(AttributeType ColumnType::* InVariable, const AttributeType& InDefaultValue)
 	{
 		if(!DataStorage)
 		{
@@ -129,22 +129,22 @@ namespace UE::EditorDataStorage
 		
 		// Create a direct property and bind it to the given variable
 		Property<AttributeType> Prop;
-		Prop.Bind(Variable);
+		Prop.Bind(InVariable);
 
 		// We don't want any references to this in the lambda because binders are designed to be used and destructed on the stack
-		return TAttribute<AttributeType>::CreateLambda([Property = Prop, Storage = DataStorage, Row = TargetRow]()
+		return TAttribute<AttributeType>::CreateLambda([Property = Prop, Storage = DataStorage, Row = TargetRow, DefaultValue = InDefaultValue]()
 		{
 			// Get the column from the given row and use that to return the stored property
 			if(ColumnType* Column = Storage->GetColumn<ColumnType>(Row))
 			{
 				return Property.Get(Column, ColumnType::StaticStruct());
 			}
-			return AttributeType();
+			return DefaultValue;
 		});
 	}
 
 	template <typename AttributeType, typename DataType, TypedElementDataStorage::TDataColumnType ColumnType>
-	TAttribute<AttributeType> FAttributeBinder::BindData(DataType ColumnType::* Variable, const TFunction<AttributeType(const DataType&)>& Converter)
+	TAttribute<AttributeType> FAttributeBinder::BindData(DataType ColumnType::* InVariable, const TFunction<AttributeType(const DataType&)>& InConverter, const DataType& InDefaultValue)
 	{
 		if(!DataStorage)
 		{
@@ -153,27 +153,27 @@ namespace UE::EditorDataStorage
 		
 		// Create a convertible property and bind it to the given variable
 		Property<AttributeType> Prop;
-		Prop.Bind(Variable, Converter);
+		Prop.Bind(InVariable, InConverter);
 	
-		return TAttribute<AttributeType>::CreateLambda([Property = Prop, Storage = DataStorage, Row = TargetRow]()
+		return TAttribute<AttributeType>::CreateLambda([Property = Prop, Storage = DataStorage, Row = TargetRow, DefaultValue = InDefaultValue, Converter = InConverter]()
 		{
 			if(ColumnType* Column = Storage->GetColumn<ColumnType>(Row))
 			{
 				return Property.Get(Column, ColumnType::StaticStruct());
 			}
 
-			return AttributeType();
+			return Converter(DefaultValue);
 		});
 	}
 	
 	template <typename DataType, TypedElementDataStorage::TDataColumnType ColumnType, typename FunctionType>
 		requires AttributeBinderInvocable<FunctionType, DataType>
-	auto FAttributeBinder::BindData(DataType ColumnType::* Variable, FunctionType Converter)
+	auto FAttributeBinder::BindData(DataType ColumnType::* InVariable, FunctionType InConverter, const DataType& InDefaultValue)
 	{
 		// Deduce the attribute type from the return value of the converter function
-		using AttributeType = decltype(Converter(std::declval<DataType>()));
+		using AttributeType = decltype(InConverter(std::declval<DataType>()));
 		
-		return BindData<AttributeType, DataType>(Variable, TFunction<AttributeType(const DataType&)>(Converter));
+		return BindData<AttributeType, DataType>(InVariable, TFunction<AttributeType(const DataType&)>(InConverter), InDefaultValue);
 	}
 	
 }
