[Backout] - CL34918893

The speculative logging is adding a significant performance overhead, and initial results from crash reporter show the issue is still occurring, without firing the double free, so there is some more digging to do in the area for the free indicies array getting bad data into it

Original CL Desc
-----------------------------------------------------------------
[Typed Elements] Add additional checks and messaging around a suspected intermittent PIE Test crash

#rnx

diff --git a/Engine/Source/Runtime/TypedElementFramework/Public/Elements/Framework/TypedElementData.h b/Engine/Source/Runtime/TypedElementFramework/Public/Elements/Framework/TypedElementData.h
--- a/Engine/Source/Runtime/TypedElementFramework/Public/Elements/Framework/TypedElementData.h
+++ b/Engine/Source/Runtime/TypedElementFramework/Public/Elements/Framework/TypedElementData.h
@@ -523,8 +523,6 @@ public:
 			? InternalDataFreeIndices.Pop(EAllowShrinking::No)
 			: InternalDataArray.Add();
 
-		checkf(InternalDataArray.IsValidIndex(InOutElementId), TEXT("Expected a valid index to use. Index %d is out of bounds of array size %d"), InOutElementId, InternalDataArray.Num());
-
 		TTypedElementInternalData<ElementDataType>& InternalData = InternalDataArray[InOutElementId];
 		InternalData.Initialize(InTypeId, InOutElementId);
 		return InternalData;
@@ -554,14 +552,13 @@ public:
 	{
 		FWriteScopeLock InternalDataLock(InternalDataRW);
 
-		check(InternalDataArray.IsValidIndex(InElementId));
+		checkSlow(InternalDataArray.IsValidIndex(InElementId));
 		
 		TTypedElementInternalData<ElementDataType>& InternalData = InternalDataArray[InElementId];
 		checkf(InExpectedDataPtr == &InternalData, TEXT("Internal data pointer did not match the expected value! Does this handle belong to a different element registry?"));
 		InternalData.LogExternalReferencesOnDestruction();
 		InternalData.Reset();
-		const TArray<int32>::SizeType FreeIndicesIndex = InternalDataFreeIndices.AddUnique(InElementId);
-		checkf(FreeIndicesIndex == InternalDataFreeIndices.Num() - 1, TEXT("Double free deteced when removing element handle at index %d"), FreeIndicesIndex);
+		InternalDataFreeIndices.Add(InElementId);
 	}
 
 	const TTypedElementInternalData<ElementDataType>& GetDataForElement(const FTypedHandleElementId InElementId) const
@@ -610,14 +607,12 @@ public:
 		FWriteScopeLock InternalDataLock(InternalDataRW);
 
 		checkSlow(InOutElementId >= 0);
-		check(!ElementIdToArrayIndex.Contains(InOutElementId));
+		checkSlow(!ElementIdToArrayIndex.Contains(InOutElementId));
 
 		const int32 InternalDataArrayIndex = InternalDataFreeIndices.Num() > 0
 			? InternalDataFreeIndices.Pop(EAllowShrinking::No)
 			: InternalDataArray.Add();
 
-		checkf(InternalDataArray.IsValidIndex(InternalDataArrayIndex), TEXT("Expected a valid index to reuse. Index %d is out of bounds of array size %d"), InternalDataArrayIndex, InternalDataArray.Num());
-
 		ElementIdToArrayIndex.Add(InOutElementId, InternalDataArrayIndex);
 
 		TTypedElementInternalData<void>& InternalData = InternalDataArray[InOutElementId];
@@ -658,8 +653,7 @@ public:
 		checkf(InExpectedDataPtr == &InternalData, TEXT("Internal data pointer did not match the expected value! Does this handle belong to a different element registry?"));
 		InternalData.LogExternalReferencesOnDestruction();
 		InternalData.Reset();
-		const TArray<int32>::SizeType FreeIndicesIndex = InternalDataFreeIndices.AddUnique(InternalDataArrayIndex);
-		checkf(FreeIndicesIndex == InternalDataFreeIndices.Num() - 1, TEXT("Double free deteced when removing element handle at index %d"), FreeIndicesIndex);
+		InternalDataFreeIndices.Add(InternalDataArrayIndex);
 	}
 
 	const TTypedElementInternalData<void>& GetDataForElement(const FTypedHandleElementId InElementId) const
