[TEDS] Add set of commands to measure performance of moving rows from processor callbacks

diff --git a/Engine/Source/Runtime/TypedElementFramework/Tests/Elements/Framework/TypedElementDataStoragePerformanceTestCommands.cpp b/Engine/Source/Runtime/TypedElementFramework/Tests/Elements/Framework/TypedElementDataStoragePerformanceTestCommands.cpp
new file mode 100644
--- /dev/null
+++ b/Engine/Source/Runtime/TypedElementFramework/Tests/Elements/Framework/TypedElementDataStoragePerformanceTestCommands.cpp
@@ -0,0 +1,152 @@
+// Copyright Epic Games, Inc. All Rights Reserved.
+
+#include "TypedElementDataStoragePerformanceTestCommands.h"
+
+#include "Elements/Framework/TypedElementQueryBuilder.h"
+#include "Elements/Framework/TypedElementRegistry.h"
+#include "HAL/IConsoleManager.h"
+
+namespace Private
+{
+	TypedElementTableHandle PerformanceTestCommandTable = TypedElementInvalidTableHandle;
+}
+
+FAutoConsoleCommand CVarAddDebugRows = FAutoConsoleCommand(
+	TEXT("Teds.Debug.PerformanceTest.AddRows"),
+	TEXT("Teds.Debug.PerformanceTest.AddRows <NumRows>;  NumRows = number of rows to add"),
+	FConsoleCommandWithArgsDelegate::CreateLambda([](const TArray<FString>& Args)
+	{
+		if (Args.Num() > 0)
+		{
+			int32 EntitiesToAdd;
+			LexFromString(EntitiesToAdd, *Args[0]);
+
+			ITypedElementDataStorageInterface* DataStorageInterface = UTypedElementRegistry::GetInstance()->GetMutableDataStorage();
+			
+			DataStorageInterface->BatchAddRow(Private::PerformanceTestCommandTable, EntitiesToAdd, [DataStorageInterface](TypedElementDataStorage::RowHandle Row)
+				{
+					FTest_PingPongPrePhys* Column = DataStorageInterface->GetColumn<FTest_PingPongPrePhys>(Row);
+					Column->Value = 0;
+				});
+		}
+	}
+));
+
+FAutoConsoleCommand CVarResetDebugEntities = FAutoConsoleCommand(
+	TEXT("Teds.Debug.PerformanceTest.RemoveAllRows"),
+	TEXT("Removes all added rows for the performance test"),
+	FConsoleCommandWithArgsDelegate::CreateLambda([](const TArray<FString>& Args)
+	{
+		using namespace TypedElementQueryBuilder;
+		using namespace TypedElementDataStorage;
+
+		using DSI = ITypedElementDataStorageInterface;
+		namespace DS = TypedElementDataStorage;
+
+		ITypedElementDataStorageInterface* DataStorageInterface = UTypedElementRegistry::GetInstance()->GetMutableDataStorage();
+		
+		TypedElementDataStorage::QueryHandle Query = DataStorageInterface->RegisterQuery(
+			Select().Where().All<FTest_PingPongPrePhys>().Compile());
+
+		TArray<TypedElementDataStorage::RowHandle> RowsToDelete;
+
+		DataStorageInterface->RunQuery(Query, CreateDirectQueryCallbackBinding([&RowsToDelete](const ITypedElementDataStorageInterface::IDirectQueryContext& Context, const TypedElementDataStorage::RowHandle* RowArray)
+		{
+			RowsToDelete.Insert(RowArray, Context.GetRowCount(), RowsToDelete.Num());
+		}));
+
+		for (TypedElementDataStorage::RowHandle Row : RowsToDelete)
+		{
+			DataStorageInterface->RemoveRow(Row);
+		}
+
+		DataStorageInterface->UnregisterQuery(Query);
+	}
+));
+
+void UTest_PingPongBetweenPhaseFactory::RegisterTables(ITypedElementDataStorageInterface& DataStorage)
+{
+	Private::PerformanceTestCommandTable = DataStorage.RegisterTable({FTest_PingPongPrePhys::StaticStruct()}, TEXT("Test_PingPongPrePhys"));
+}
+
+// Performance test is a small benchmark to better characterize the performance of adding and removing columns during the processor phases
+// The test comprises of three processors, A, B and C.  The run in consecutive phases: A in PrePhys, B in DuringPhys and C in PostPhys
+// A is sensitive to tag PingPongPrePhys.  It adds a PingPongDurPhys tag and removes the PingPonPrePhys tag.  This causes the processed row
+// to then be processed by processor B.  Processor B does a similar thing to ensure it is processed by C.  C then ensures it is processed by A
+// the next time A is run.
+void UTest_PingPongBetweenPhaseFactory::RegisterQueries(ITypedElementDataStorageInterface& DataStorage)
+{
+	using namespace TypedElementQueryBuilder;
+	using namespace TypedElementDataStorage;
+
+	using DSI = ITypedElementDataStorageInterface;
+	namespace DS = TypedElementDataStorage;
+	
+	Super::RegisterQueries(DataStorage);
+
+	DataStorage.RegisterQuery(
+		Select(TEXT("PingPong PrePhysics->DurPhysics"),
+			FProcessor(DSI::EQueryTickPhase::PrePhysics, DataStorage.GetQueryTickGroupName(DSI::EQueryTickGroups::Default))
+			.ForceToGameThread(true),
+			[](DS::IQueryContext& Context, const TypedElementRowHandle* RowPtr, FTest_PingPongPrePhys* PingPongAPtr)
+			{
+				QUICK_SCOPE_CYCLE_COUNTER(PingPong_Pre_During);
+				TArrayView<const TypedElementRowHandle> Rows(RowPtr, Context.GetRowCount());
+				TArrayView<FTest_PingPongPrePhys> PingPongAs(PingPongAPtr, Context.GetRowCount());
+
+				UScriptStruct* RemovedColumnStruct = FTest_PingPongPrePhys::StaticStruct();
+				for (int32 Index = 0, End = Context.GetRowCount(); Index < End; ++Index)
+				{
+					uint64 Value = PingPongAs[Index].Value;
+					++Value;
+					Context.AddColumn(Rows[Index], FTest_PingPongDurPhys{.Value = Value});
+					Context.RemoveColumns(Rows[Index], {RemovedColumnStruct});
+				}
+			})
+			.Where()
+			.Compile());
+
+	DataStorage.RegisterQuery(
+	Select(TEXT("PingPong DurPhysics->PostPhysics"),
+		FProcessor(DSI::EQueryTickPhase::DuringPhysics, DataStorage.GetQueryTickGroupName(DSI::EQueryTickGroups::Default))
+		.ForceToGameThread(true),
+		[](DS::IQueryContext& Context, const TypedElementRowHandle* RowPtr, FTest_PingPongDurPhys* PingPongBPtr)
+		{
+			QUICK_SCOPE_CYCLE_COUNTER(PingPong_During_Post);
+			TArrayView<const TypedElementRowHandle> Rows(RowPtr, Context.GetRowCount());
+			TArrayView<FTest_PingPongDurPhys> PingPongAs(PingPongBPtr, Context.GetRowCount());
+
+			UScriptStruct* RemovedColumnStruct = FTest_PingPongDurPhys::StaticStruct();
+			for (int32 Index = 0, End = Context.GetRowCount(); Index < End; ++Index)
+			{
+				uint64 Value = PingPongAs[Index].Value;
+				++Value;
+				Context.AddColumn(Rows[Index], FTest_PingPongPostPhys{.Value = Value});
+				Context.RemoveColumns(Rows[Index], {RemovedColumnStruct});
+			}
+		})
+		.Where()
+		.Compile());
+
+	DataStorage.RegisterQuery(
+	Select(TEXT("PingPong PostPhysics->PrePhysics"),
+	FProcessor(DSI::EQueryTickPhase::PostPhysics, DataStorage.GetQueryTickGroupName(DSI::EQueryTickGroups::Default))
+	.ForceToGameThread(true),
+	[](DS::IQueryContext& Context, const TypedElementRowHandle* RowPtr, FTest_PingPongPostPhys* PingPongBPtr)
+	{
+		QUICK_SCOPE_CYCLE_COUNTER(PingPong_Post_Pre);
+		TArrayView<const TypedElementRowHandle> Rows(RowPtr, Context.GetRowCount());
+		TArrayView<FTest_PingPongPostPhys> PingPongAs(PingPongBPtr, Context.GetRowCount());
+
+		UScriptStruct* RemovedColumnStruct = FTest_PingPongPostPhys::StaticStruct();
+		for (int32 Index = 0, End = Context.GetRowCount(); Index < End; ++Index)
+		{
+			uint64 Value = PingPongAs[Index].Value;
+			++Value;
+			Context.AddColumn(Rows[Index], FTest_PingPongPrePhys{.Value = Value});
+			Context.RemoveColumns(Rows[Index], {RemovedColumnStruct});
+		}
+	})
+	.Where()
+	.Compile());
+}
\ No newline at end of file
diff --git a/Engine/Source/Runtime/TypedElementFramework/Tests/Elements/Framework/TypedElementDataStoragePerformanceTestCommands.h b/Engine/Source/Runtime/TypedElementFramework/Tests/Elements/Framework/TypedElementDataStoragePerformanceTestCommands.h
new file mode 100644
--- /dev/null
+++ b/Engine/Source/Runtime/TypedElementFramework/Tests/Elements/Framework/TypedElementDataStoragePerformanceTestCommands.h
@@ -0,0 +1,58 @@
+// Copyright Epic Games, Inc. All Rights Reserved.
+
+#pragma once
+
+#include "CoreMinimal.h"
+#include "Elements/Interfaces/TypedElementDataStorageFactory.h"
+#include "Elements/Interfaces/TypedElementDataStorageInterface.h"
+#include "Modules/ModuleManager.h"
+
+#include "TypedElementDataStoragePerformanceTestCommands.generated.h"
+
+/**
+ * Column to represent that a row is selected
+ */
+USTRUCT(meta = (DisplayName = ""))
+struct FTest_PingPongPrePhys final : public FTypedElementDataStorageColumn
+{
+	GENERATED_BODY()
+	
+	UPROPERTY()
+	uint64 Value;
+};
+
+USTRUCT(meta = (DisplayName = ""))
+struct FTest_PingPongDurPhys final : public FTypedElementDataStorageColumn
+{
+	GENERATED_BODY()
+
+	UPROPERTY()
+	uint64 Value;
+};
+
+USTRUCT(meta = (DisplayName = ""))
+struct FTest_PingPongPostPhys final : public FTypedElementDataStorageColumn
+{
+	GENERATED_BODY()
+	
+	UPROPERTY()
+	uint64 Value;
+};
+
+/**
+ * The heads up transform display provides at a glance view in a scene outliner row of abnormal transform characteristics, including:
+ *		1. Non-uniform scale
+ *		2. Negative scaling on X, Y, or Z axis
+ *		3. Unnormalized rotation
+ */
+UCLASS()
+class UTest_PingPongBetweenPhaseFactory : public UTypedElementDataStorageFactory
+{
+	GENERATED_BODY()
+
+public:
+	~UTest_PingPongBetweenPhaseFactory() override = default;
+
+	void RegisterTables(ITypedElementDataStorageInterface& DataStorage) override;
+	void RegisterQueries(ITypedElementDataStorageInterface& DataStorage) override;
+};
\ No newline at end of file
