[Teds] Rename tag and column base classes, and namespace the common types header to UE::Editor::DataStorage

#rnx

diff --git a/Engine/Plugins/Experimental/Animation/HierarchyTable/Source/Editor/Public/HierarchyTableProxyColumn.h b/Engine/Plugins/Experimental/Animation/HierarchyTable/Source/Editor/Public/HierarchyTableProxyColumn.h
--- a/Engine/Plugins/Experimental/Animation/HierarchyTable/Source/Editor/Public/HierarchyTableProxyColumn.h
+++ b/Engine/Plugins/Experimental/Animation/HierarchyTable/Source/Editor/Public/HierarchyTableProxyColumn.h
@@ -10,7 +10,7 @@ class UHierarchyTable;
 struct FHierarchyTableEntryData;
 
 USTRUCT()
-struct FHierarchyTableProxyColumn : public FTypedElementDataStorageColumn
+struct FHierarchyTableProxyColumn : public FEditorDataStorageColumn
 {
 	GENERATED_BODY()
 
diff --git a/Engine/Plugins/Experimental/Animation/HierarchyTableBuiltin/Source/Editor/Public/MaskProfile/MaskProfileProxyColumn.h b/Engine/Plugins/Experimental/Animation/HierarchyTableBuiltin/Source/Editor/Public/MaskProfile/MaskProfileProxyColumn.h
--- a/Engine/Plugins/Experimental/Animation/HierarchyTableBuiltin/Source/Editor/Public/MaskProfile/MaskProfileProxyColumn.h
+++ b/Engine/Plugins/Experimental/Animation/HierarchyTableBuiltin/Source/Editor/Public/MaskProfile/MaskProfileProxyColumn.h
@@ -7,7 +7,7 @@
 #include "MaskProfileProxyColumn.generated.h"
 
 USTRUCT()
-struct FHierarchyTableMaskColumn_Value final : public FTypedElementDataStorageColumn
+struct FHierarchyTableMaskColumn_Value final : public FEditorDataStorageColumn
 {
 	GENERATED_BODY()
 };
\ No newline at end of file
diff --git a/Engine/Plugins/Experimental/Animation/HierarchyTableBuiltin/Source/Editor/Public/TimeProfile/TimeProfileProxyColumn.h b/Engine/Plugins/Experimental/Animation/HierarchyTableBuiltin/Source/Editor/Public/TimeProfile/TimeProfileProxyColumn.h
--- a/Engine/Plugins/Experimental/Animation/HierarchyTableBuiltin/Source/Editor/Public/TimeProfile/TimeProfileProxyColumn.h
+++ b/Engine/Plugins/Experimental/Animation/HierarchyTableBuiltin/Source/Editor/Public/TimeProfile/TimeProfileProxyColumn.h
@@ -7,25 +7,25 @@
 #include "TimeProfileProxyColumn.generated.h"
 
 USTRUCT()
-struct FHierarchyTableTimeColumn_StartTime final : public FTypedElementDataStorageColumn
+struct FHierarchyTableTimeColumn_StartTime final : public FEditorDataStorageColumn
 {
 	GENERATED_BODY()
 };
 
 USTRUCT()
-struct FHierarchyTableTimeColumn_EndTime final : public FTypedElementDataStorageColumn
+struct FHierarchyTableTimeColumn_EndTime final : public FEditorDataStorageColumn
 {
 	GENERATED_BODY()
 };
 
 USTRUCT()
-struct FHierarchyTableTimeColumn_TimeFactor final : public FTypedElementDataStorageColumn
+struct FHierarchyTableTimeColumn_TimeFactor final : public FEditorDataStorageColumn
 {
 	GENERATED_BODY()
 };
 
 USTRUCT()
-struct FHierarchyTableTimeColumn_Preview final : public FTypedElementDataStorageColumn
+struct FHierarchyTableTimeColumn_Preview final : public FEditorDataStorageColumn
 {
 	GENERATED_BODY()
 };
\ No newline at end of file
diff --git a/Engine/Plugins/Experimental/EditorDataStorage/Source/TedsCore/Private/DynamicColumnGenerator.cpp b/Engine/Plugins/Experimental/EditorDataStorage/Source/TedsCore/Private/DynamicColumnGenerator.cpp
--- a/Engine/Plugins/Experimental/EditorDataStorage/Source/TedsCore/Private/DynamicColumnGenerator.cpp
+++ b/Engine/Plugins/Experimental/EditorDataStorage/Source/TedsCore/Private/DynamicColumnGenerator.cpp
@@ -31,8 +31,8 @@ namespace UE::Editor::DataStorage
 		}
 		
 		checkf(
-			Template.IsChildOf(FTypedElementDataStorageColumn::StaticStruct()) ||
-			Template.IsChildOf(FTypedElementDataStorageTag::StaticStruct()) ||
+			Template.IsChildOf(UE::Editor::DataStorage::FColumn::StaticStruct()) ||
+			Template.IsChildOf(UE::Editor::DataStorage::FTag::StaticStruct()) ||
 			Template.IsChildOf(FTedsSharedColumn::StaticStruct()),
 			TEXT("Template struct must derive from Column, Tag or SharedColumn"));
 	
diff --git a/Engine/Plugins/Experimental/EditorDataStorage/Source/TedsCore/Private/Memento/TypedElementMementoRowTypes.h b/Engine/Plugins/Experimental/EditorDataStorage/Source/TedsCore/Private/Memento/TypedElementMementoRowTypes.h
--- a/Engine/Plugins/Experimental/EditorDataStorage/Source/TedsCore/Private/Memento/TypedElementMementoRowTypes.h
+++ b/Engine/Plugins/Experimental/EditorDataStorage/Source/TedsCore/Private/Memento/TypedElementMementoRowTypes.h
@@ -2,7 +2,7 @@
 
 #pragma once
 
-#include "Elements/Interfaces/TypedElementDataStorageInterface.h"
+#include "Elements/Common/TypedElementCommonTypes.h"
 #include "TypedElementMementoRowTypes.generated.h"
 
 /**
@@ -14,7 +14,7 @@
  * MementoTag denotes that the row is a memento
  */
 USTRUCT()
-struct FTypedElementMementoTag : public FTypedElementDataStorageTag
+struct FTypedElementMementoTag : public FEditorDataStorageTag
 {
 	GENERATED_BODY()
 };
diff --git a/Engine/Plugins/Experimental/EditorDataStorage/Source/TedsCore/Private/Memento/TypedElementMementoTranslators.cpp b/Engine/Plugins/Experimental/EditorDataStorage/Source/TedsCore/Private/Memento/TypedElementMementoTranslators.cpp
--- a/Engine/Plugins/Experimental/EditorDataStorage/Source/TedsCore/Private/Memento/TypedElementMementoTranslators.cpp
+++ b/Engine/Plugins/Experimental/EditorDataStorage/Source/TedsCore/Private/Memento/TypedElementMementoTranslators.cpp
@@ -35,8 +35,8 @@ void UTedsDefaultMementoTranslator::PostInitProperties()
 	const UPropertyBag* PropertyBag = UPropertyBag::GetOrCreateFromDescs(PropertyDescs, *GeneratedMementoStructName);
 	MementoType = PropertyBag;
 
-	// Need to change the type to a FTypedElementDataStorageColumn to appease TEDS/Mass
-	const_cast<UPropertyBag*>(PropertyBag)->SetSuperStruct(FTypedElementDataStorageColumn::StaticStruct());
+	// Need to change the type to a FColumn to appease TEDS/Mass
+	const_cast<UPropertyBag*>(PropertyBag)->SetSuperStruct(UE::Editor::DataStorage::FColumn::StaticStruct());
 
 	// Create the property mapping
 	for (FProperty* SourceProperty = SourceColumnType->PropertyLink; SourceProperty; SourceProperty = SourceProperty->PropertyLinkNext)
diff --git a/Engine/Plugins/Experimental/EditorDataStorage/Source/TedsCore/Private/Queries/TypedElementExtendedQueryStore.cpp b/Engine/Plugins/Experimental/EditorDataStorage/Source/TedsCore/Private/Queries/TypedElementExtendedQueryStore.cpp
--- a/Engine/Plugins/Experimental/EditorDataStorage/Source/TedsCore/Private/Queries/TypedElementExtendedQueryStore.cpp
+++ b/Engine/Plugins/Experimental/EditorDataStorage/Source/TedsCore/Private/Queries/TypedElementExtendedQueryStore.cpp
@@ -505,9 +505,9 @@ namespace UE::Editor::DataStorage
 					DSI::EQueryAccessType AccessType = Query.SelectionAccessTypes[SelectionIndex];
 					if (ensureMsgf(Type.IsValid(), TEXT("Provided query selection type can not be null.")) &&
 						ensureMsgf(
-							Type->IsChildOf(FTypedElementDataStorageColumn::StaticStruct()) ||
+							Type->IsChildOf(UE::Editor::DataStorage::FColumn::StaticStruct()) ||
 							Type->IsChildOf(FMassFragment::StaticStruct()),
-							TEXT("Provided query selection type '%s' is not based on FTypedElementDataStorageColumn or another supported base type."),
+							TEXT("Provided query selection type '%s' is not based on UE::Editor::DataStorage::FColumn or another supported base type."),
 							*Type->GetStructPathName().ToString()))
 					{
 						NativeQuery.AddRequirement(Type.Get(), ConvertToNativeAccessType(AccessType), ConvertToNativePresenceType(AccessType));
diff --git a/Engine/Plugins/Experimental/EditorDataStorage/Source/TedsCore/Private/TypedElementDatabase.cpp b/Engine/Plugins/Experimental/EditorDataStorage/Source/TedsCore/Private/TypedElementDatabase.cpp
--- a/Engine/Plugins/Experimental/EditorDataStorage/Source/TedsCore/Private/TypedElementDatabase.cpp
+++ b/Engine/Plugins/Experimental/EditorDataStorage/Source/TedsCore/Private/TypedElementDatabase.cpp
@@ -363,7 +363,7 @@ TypedElementDataStorage::RowHandle UEditorDataStorage::AddRow(TypedElementDataSt
 }
 
 TypedElementDataStorage::RowHandle UEditorDataStorage::AddRow(TypedElementDataStorage::TableHandle Table,
-	TypedElementDataStorage::RowCreationCallbackRef OnCreated)
+	RowCreationCallbackRef OnCreated)
 {
 	using namespace TypedElementDataStorage;
 
@@ -401,7 +401,7 @@ bool UEditorDataStorage::AddRow(TypedElementDataStorage::RowHandle ReservedRow,
 }
 
 bool UEditorDataStorage::AddRow(TypedElementDataStorage::RowHandle ReservedRow, TypedElementDataStorage::TableHandle Table,
-	TypedElementDataStorage::RowCreationCallbackRef OnCreated)
+	RowCreationCallbackRef OnCreated)
 {
 	OnCreated.CheckCallable();
 	if (ActiveEditorEntityManager)
@@ -418,7 +418,7 @@ bool UEditorDataStorage::AddRow(TypedElementDataStorage::RowHandle ReservedRow,
 }
 
 bool UEditorDataStorage::BatchAddRow(
-	TypedElementDataStorage::TableHandle Table, int32 Count, TypedElementDataStorage::RowCreationCallbackRef OnCreated)
+	TypedElementDataStorage::TableHandle Table, int32 Count, RowCreationCallbackRef OnCreated)
 {
 	OnCreated.CheckCallable();
 	if (ActiveEditorEntityManager)
@@ -441,7 +441,7 @@ bool UEditorDataStorage::BatchAddRow(
 }
 
 bool UEditorDataStorage::BatchAddRow(TypedElementDataStorage::TableHandle Table, 
-	TConstArrayView<TypedElementDataStorage::RowHandle> ReservedHandles, TypedElementDataStorage::RowCreationCallbackRef OnCreated)
+	TConstArrayView<TypedElementDataStorage::RowHandle> ReservedHandles, RowCreationCallbackRef OnCreated)
 {
 	OnCreated.CheckCallable();
 	if (ActiveEditorEntityManager)
@@ -513,8 +513,8 @@ void UEditorDataStorage::AddColumn(TypedElementRowHandle Row, const UScriptStruc
 }
 
 void UEditorDataStorage::AddColumnData(TypedElementRowHandle Row, const UScriptStruct* ColumnType,
-	const TypedElementDataStorage::ColumnCreationCallbackRef& Initializer,
-	TypedElementDataStorage::ColumnCopyOrMoveCallback Relocator)
+	const ColumnCreationCallbackRef& Initializer,
+	ColumnCopyOrMoveCallback Relocator)
 {
 	if (ActiveEditorEntityManager && ColumnType && ColumnType->IsChildOf(FMassFragment::StaticStruct()))
 	{
@@ -854,7 +854,7 @@ bool UEditorDataStorage::HasColumns(TypedElementRowHandle Row, TConstArrayView<T
 	return false;
 }
 
-void UEditorDataStorage::ListColumns(TypedElementDataStorage::RowHandle Row, TypedElementDataStorage::ColumnListCallbackRef Callback) const
+void UEditorDataStorage::ListColumns(TypedElementDataStorage::RowHandle Row, ColumnListCallbackRef Callback) const
 {
 	if (ActiveEditorEntityManager)
 	{
@@ -879,7 +879,7 @@ void UEditorDataStorage::ListColumns(TypedElementDataStorage::RowHandle Row, Typ
 	}
 }
 
-void UEditorDataStorage::ListColumns(TypedElementDataStorage::RowHandle Row, TypedElementDataStorage::ColumnListWithDataCallbackRef Callback)
+void UEditorDataStorage::ListColumns(TypedElementDataStorage::RowHandle Row, ColumnListWithDataCallbackRef Callback)
 {
 	if (ActiveEditorEntityManager)
 	{
diff --git a/Engine/Plugins/Experimental/EditorDataStorage/Source/TedsCore/Private/TypedElementDatabase.h b/Engine/Plugins/Experimental/EditorDataStorage/Source/TedsCore/Private/TypedElementDatabase.h
--- a/Engine/Plugins/Experimental/EditorDataStorage/Source/TedsCore/Private/TypedElementDatabase.h
+++ b/Engine/Plugins/Experimental/EditorDataStorage/Source/TedsCore/Private/TypedElementDatabase.h
@@ -31,6 +31,12 @@ class TEDSCORE_API UEditorDataStorage
 {
 	GENERATED_BODY()
 
+	using RowCreationCallbackRef = UE::Editor::DataStorage::RowCreationCallbackRef;
+	using ColumnCreationCallbackRef = UE::Editor::DataStorage::ColumnCreationCallbackRef;
+	using ColumnListCallbackRef = UE::Editor::DataStorage::ColumnListCallbackRef;
+	using ColumnListWithDataCallbackRef = UE::Editor::DataStorage::ColumnListWithDataCallbackRef;
+	using ColumnCopyOrMoveCallback = UE::Editor::DataStorage::ColumnCopyOrMoveCallback;
+
 public:
 	template<typename FactoryType, typename DatabaseType>
 	class TFactoryIterator
@@ -93,15 +99,15 @@ public:
 	virtual void BatchReserveRows(int32 Count, TFunctionRef<void(TypedElementDataStorage::RowHandle)> ReservationCallback) override;
 	virtual void BatchReserveRows(TArrayView<TypedElementDataStorage::RowHandle> ReservedRows) override;
 	virtual TypedElementDataStorage::RowHandle AddRow(TypedElementDataStorage::TableHandle Table, 
-		TypedElementDataStorage::RowCreationCallbackRef OnCreated) override;
+		RowCreationCallbackRef OnCreated) override;
 	TypedElementDataStorage::RowHandle AddRow(TypedElementDataStorage::TableHandle Table) override;
 	virtual bool AddRow(TypedElementDataStorage::RowHandle ReservedRow, TypedElementDataStorage::TableHandle Table) override;
 	virtual bool AddRow(TypedElementDataStorage::RowHandle ReservedRow, TypedElementDataStorage::TableHandle Table,
-		TypedElementDataStorage::RowCreationCallbackRef OnCreated) override;
+		RowCreationCallbackRef OnCreated) override;
 	virtual bool BatchAddRow(TypedElementDataStorage::TableHandle Table, int32 Count,
-		TypedElementDataStorage::RowCreationCallbackRef OnCreated) override;
+		RowCreationCallbackRef OnCreated) override;
 	virtual bool BatchAddRow(TypedElementDataStorage::TableHandle Table, TConstArrayView<TypedElementDataStorage::RowHandle> ReservedHandles,
-		TypedElementDataStorage::RowCreationCallbackRef OnCreated) override;
+		RowCreationCallbackRef OnCreated) override;
 	virtual void RemoveRow(TypedElementDataStorage::RowHandle Row) override;
 	virtual bool IsRowAvailable(TypedElementDataStorage::RowHandle Row) const override;
 	virtual bool IsRowAssigned(TypedElementDataStorage::RowHandle Row) const override;
@@ -109,8 +115,8 @@ public:
 	virtual void AddColumn(TypedElementRowHandle Row, const UScriptStruct* ColumnType) override;
 	virtual void AddColumn(TypedElementRowHandle Row, const UE::Editor::DataStorage::FDynamicTag& Tag, const FName& InValue) override;
 	virtual void AddColumnData(TypedElementRowHandle Row, const UScriptStruct* ColumnType,
-		const TypedElementDataStorage::ColumnCreationCallbackRef& Initializer,
-		TypedElementDataStorage::ColumnCopyOrMoveCallback Relocator) override;
+		const ColumnCreationCallbackRef& Initializer,
+		ColumnCopyOrMoveCallback Relocator) override;
 	virtual void RemoveColumn(TypedElementRowHandle Row, const UScriptStruct* ColumnType) override;
 	virtual void RemoveColumn(TypedElementRowHandle Row, const UE::Editor::DataStorage::FDynamicTag& Tag) override;
 	virtual void* GetColumnData(TypedElementRowHandle Row, const UScriptStruct* ColumnType) override;
@@ -123,8 +129,8 @@ public:
 		TConstArrayView<const UScriptStruct*> ColumnsToRemove) override;
 	virtual bool HasColumns(TypedElementRowHandle Row, TConstArrayView<const UScriptStruct*> ColumnTypes) const override;
 	virtual bool HasColumns(TypedElementRowHandle Row, TConstArrayView<TWeakObjectPtr<const UScriptStruct>> ColumnTypes) const override;
-	virtual void ListColumns(TypedElementDataStorage::RowHandle Row, TypedElementDataStorage::ColumnListCallbackRef Callback) const;
-	virtual void ListColumns(TypedElementDataStorage::RowHandle Row, TypedElementDataStorage::ColumnListWithDataCallbackRef Callback);
+	virtual void ListColumns(TypedElementDataStorage::RowHandle Row, ColumnListCallbackRef Callback) const;
+	virtual void ListColumns(TypedElementDataStorage::RowHandle Row, ColumnListWithDataCallbackRef Callback);
 	virtual bool MatchesColumns(TypedElementDataStorage::RowHandle Row, const TypedElementDataStorage::FQueryConditions& Conditions) const override;
 
 	void RegisterTickGroup(FName GroupName, EQueryTickPhase Phase, FName BeforeGroup, FName AfterGroup, TypedElementDataStorage::EExecutionMode ExecutionMode);
diff --git a/Engine/Plugins/Experimental/EditorDataStorage/Source/TedsCore/Private/TypedElementDatabaseCommandBuffer.cpp b/Engine/Plugins/Experimental/EditorDataStorage/Source/TedsCore/Private/TypedElementDatabaseCommandBuffer.cpp
--- a/Engine/Plugins/Experimental/EditorDataStorage/Source/TedsCore/Private/TypedElementDatabaseCommandBuffer.cpp
+++ b/Engine/Plugins/Experimental/EditorDataStorage/Source/TedsCore/Private/TypedElementDatabaseCommandBuffer.cpp
@@ -96,7 +96,7 @@ namespace UE::Editor::DataStorage::Legacy
 	}
 
 	void* FCommandBuffer::Queue_AddDataColumnCommandUnitialized(
-		TypedElementDataStorage::RowHandle Row, const UScriptStruct* ColumnType, TypedElementDataStorage::ColumnCopyOrMoveCallback Relocator)
+		TypedElementDataStorage::RowHandle Row, const UScriptStruct* ColumnType, ColumnCopyOrMoveCallback Relocator)
 	{
 		checkf(ColumnType->IsChildOf(FMassFragment::StaticStruct()),
 			TEXT("Trying to queue a data column creation for '%s' which isn't a data column."), *ColumnType->GetName());
@@ -215,7 +215,7 @@ namespace UE::Editor::DataStorage::Legacy
 
 	void FCommandBuffer::Execute_AddDataColumnCommand(
 		FMassEntityManager& MassEntityManager, TypedElementDataStorage::RowHandle Row, const UScriptStruct* ColumnType,
-		void* Data, TypedElementDataStorage::ColumnCopyOrMoveCallback Relocator)
+		void* Data, ColumnCopyOrMoveCallback Relocator)
 	{
 		if (ColumnType)
 		{
diff --git a/Engine/Plugins/Experimental/EditorDataStorage/Source/TedsCore/Private/TypedElementDatabaseCommandBuffer.h b/Engine/Plugins/Experimental/EditorDataStorage/Source/TedsCore/Private/TypedElementDatabaseCommandBuffer.h
--- a/Engine/Plugins/Experimental/EditorDataStorage/Source/TedsCore/Private/TypedElementDatabaseCommandBuffer.h
+++ b/Engine/Plugins/Experimental/EditorDataStorage/Source/TedsCore/Private/TypedElementDatabaseCommandBuffer.h
@@ -37,7 +37,7 @@ namespace UE::Editor::DataStorage
 			 */
 			void Queue_AddColumnCommand(TypedElementDataStorage::RowHandle Row, const UScriptStruct* ColumnType);
 			void* Queue_AddDataColumnCommandUnitialized(TypedElementDataStorage::RowHandle Row, const UScriptStruct* ColumnType,
-				TypedElementDataStorage::ColumnCopyOrMoveCallback Relocator);
+				ColumnCopyOrMoveCallback Relocator);
 			void Queue_AddColumnsCommand(TypedElementDataStorage::RowHandle Row, FMassFragmentBitSet FragmentsToAdd, FMassTagBitSet TagsToAdd);
 	
 			void Queue_RemoveColumnCommand(TypedElementDataStorage::RowHandle Row, const UScriptStruct* ColumnType);
@@ -54,7 +54,7 @@ namespace UE::Editor::DataStorage
 				FMassEntityManager& MassEntityManager, TypedElementDataStorage::RowHandle Row, const UScriptStruct* ColumnType);
 			static void Execute_AddDataColumnCommand(
 				FMassEntityManager& MassEntityManager, TypedElementDataStorage::RowHandle Row, const UScriptStruct* ColumnType,
-				void* Data, TypedElementDataStorage::ColumnCopyOrMoveCallback Relocator);
+				void* Data, ColumnCopyOrMoveCallback Relocator);
 			static void Execute_AddSharedColumnCommand(
 				FMassEntityManager& MassEntityManager, TypedElementDataStorage::RowHandle Row, const FConstSharedStruct& SharedColumn);
 			static void Execute_RemoveSharedColumnCommand(
@@ -78,7 +78,7 @@ namespace UE::Editor::DataStorage
 			struct FAddDataColumnCommand
 			{
 				TWeakObjectPtr<const UScriptStruct> ColumnType;
-				TypedElementDataStorage::ColumnCopyOrMoveCallback Relocator;
+				ColumnCopyOrMoveCallback Relocator;
 				void* Data;
 	
 				~FAddDataColumnCommand();
diff --git a/Engine/Plugins/Experimental/EditorDataStorage/Source/TedsCore/Private/TypedElementsDataStorage.cpp b/Engine/Plugins/Experimental/EditorDataStorage/Source/TedsCore/Private/TypedElementsDataStorage.cpp
--- a/Engine/Plugins/Experimental/EditorDataStorage/Source/TedsCore/Private/TypedElementsDataStorage.cpp
+++ b/Engine/Plugins/Experimental/EditorDataStorage/Source/TedsCore/Private/TypedElementsDataStorage.cpp
@@ -31,23 +31,23 @@
  */
 void ImpersonateMassTagsAndFragments()
 {
-	// Have FTypedElementDataStorageColumn impersonate a FMassFragment, which is the actual data storage when using MASS as a backend.
-	static_assert(sizeof(FTypedElementDataStorageColumn) == sizeof(FMassFragment),
-		"In order for FTypedElementDataStorageColumn to impersonate FMassFragment they need to be identical.");
+	// Have UE::Editor::DataStorage::FColumn impersonate a FMassFragment, which is the actual data storage when using MASS as a backend.
+	static_assert(sizeof(UE::Editor::DataStorage::FColumn) == sizeof(FMassFragment),
+		"In order for UE::Editor::DataStorage::FColumn to impersonate FMassFragment they need to be identical.");
 	static_assert(!TIsPolymorphic<FMassFragment>::Value,
 		"In order to be able to impersonate FMassFragment it can't have any virtual functions.");
-	static_assert(!TIsPolymorphic<FTypedElementDataStorageColumn>::Value,
-		"In order to be able to use FTypedElementDataStorageColumn to impersonate FMassFragment it can't have any virtual functions.");
-	FTypedElementDataStorageColumn::StaticStruct()->SetSuperStruct(FMassFragment::StaticStruct());
+	static_assert(!TIsPolymorphic<UE::Editor::DataStorage::FColumn>::Value,
+		"In order to be able to use UE::Editor::DataStorage::FColumn to impersonate FMassFragment it can't have any virtual functions.");
+	UE::Editor::DataStorage::FColumn::StaticStruct()->SetSuperStruct(FMassFragment::StaticStruct());
 
-	// Have FTypedElementDataStorageTag impersonate a FMassTag, which is the tag type when using MASS as a backend.
-	static_assert(sizeof(FTypedElementDataStorageTag) == sizeof(FMassTag),
-		"In order for FTypedElementDataStorageTag to impersonate FMassTag they need to be identical.");
+	// Have UE::Editor::DataStorage::FTag impersonate a FMassTag, which is the tag type when using MASS as a backend.
+	static_assert(sizeof(UE::Editor::DataStorage::FTag) == sizeof(FMassTag),
+		"In order for UE::Editor::DataStorage::FTag to impersonate FMassTag they need to be identical.");
 	static_assert(!TIsPolymorphic<FMassTag>::Value,
 		"In order to be able to impersonate FMassTag it can't have any virtual functions.");
-	static_assert(!TIsPolymorphic<FTypedElementDataStorageTag>::Value,
-		"In order to be able to use FTypedElementDataStorageTag to impersonate FMassTag it can't have any virtual functions.");
-	FTypedElementDataStorageTag::StaticStruct()->SetSuperStruct(FMassTag::StaticStruct());
+	static_assert(!TIsPolymorphic<UE::Editor::DataStorage::FTag>::Value,
+		"In order to be able to use UE::Editor::DataStorage::FTag to impersonate FMassTag it can't have any virtual functions.");
+	UE::Editor::DataStorage::FTag::StaticStruct()->SetSuperStruct(FMassTag::StaticStruct());
 
 	FTedsSharedColumn::StaticStruct()->SetSuperStruct(FMassConstSharedFragment::StaticStruct());
 }
diff --git a/Engine/Plugins/Experimental/EditorDataStorage/Source/TedsCore/Public/Compatibility/Columns/TypedElement.h b/Engine/Plugins/Experimental/EditorDataStorage/Source/TedsCore/Public/Compatibility/Columns/TypedElement.h
--- a/Engine/Plugins/Experimental/EditorDataStorage/Source/TedsCore/Public/Compatibility/Columns/TypedElement.h
+++ b/Engine/Plugins/Experimental/EditorDataStorage/Source/TedsCore/Public/Compatibility/Columns/TypedElement.h
@@ -10,7 +10,7 @@
 // A column which contains a TypedElementHandle from the TypedElementFramework
 // Prefer to use the namespaced alias when referencing in code - UE::Editor::DataStorage::Compatibility::FTypedElementColumn
 USTRUCT()
-struct FTedsTypedElementColumn : public FTypedElementDataStorageColumn
+struct FTedsTypedElementColumn : public FEditorDataStorageColumn
 {
 	GENERATED_BODY()
 
diff --git a/Engine/Plugins/Experimental/EditorDataStorage/Source/TedsCore/Public/Memento/TypedElementMementoTranslators.h b/Engine/Plugins/Experimental/EditorDataStorage/Source/TedsCore/Public/Memento/TypedElementMementoTranslators.h
--- a/Engine/Plugins/Experimental/EditorDataStorage/Source/TedsCore/Public/Memento/TypedElementMementoTranslators.h
+++ b/Engine/Plugins/Experimental/EditorDataStorage/Source/TedsCore/Public/Memento/TypedElementMementoTranslators.h
@@ -50,7 +50,7 @@ public:
  * As an example, suppose there was a column defined as:
  * ```
  * USTRUCT()
- * struct FPathColumn : public FTypedElementDataStorageColumn
+ * struct FPathColumn : public FEditorDataStorageColumn
  * {
  *		GENERATED_BODY()
  *		
diff --git a/Engine/Plugins/Experimental/EditorDataStorage/Source/TedsUI/Private/Widgets/ExportedTextWidget.cpp b/Engine/Plugins/Experimental/EditorDataStorage/Source/TedsUI/Private/Widgets/ExportedTextWidget.cpp
--- a/Engine/Plugins/Experimental/EditorDataStorage/Source/TedsUI/Private/Widgets/ExportedTextWidget.cpp
+++ b/Engine/Plugins/Experimental/EditorDataStorage/Source/TedsUI/Private/Widgets/ExportedTextWidget.cpp
@@ -170,7 +170,7 @@ bool FExportedTextWidgetConstructor::FinalizeWidget(
 	// TEDS UI TODO: We should work around it by refactoring this into an STedsWidget in the future so it can store the column conditions per instance
 	MatchedColumn = TypedElementDataStorage::FQueryConditions(TypedElementDataStorage::FColumn(TypeInfoColumn.TypeInfo));
 
-	if (TypeInfoColumn.TypeInfo->IsChildOf(FTypedElementDataStorageTag::StaticStruct()))
+	if (TypeInfoColumn.TypeInfo->IsChildOf(UE::Editor::DataStorage::FTag::StaticStruct()))
 	{
 		UpdateExportedTextWidget(
 			LOCTEXT("ExportedTextWidgetTag", "<Tag>"),
diff --git a/Engine/Plugins/Experimental/EditorDataStorage/Source/TedsUI/Public/Widgets/CounterWidget.h b/Engine/Plugins/Experimental/EditorDataStorage/Source/TedsUI/Public/Widgets/CounterWidget.h
--- a/Engine/Plugins/Experimental/EditorDataStorage/Source/TedsUI/Public/Widgets/CounterWidget.h
+++ b/Engine/Plugins/Experimental/EditorDataStorage/Source/TedsUI/Public/Widgets/CounterWidget.h
@@ -65,7 +65,7 @@ protected:
 };
 
 USTRUCT(meta = (DisplayName = "Counter widget"))
-struct TEDSUI_API FCounterWidgetColumn : public FTypedElementDataStorageColumn
+struct TEDSUI_API FCounterWidgetColumn : public FEditorDataStorageColumn
 {
 	GENERATED_BODY()
 
diff --git a/Engine/Plugins/Experimental/EditorDataStorage/Source/TedsUI/Public/Widgets/ExportedTextWidget.h b/Engine/Plugins/Experimental/EditorDataStorage/Source/TedsUI/Public/Widgets/ExportedTextWidget.h
--- a/Engine/Plugins/Experimental/EditorDataStorage/Source/TedsUI/Public/Widgets/ExportedTextWidget.h
+++ b/Engine/Plugins/Experimental/EditorDataStorage/Source/TedsUI/Public/Widgets/ExportedTextWidget.h
@@ -53,7 +53,7 @@ protected:
 };
 
 USTRUCT(meta = (DisplayName = "Exported text widget"))
-struct TEDSUI_API FExportedTextWidgetTag : public FTypedElementDataStorageTag
+struct TEDSUI_API FExportedTextWidgetTag : public FEditorDataStorageTag
 {
 	GENERATED_BODY()
 };
\ No newline at end of file
diff --git a/Engine/Plugins/Experimental/EditorDataStorage/Source/TedsUI/Public/Widgets/LabelWidget.h b/Engine/Plugins/Experimental/EditorDataStorage/Source/TedsUI/Public/Widgets/LabelWidget.h
--- a/Engine/Plugins/Experimental/EditorDataStorage/Source/TedsUI/Public/Widgets/LabelWidget.h
+++ b/Engine/Plugins/Experimental/EditorDataStorage/Source/TedsUI/Public/Widgets/LabelWidget.h
@@ -50,7 +50,7 @@ protected:
 };
 
 USTRUCT(meta = (DisplayName = "Label widget"))
-struct FLabelWidgetColumn : public FTypedElementDataStorageColumn
+struct FLabelWidgetColumn : public FEditorDataStorageColumn
 {
 	GENERATED_BODY()
 
diff --git a/Engine/Plugins/Experimental/EditorDataStorageFeatures/Source/TedsAssetData/Public/TedsAssetDataColumns.h b/Engine/Plugins/Experimental/EditorDataStorageFeatures/Source/TedsAssetData/Public/TedsAssetDataColumns.h
--- a/Engine/Plugins/Experimental/EditorDataStorageFeatures/Source/TedsAssetData/Public/TedsAssetDataColumns.h
+++ b/Engine/Plugins/Experimental/EditorDataStorageFeatures/Source/TedsAssetData/Public/TedsAssetDataColumns.h
@@ -10,7 +10,7 @@
 #include "TedsAssetDataColumns.generated.h"
 
 USTRUCT()
-struct FAssetPathColumn_Experimental : public FTypedElementDataStorageColumn
+struct FAssetPathColumn_Experimental : public FEditorDataStorageColumn
 {
 	GENERATED_BODY()
 	
@@ -19,7 +19,7 @@ struct FAssetPathColumn_Experimental : public FTypedElementDataStorageColumn
 };
 
 USTRUCT()
-struct FParentAssetPathColumn_Experimental : public FTypedElementDataStorageColumn
+struct FParentAssetPathColumn_Experimental : public FEditorDataStorageColumn
 {
 	GENERATED_BODY()
 	
@@ -27,7 +27,7 @@ struct FParentAssetPathColumn_Experimental : public FTypedElementDataStorageColu
 };
 
 USTRUCT()
-struct FChildrenAssetPathColumn_Experimental : public FTypedElementDataStorageColumn
+struct FChildrenAssetPathColumn_Experimental : public FEditorDataStorageColumn
 {
 	GENERATED_BODY()
 	
@@ -35,7 +35,7 @@ struct FChildrenAssetPathColumn_Experimental : public FTypedElementDataStorageCo
 };
 
 USTRUCT()
-struct FAssetsInPathColumn_Experimental : public FTypedElementDataStorageColumn
+struct FAssetsInPathColumn_Experimental : public FEditorDataStorageColumn
 {
 	GENERATED_BODY()
 	
@@ -43,7 +43,7 @@ struct FAssetsInPathColumn_Experimental : public FTypedElementDataStorageColumn
 };
 
 USTRUCT()
-struct FAssetDataColumn_Experimental : public FTypedElementDataStorageColumn
+struct FAssetDataColumn_Experimental : public FEditorDataStorageColumn
 {
 	GENERATED_BODY()
 	
@@ -52,7 +52,7 @@ struct FAssetDataColumn_Experimental : public FTypedElementDataStorageColumn
 };
 
 USTRUCT()
-struct FUnresolvedParentAssetPathColumn_Experimental : public FTypedElementDataStorageColumn
+struct FUnresolvedParentAssetPathColumn_Experimental : public FEditorDataStorageColumn
 {
 	GENERATED_BODY()
 	
@@ -60,7 +60,7 @@ struct FUnresolvedParentAssetPathColumn_Experimental : public FTypedElementDataS
 };
 
 USTRUCT()
-struct FUnresolvedAssetsInPathColumn_Experimental : public FTypedElementDataStorageColumn
+struct FUnresolvedAssetsInPathColumn_Experimental : public FEditorDataStorageColumn
 {
 	GENERATED_BODY()
 	
@@ -69,28 +69,28 @@ struct FUnresolvedAssetsInPathColumn_Experimental : public FTypedElementDataStor
 
 // Tag to identify assets
 USTRUCT(meta = (DisplayName = "Asset"))
-struct FAssetTag final : public FTypedElementDataStorageTag
+struct FAssetTag final : public FEditorDataStorageTag
 {
 	GENERATED_BODY()
 };
 
 // Tag to identify assets with private visibility
 USTRUCT(meta = (DisplayName = "Private Asset"))
-struct FPrivateAssetTag final : public FTypedElementDataStorageTag
+struct FPrivateAssetTag final : public FEditorDataStorageTag
 {
 	GENERATED_BODY()
 };
 
 // Tag to identify assets with public visibility
 USTRUCT(meta = (DisplayName = "Public Asset"))
-struct FPublicAssetTag final : public FTypedElementDataStorageTag
+struct FPublicAssetTag final : public FEditorDataStorageTag
 {
 	GENERATED_BODY()
 };
 
 // Column to store the disk size of an asset
 USTRUCT(meta = (DisplayName = "Disk Size"))
-struct FDiskSizeColumn final : public FTypedElementDataStorageColumn
+struct FDiskSizeColumn final : public FEditorDataStorageColumn
 {
 	GENERATED_BODY()
 
@@ -101,7 +101,7 @@ struct FDiskSizeColumn final : public FTypedElementDataStorageColumn
 
 // Column to store the verse path of an asset
 USTRUCT(meta = (DisplayName = "Verse Path"))
-struct FVersePathColumn final : public FTypedElementDataStorageColumn
+struct FVersePathColumn final : public FEditorDataStorageColumn
 {
 	GENERATED_BODY()
 	
@@ -110,20 +110,20 @@ struct FVersePathColumn final : public FTypedElementDataStorageColumn
 
 // Used to notify the dependent queries of the update to the path of the row
 USTRUCT()
-struct FUpdatedPathTag : public FTypedElementDataStorageTag
+struct FUpdatedPathTag : public FEditorDataStorageTag
 {
 	GENERATED_BODY()
 };
 
 // Used to notify the dependent queries of a update to the asset data
 USTRUCT()
-struct FUpdatedAssetDataTag : public FTypedElementDataStorageTag
+struct FUpdatedAssetDataTag : public FEditorDataStorageTag
 {
 	GENERATED_BODY()
 };
 
 USTRUCT(meta = (DisplayName = "CB Item Path"))
-struct FVirtualPathColumn_Experimental : public FTypedElementDataStorageColumn
+struct FVirtualPathColumn_Experimental : public FEditorDataStorageColumn
 {
 	GENERATED_BODY()
 
@@ -132,7 +132,7 @@ struct FVirtualPathColumn_Experimental : public FTypedElementDataStorageColumn
 };
 
 USTRUCT(meta = (DisplayName = "Name"))
-struct FItemNameColumn_Experimental : public FTypedElementDataStorageColumn
+struct FItemNameColumn_Experimental : public FEditorDataStorageColumn
 {
 	GENERATED_BODY()
 
diff --git a/Engine/Plugins/Experimental/EditorDataStorageFeatures/Source/TedsDebugger/Private/QueryEditor/TedsQueryEditorModel.cpp b/Engine/Plugins/Experimental/EditorDataStorageFeatures/Source/TedsDebugger/Private/QueryEditor/TedsQueryEditorModel.cpp
--- a/Engine/Plugins/Experimental/EditorDataStorageFeatures/Source/TedsDebugger/Private/QueryEditor/TedsQueryEditorModel.cpp
+++ b/Engine/Plugins/Experimental/EditorDataStorageFeatures/Source/TedsDebugger/Private/QueryEditor/TedsQueryEditorModel.cpp
@@ -33,8 +33,8 @@ namespace UE::EditorDataStorage::Debug::QueryEditor
 		TArray<const UScriptStruct*> Columns;
 		TArray<const UScriptStruct*> Tags;
 
-		UScriptStruct* ColumnType = FTypedElementDataStorageColumn::StaticStruct();
-		UScriptStruct* TagType = FTypedElementDataStorageTag::StaticStruct();
+		UScriptStruct* ColumnType = UE::Editor::DataStorage::FColumn::StaticStruct();
+		UScriptStruct* TagType = UE::Editor::DataStorage::FTag::StaticStruct();
 		// Not sure if there is a faster way to do this.  Would be nice to iterate only the derived classes
 		for(TObjectIterator< UScriptStruct > It; It; ++It)
 		{
@@ -260,7 +260,7 @@ namespace UE::EditorDataStorage::Debug::QueryEditor
 			}
 		}
 
-		const UScriptStruct* TagType = FTypedElementDataStorageTag::StaticStruct();
+		const UScriptStruct* TagType = UE::Editor::DataStorage::FTag::StaticStruct();
 		for (FConditionEntryInternal& Entry : Conditions)
 		{
 			if (Entry.OperatorType == EOperatorType::Unset)
diff --git a/Engine/Plugins/Experimental/EditorDataStorageFeatures/Source/TedsOutliner/Private/Compatibility/SceneOutlinerTedsBridge.cpp b/Engine/Plugins/Experimental/EditorDataStorageFeatures/Source/TedsOutliner/Private/Compatibility/SceneOutlinerTedsBridge.cpp
--- a/Engine/Plugins/Experimental/EditorDataStorageFeatures/Source/TedsOutliner/Private/Compatibility/SceneOutlinerTedsBridge.cpp
+++ b/Engine/Plugins/Experimental/EditorDataStorageFeatures/Source/TedsOutliner/Private/Compatibility/SceneOutlinerTedsBridge.cpp
@@ -87,7 +87,7 @@ FAutoConsoleCommand BindColumnsToSceneOutlinerConsoleCommand(
 								if (!Arg.IsEmpty() && Arg[0] == '/' && Path.TrySetPath(Arg))
 								{
 									const UScriptStruct* ColumnType = TypeOptional(Path);
-									if (ColumnType && ColumnType->IsChildOf(FTypedElementDataStorageColumn::StaticStruct()))
+									if (ColumnType && ColumnType->IsChildOf(UE::Editor::DataStorage::FColumn::StaticStruct()))
 									{
 										Query.ReadOnly(ColumnType);
 										++AdditionCount;
diff --git a/Engine/Plugins/Experimental/EditorDataStorageFeatures/Source/TedsOutliner/Private/TedsOutlinerImpl.cpp b/Engine/Plugins/Experimental/EditorDataStorageFeatures/Source/TedsOutliner/Private/TedsOutlinerImpl.cpp
--- a/Engine/Plugins/Experimental/EditorDataStorageFeatures/Source/TedsOutliner/Private/TedsOutlinerImpl.cpp
+++ b/Engine/Plugins/Experimental/EditorDataStorageFeatures/Source/TedsOutliner/Private/TedsOutlinerImpl.cpp
@@ -126,10 +126,10 @@ void FTedsOutlinerImpl::CreateFilterQueries()
 		TSharedRef<FFilterCategory> TedsColumnFilterCategory = MakeShared<FFilterCategory>(LOCTEXT("TedsColumnFilters", "TEDS Columns"), LOCTEXT("TedsColumnFiltersTooltip", "Filter by TEDS columns"));
 		TSharedRef<FFilterCategory> TedsTagFilterCategory = MakeShared<FFilterCategory>(LOCTEXT("TedsTagFilters", "TEDS Tags"), LOCTEXT("TedsTagFiltersTooltip", "Filter by TEDS Tags"));
 
-		const UStruct* TedsColumn = FTypedElementDataStorageColumn::StaticStruct();
-		const UStruct* TedsTag = FTypedElementDataStorageTag::StaticStruct();
+		const UStruct* TedsColumn = UE::Editor::DataStorage::FColumn::StaticStruct();
+		const UStruct* TedsTag = UE::Editor::DataStorage::FTag::StaticStruct();
 
-		// Grab all UStruct types to see if they derive from FTypedElementDataStorageColumn or FTypedElementDataStorageTag
+		// Grab all UStruct types to see if they derive from FColumn or FTag
 		ForEachObjectOfClass(UScriptStruct::StaticClass(), [&](UObject* Obj)
 		{
 			if (UScriptStruct* Struct = Cast<UScriptStruct>(Obj))
diff --git a/Engine/Plugins/Experimental/EditorDataStorageFeatures/Source/TedsOutliner/Private/Widgets/OverrideWidget.h b/Engine/Plugins/Experimental/EditorDataStorageFeatures/Source/TedsOutliner/Private/Widgets/OverrideWidget.h
--- a/Engine/Plugins/Experimental/EditorDataStorageFeatures/Source/TedsOutliner/Private/Widgets/OverrideWidget.h
+++ b/Engine/Plugins/Experimental/EditorDataStorageFeatures/Source/TedsOutliner/Private/Widgets/OverrideWidget.h
@@ -29,7 +29,7 @@ public:
 };
 
 USTRUCT(meta = (DisplayName = "Override Widget"))
-struct FOverrideWidgetTag : public FTypedElementDataStorageTag
+struct FOverrideWidgetTag : public FEditorDataStorageTag
 {
 	GENERATED_BODY()
 };
diff --git a/Engine/Plugins/Experimental/EditorDataStorageFeatures/Source/TedsOutliner/Private/Widgets/TransformHeadsUpWidget.cpp b/Engine/Plugins/Experimental/EditorDataStorageFeatures/Source/TedsOutliner/Private/Widgets/TransformHeadsUpWidget.cpp
--- a/Engine/Plugins/Experimental/EditorDataStorageFeatures/Source/TedsOutliner/Private/Widgets/TransformHeadsUpWidget.cpp
+++ b/Engine/Plugins/Experimental/EditorDataStorageFeatures/Source/TedsOutliner/Private/Widgets/TransformHeadsUpWidget.cpp
@@ -6,6 +6,7 @@
 #include "Elements/Columns/TypedElementSlateWidgetColumns.h"
 #include "Elements/Columns/TypedElementTransformColumns.h"
 #include "Elements/Framework/TypedElementQueryBuilder.h"
+#include "Elements/Interfaces/TypedElementDataStorageUiInterface.h"
 #include "Widgets/SCompoundWidget.h"
 #include "Widgets/SBoxPanel.h"
 #include "Widgets/Images/SImage.h"
diff --git a/Engine/Plugins/Experimental/EditorDataStorageFeatures/Source/TedsOutliner/Public/Columns/TedsOutlinerColumns.h b/Engine/Plugins/Experimental/EditorDataStorageFeatures/Source/TedsOutliner/Public/Columns/TedsOutlinerColumns.h
--- a/Engine/Plugins/Experimental/EditorDataStorageFeatures/Source/TedsOutliner/Public/Columns/TedsOutlinerColumns.h
+++ b/Engine/Plugins/Experimental/EditorDataStorageFeatures/Source/TedsOutliner/Public/Columns/TedsOutlinerColumns.h
@@ -12,7 +12,7 @@ class ISceneOutliner;
 // Column used to store a reference to the Teds Outliner owning a specific row
 // Currently only added to widget rows in the Teds Outliner
 USTRUCT(meta = (DisplayName = "Owning Table Viewer"))
-struct FTedsOutlinerColumn final : public FTypedElementDataStorageColumn
+struct FTedsOutlinerColumn final : public FEditorDataStorageColumn
 {
 	GENERATED_BODY()
 
diff --git a/Engine/Plugins/Experimental/EditorDataStorageFeatures/Source/TedsOutliner/Public/Widgets/AlertWidget.h b/Engine/Plugins/Experimental/EditorDataStorageFeatures/Source/TedsOutliner/Public/Widgets/AlertWidget.h
--- a/Engine/Plugins/Experimental/EditorDataStorageFeatures/Source/TedsOutliner/Public/Widgets/AlertWidget.h
+++ b/Engine/Plugins/Experimental/EditorDataStorageFeatures/Source/TedsOutliner/Public/Widgets/AlertWidget.h
@@ -3,7 +3,7 @@
 #pragma once
 
 #include "Elements/Interfaces/TypedElementDataStorageFactory.h"
-#include "Elements/Interfaces/TypedElementDataStorageInterface.h"
+#include "Elements/Common/TypedElementCommonTypes.h"
 #include "Elements/Interfaces/TypedElementDataStorageUiInterface.h"
 
 #include "AlertWidget.generated.h"
@@ -51,7 +51,7 @@ protected:
 };
 
 USTRUCT(meta = (DisplayName = "General purpose alert"))
-struct FAlertWidgetTag : public FTypedElementDataStorageTag
+struct FAlertWidgetTag : public FEditorDataStorageTag
 {
 	GENERATED_BODY()
 };
@@ -73,13 +73,13 @@ protected:
 };
 
 USTRUCT(meta = (DisplayName = "General purpose alert header"))
-struct FAlertHeaderWidgetTag : public FTypedElementDataStorageTag
+struct FAlertHeaderWidgetTag : public FEditorDataStorageTag
 {
 	GENERATED_BODY()
 };
 
 USTRUCT(meta = (DisplayName = "Alert header active"))
-struct FAlertHeaderActiveWidgetTag : public FTypedElementDataStorageTag
+struct FAlertHeaderActiveWidgetTag : public FEditorDataStorageTag
 {
 	GENERATED_BODY()
 };
diff --git a/Engine/Plugins/Experimental/EditorDataStorageFeatures/Source/TedsOutliner/Public/Widgets/TransformHeadsUpWidget.h b/Engine/Plugins/Experimental/EditorDataStorageFeatures/Source/TedsOutliner/Public/Widgets/TransformHeadsUpWidget.h
--- a/Engine/Plugins/Experimental/EditorDataStorageFeatures/Source/TedsOutliner/Public/Widgets/TransformHeadsUpWidget.h
+++ b/Engine/Plugins/Experimental/EditorDataStorageFeatures/Source/TedsOutliner/Public/Widgets/TransformHeadsUpWidget.h
@@ -3,8 +3,10 @@
 #pragma once
 
 #include "Elements/Interfaces/TypedElementDataStorageFactory.h"
-#include "Elements/Interfaces/TypedElementDataStorageInterface.h"
 #include "Elements/Interfaces/TypedElementDataStorageUiInterface.h"
+#include "Elements/Common/TypedElementCommonTypes.h"
+#include "Elements/Common/TypedElementHandles.h"
+#include "Elements/Framework/TypedElementMetaData.h"
 #include "Internationalization/Text.h"
 #include "UObject/ObjectMacros.h"
 
@@ -45,12 +47,12 @@ protected:
 	bool FinalizeWidget(
 		ITypedElementDataStorageInterface* DataStorage,
 		ITypedElementDataStorageUiInterface* DataStorageUi,
-		TypedElementRowHandle Row,
+		TypedElementDataStorage::RowHandle Row,
 		const TSharedPtr<SWidget>& Widget) override;
 };
 
 USTRUCT(meta = (DisplayName = "Heads up display for transforms widget"))
-struct FTransformHeadsUpWidgetTag : public FTypedElementDataStorageTag
+struct FTransformHeadsUpWidgetTag : public FEditorDataStorageTag
 {
 	GENERATED_BODY()
 };
diff --git a/Engine/Plugins/Experimental/EditorDataStorageFeatures/Source/TedsSettings/Public/TedsSettingsColumns.h b/Engine/Plugins/Experimental/EditorDataStorageFeatures/Source/TedsSettings/Public/TedsSettingsColumns.h
--- a/Engine/Plugins/Experimental/EditorDataStorageFeatures/Source/TedsSettings/Public/TedsSettingsColumns.h
+++ b/Engine/Plugins/Experimental/EditorDataStorageFeatures/Source/TedsSettings/Public/TedsSettingsColumns.h
@@ -8,7 +8,7 @@
 #include "TedsSettingsColumns.generated.h"
 
 USTRUCT(meta = (DisplayName = "Settings Container"))
-struct FSettingsContainerColumn final : public FTypedElementDataStorageColumn
+struct FSettingsContainerColumn final : public FEditorDataStorageColumn
 {
 	GENERATED_BODY()
 
@@ -17,7 +17,7 @@ struct FSettingsContainerColumn final : public FTypedElementDataStorageColumn
 };
 
 USTRUCT(meta = (DisplayName = "Settings Category"))
-struct FSettingsCategoryColumn final : public FTypedElementDataStorageColumn
+struct FSettingsCategoryColumn final : public FEditorDataStorageColumn
 {
 	GENERATED_BODY()
 
@@ -26,7 +26,7 @@ struct FSettingsCategoryColumn final : public FTypedElementDataStorageColumn
 };
 
 USTRUCT(meta = (DisplayName = "Settings Section"))
-struct FSettingsSectionColumn final : public FTypedElementDataStorageColumn
+struct FSettingsSectionColumn final : public FEditorDataStorageColumn
 {
 	GENERATED_BODY()
 
diff --git a/Engine/Plugins/Experimental/EditorDataStorageFeatures/Source/TedsTableViewer/Public/Columns/UIPropertiesColumns.h b/Engine/Plugins/Experimental/EditorDataStorageFeatures/Source/TedsTableViewer/Public/Columns/UIPropertiesColumns.h
--- a/Engine/Plugins/Experimental/EditorDataStorageFeatures/Source/TedsTableViewer/Public/Columns/UIPropertiesColumns.h
+++ b/Engine/Plugins/Experimental/EditorDataStorageFeatures/Source/TedsTableViewer/Public/Columns/UIPropertiesColumns.h
@@ -8,7 +8,7 @@
 #include "UIPropertiesColumns.generated.h"
 
 USTRUCT()
-struct FUIHeaderPropertiesColumn final : public FTypedElementDataStorageColumn
+struct FUIHeaderPropertiesColumn final : public FEditorDataStorageColumn
 {
 	GENERATED_BODY()
 	
diff --git a/Engine/Source/Runtime/TypedElementFramework/Public/Elements/Columns/TypedElementAlertColumns.h b/Engine/Source/Runtime/TypedElementFramework/Public/Elements/Columns/TypedElementAlertColumns.h
--- a/Engine/Source/Runtime/TypedElementFramework/Public/Elements/Columns/TypedElementAlertColumns.h
+++ b/Engine/Source/Runtime/TypedElementFramework/Public/Elements/Columns/TypedElementAlertColumns.h
@@ -22,7 +22,7 @@ enum class FTypedElementAlertColumnType : uint8
  * Column containing information a user needs to be alerted of.
  */
 USTRUCT(meta = (DisplayName = "Alert"))
-struct FTypedElementAlertColumn final : public FTypedElementDataStorageColumn
+struct FTypedElementAlertColumn final : public FEditorDataStorageColumn
 {
 	GENERATED_BODY()
 
@@ -40,7 +40,7 @@ struct FTypedElementAlertColumn final : public FTypedElementDataStorageColumn
  * Column containing a count for the number of alerts any child rows have.
  */
 USTRUCT(meta = (DisplayName = "Child alert"))
-struct FTypedElementChildAlertColumn final : public FTypedElementDataStorageColumn
+struct FTypedElementChildAlertColumn final : public FEditorDataStorageColumn
 {
 	GENERATED_BODY()
 
@@ -54,7 +54,7 @@ struct FTypedElementChildAlertColumn final : public FTypedElementDataStorageColu
  * Column that can be added to an alert column to have it trigger an action when the alert is clicked.
  */
 USTRUCT(meta = (DisplayName = "Alert action"))
-struct FTypedElementAlertActionColumn final : public FTypedElementDataStorageColumn
+struct FTypedElementAlertActionColumn final : public FEditorDataStorageColumn
 {
 	GENERATED_BODY()
 
diff --git a/Engine/Source/Runtime/TypedElementFramework/Public/Elements/Columns/TypedElementCompatibilityColumns.h b/Engine/Source/Runtime/TypedElementFramework/Public/Elements/Columns/TypedElementCompatibilityColumns.h
--- a/Engine/Source/Runtime/TypedElementFramework/Public/Elements/Columns/TypedElementCompatibilityColumns.h
+++ b/Engine/Source/Runtime/TypedElementFramework/Public/Elements/Columns/TypedElementCompatibilityColumns.h
@@ -16,7 +16,7 @@ class UWorld;
  * Column containing a non-owning reference to a UObject.
  */
 USTRUCT(meta = (DisplayName = "UObject reference"))
-struct FTypedElementUObjectColumn final : public FTypedElementDataStorageColumn
+struct FTypedElementUObjectColumn final : public FEditorDataStorageColumn
 {
 	GENERATED_BODY()
 
@@ -30,7 +30,7 @@ struct FTypedElementUObjectColumn final : public FTypedElementDataStorageColumn
  * Column containing information to uniquely identify the UObject, e.g. for use by the garbage collection.
  */
 USTRUCT(meta = (DisplayName = "UObject ID"))
-struct FTypedElementUObjectIdColumn final : public FTypedElementDataStorageColumn
+struct FTypedElementUObjectIdColumn final : public FEditorDataStorageColumn
 {
 	GENERATED_BODY()
 
@@ -46,7 +46,7 @@ struct FTypedElementUObjectIdColumn final : public FTypedElementDataStorageColum
  * recovered.
  */
 USTRUCT(meta = (DisplayName = "External object reference"))
-struct FTypedElementExternalObjectColumn final : public FTypedElementDataStorageColumn
+struct FTypedElementExternalObjectColumn final : public FEditorDataStorageColumn
 {
 	GENERATED_BODY()
 
@@ -58,7 +58,7 @@ struct FTypedElementExternalObjectColumn final : public FTypedElementDataStorage
  * If there's a FTypedElementUObjectColumn, the stored object will be a CDO if this tag is present.
  */
 USTRUCT(meta = (DisplayName = "Class Default Object"))
-struct FTypedElementClassDefaultObjectTag final : public FTypedElementDataStorageTag
+struct FTypedElementClassDefaultObjectTag final : public FEditorDataStorageTag
 {
 	GENERATED_BODY()
 };
@@ -68,7 +68,7 @@ struct FTypedElementClassDefaultObjectTag final : public FTypedElementDataStorag
  * an actor if this tag is present.
  */
 USTRUCT(meta = (DisplayName = "Actor"))
-struct FTypedElementActorTag final : public FTypedElementDataStorageTag
+struct FTypedElementActorTag final : public FEditorDataStorageTag
 {
 	GENERATED_BODY()
 };
@@ -77,7 +77,7 @@ struct FTypedElementActorTag final : public FTypedElementDataStorageTag
  * Column that stores a reference to the world.
  */
 USTRUCT(meta = (DisplayName = "World"))
-struct FTypedElementWorldColumn final : public FTypedElementDataStorageColumn
+struct FTypedElementWorldColumn final : public FEditorDataStorageColumn
 {
 	GENERATED_BODY()
 
@@ -90,7 +90,7 @@ struct FTypedElementWorldColumn final : public FTypedElementDataStorageColumn
  * bag" containing any serialized data that was loaded for it as a set of "loose" properties.
  */
 USTRUCT(meta = (DisplayName = "Property bag placeholder"))
-struct FTypedElementPropertyBagPlaceholderTag final : public FTypedElementDataStorageTag
+struct FTypedElementPropertyBagPlaceholderTag final : public FEditorDataStorageTag
 {
 	GENERATED_BODY()
 };
@@ -99,7 +99,7 @@ struct FTypedElementPropertyBagPlaceholderTag final : public FTypedElementDataSt
  * Tag to signal that a row contains an object with at least one loose property associated with it.
  */
 USTRUCT(meta = (DisplayName = "Loose property"))
-struct FTypedElementLoosePropertyTag final : public FTypedElementDataStorageTag
+struct FTypedElementLoosePropertyTag final : public FEditorDataStorageTag
 {
 	GENERATED_BODY()
 };
@@ -111,7 +111,7 @@ struct FTypedElementLoosePropertyTag final : public FTypedElementDataStorageTag
  * for systems that need to look/behave differently based on an inferred base type context (e.g. UI).
  */
 USTRUCT(meta = (DisplayName = "Placeholder type info"))
-struct FTypedElementPropertyBagPlaceholderTypeInfoColumn final : public FTypedElementDataStorageColumn
+struct FTypedElementPropertyBagPlaceholderTypeInfoColumn final : public FEditorDataStorageColumn
 {
 	GENERATED_BODY()
 
diff --git a/Engine/Source/Runtime/TypedElementFramework/Public/Elements/Columns/TypedElementHiearchyColumns.h b/Engine/Source/Runtime/TypedElementFramework/Public/Elements/Columns/TypedElementHiearchyColumns.h
--- a/Engine/Source/Runtime/TypedElementFramework/Public/Elements/Columns/TypedElementHiearchyColumns.h
+++ b/Engine/Source/Runtime/TypedElementFramework/Public/Elements/Columns/TypedElementHiearchyColumns.h
@@ -12,7 +12,7 @@
  * A reference to the direct hierarchical parent of this row.
  */
 USTRUCT(meta = (DisplayName = "Parent"))
-struct FTableRowParentColumn final : public FTypedElementDataStorageColumn
+struct FTableRowParentColumn final : public FEditorDataStorageColumn
 {
 	GENERATED_BODY()
 
@@ -24,7 +24,7 @@ struct FTableRowParentColumn final : public FTypedElementDataStorageColumn
  * be used to attempt to find the indexed row. This column can not be used to find rows that are not indexed.
  */
 USTRUCT(meta = (DisplayName = "Parent (Unresolved)"))
-struct FUnresolvedTableRowParentColumn final : public FTypedElementDataStorageColumn
+struct FUnresolvedTableRowParentColumn final : public FEditorDataStorageColumn
 {
 	GENERATED_BODY()
 
diff --git a/Engine/Source/Runtime/TypedElementFramework/Public/Elements/Columns/TypedElementLabelColumns.h b/Engine/Source/Runtime/TypedElementFramework/Public/Elements/Columns/TypedElementLabelColumns.h
--- a/Engine/Source/Runtime/TypedElementFramework/Public/Elements/Columns/TypedElementLabelColumns.h
+++ b/Engine/Source/Runtime/TypedElementFramework/Public/Elements/Columns/TypedElementLabelColumns.h
@@ -8,12 +8,11 @@
 
 #include "TypedElementLabelColumns.generated.h"
 
-
 /**
  * Column that stores a label.
  */
 USTRUCT(meta = (DisplayName = "Label"))
-struct FTypedElementLabelColumn final : public FTypedElementDataStorageColumn
+struct FTypedElementLabelColumn final : public FEditorDataStorageColumn
 {
 	GENERATED_BODY()
 
@@ -26,7 +25,7 @@ struct FTypedElementLabelColumn final : public FTypedElementDataStorageColumn
  * kept separate in order to iterate quickly over all hash values.
  */
 USTRUCT(meta = (DisplayName = "Label hash"))
-struct FTypedElementLabelHashColumn final : public FTypedElementDataStorageColumn
+struct FTypedElementLabelHashColumn final : public FEditorDataStorageColumn
 {
 	GENERATED_BODY()
 
diff --git a/Engine/Source/Runtime/TypedElementFramework/Public/Elements/Columns/TypedElementMiscColumns.h b/Engine/Source/Runtime/TypedElementFramework/Public/Elements/Columns/TypedElementMiscColumns.h
--- a/Engine/Source/Runtime/TypedElementFramework/Public/Elements/Columns/TypedElementMiscColumns.h
+++ b/Engine/Source/Runtime/TypedElementFramework/Public/Elements/Columns/TypedElementMiscColumns.h
@@ -2,7 +2,7 @@
 
 #pragma once
 
-#include "Elements/Interfaces/TypedElementDataStorageInterface.h"
+#include "Elements/Common/TypedElementCommonTypes.h"
 #include "UObject/ObjectMacros.h"
 
 #include "TypedElementMiscColumns.generated.h"
@@ -14,7 +14,7 @@
  * will automatically be removed at the end of a tick.
  */
 USTRUCT(meta = (DisplayName = "Sync back to world"))
-struct FTypedElementSyncBackToWorldTag final : public FTypedElementDataStorageTag
+struct FTypedElementSyncBackToWorldTag final : public FEditorDataStorageTag
 {
 	GENERATED_BODY()
 };
@@ -26,7 +26,7 @@ struct FTypedElementSyncBackToWorldTag final : public FTypedElementDataStorageTa
  * queries for specific changed properties.
  */
 USTRUCT(meta = (DisplayName = "Sync from world"))
-struct FTypedElementSyncFromWorldTag final : public FTypedElementDataStorageTag
+struct FTypedElementSyncFromWorldTag final : public FEditorDataStorageTag
 {
 	GENERATED_BODY()
 };
@@ -40,7 +40,7 @@ struct FTypedElementSyncFromWorldTag final : public FTypedElementDataStorageTag
  * queries for specific changed properties.
  */
 USTRUCT(meta = (DisplayName = "Sync from world (interactive)"))
-struct FTypedElementSyncFromWorldInteractiveTag final : public FTypedElementDataStorageTag
+struct FTypedElementSyncFromWorldInteractiveTag final : public FEditorDataStorageTag
 {
 	GENERATED_BODY()
 };
@@ -49,7 +49,7 @@ struct FTypedElementSyncFromWorldInteractiveTag final : public FTypedElementData
  * A general reference to another row. 
  */
 USTRUCT(meta = (DisplayName = "Row reference"))
-struct FTypedElementRowReferenceColumn final : public FTypedElementDataStorageColumn
+struct FTypedElementRowReferenceColumn final : public FEditorDataStorageColumn
 {
 	GENERATED_BODY()
 
diff --git a/Engine/Source/Runtime/TypedElementFramework/Public/Elements/Columns/TypedElementOverrideColumns.h b/Engine/Source/Runtime/TypedElementFramework/Public/Elements/Columns/TypedElementOverrideColumns.h
--- a/Engine/Source/Runtime/TypedElementFramework/Public/Elements/Columns/TypedElementOverrideColumns.h
+++ b/Engine/Source/Runtime/TypedElementFramework/Public/Elements/Columns/TypedElementOverrideColumns.h
@@ -2,14 +2,14 @@
 
 #pragma once
 
-#include "Elements/Interfaces/TypedElementDataStorageInterface.h"
+#include "Elements/Common/TypedElementCommonTypes.h"
 #include "UObject/OverridableManager.h"
 
 #include "TypedElementOverrideColumns.generated.h"
 
 // Whether the object has an override on the base
 USTRUCT(meta = (DisplayName = "Override"))
-struct FObjectOverrideColumn final : public FTypedElementDataStorageColumn
+struct FObjectOverrideColumn final : public FEditorDataStorageColumn
 {
 	GENERATED_BODY()
 
diff --git a/Engine/Source/Runtime/TypedElementFramework/Public/Elements/Columns/TypedElementPackageColumns.h b/Engine/Source/Runtime/TypedElementFramework/Public/Elements/Columns/TypedElementPackageColumns.h
--- a/Engine/Source/Runtime/TypedElementFramework/Public/Elements/Columns/TypedElementPackageColumns.h
+++ b/Engine/Source/Runtime/TypedElementFramework/Public/Elements/Columns/TypedElementPackageColumns.h
@@ -3,7 +3,9 @@
 #pragma once
 
 #include "Containers/UnrealString.h"
-#include "Elements/Interfaces/TypedElementDataStorageInterface.h"
+#include "Elements/Common/TypedElementCommonTypes.h"
+#include "Elements/Common/TypedElementHandles.h"
+#include "Elements/Common/TypedElementQueryTypes.h"
 #include "IO/PackageId.h"
 #include "Misc/PackagePath.h"
 #include "UObject/ObjectMacros.h"
@@ -14,7 +16,7 @@
  * A package reference column that has not yet been resolved to reference a package.
  */
 USTRUCT(meta = (DisplayName = "Unresolved package path reference"))
-struct FTypedElementPackageUnresolvedReference final : public FTypedElementDataStorageColumn
+struct FTypedElementPackageUnresolvedReference final : public FEditorDataStorageColumn
 {
 	GENERATED_BODY()
 
@@ -27,18 +29,18 @@ struct FTypedElementPackageUnresolvedReference final : public FTypedElementDataS
  * Column that references a row in the table that provides package and source control information.
  */
 USTRUCT(meta = (DisplayName = "Package path reference"))
-struct FTypedElementPackageReference final : public FTypedElementDataStorageColumn
+struct FTypedElementPackageReference final : public FEditorDataStorageColumn
 {
 	GENERATED_BODY()
 	
-	TypedElementRowHandle Row;
+	TypedElementDataStorage::RowHandle Row;
 };
 
 /**
  * Tag that indicates some related package information has been modified.
  */
 USTRUCT(meta = (DisplayName = "Package information has been updated"))
-struct FTypedElementPackageUpdatedTag final : public FTypedElementDataStorageTag
+struct FTypedElementPackageUpdatedTag final : public FEditorDataStorageTag
 {
 	GENERATED_BODY()
 };
@@ -47,7 +49,7 @@ struct FTypedElementPackageUpdatedTag final : public FTypedElementDataStorageTag
  * Column that stores relevant rows when a package is updated
  */
 USTRUCT()
-struct FTypedElementPackageUpdateColumn final : public FTypedElementDataStorageColumn
+struct FTypedElementPackageUpdateColumn final : public FEditorDataStorageColumn
 {
 	GENERATED_BODY()
 	
@@ -59,7 +61,7 @@ struct FTypedElementPackageUpdateColumn final : public FTypedElementDataStorageC
  * Column that stores the path of a package.
  */
 USTRUCT(meta = (DisplayName = "Package path"))
-struct FTypedElementPackagePathColumn final : public FTypedElementDataStorageColumn
+struct FTypedElementPackagePathColumn final : public FEditorDataStorageColumn
 {
 	GENERATED_BODY()
 
@@ -76,7 +78,7 @@ inline uint32 GetTypeHash(const FTypedElementPackagePathColumn& InStruct)
  * Column that stores the full loading path to a package.
  */
 USTRUCT(meta = (DisplayName = "Package loaded path"))
-struct FTypedElementPackageLoadedPathColumn final : public FTypedElementDataStorageColumn
+struct FTypedElementPackageLoadedPathColumn final : public FEditorDataStorageColumn
 {
 	GENERATED_BODY()
 
diff --git a/Engine/Source/Runtime/TypedElementFramework/Public/Elements/Columns/TypedElementPivotOffsetColumns.h b/Engine/Source/Runtime/TypedElementFramework/Public/Elements/Columns/TypedElementPivotOffsetColumns.h
--- a/Engine/Source/Runtime/TypedElementFramework/Public/Elements/Columns/TypedElementPivotOffsetColumns.h
+++ b/Engine/Source/Runtime/TypedElementFramework/Public/Elements/Columns/TypedElementPivotOffsetColumns.h
@@ -1,13 +1,13 @@
-﻿// Copyright Epic Games, Inc. All Rights Reserved.
+// Copyright Epic Games, Inc. All Rights Reserved.
 
 #pragma once
 
-#include "Elements/Interfaces/TypedElementDataStorageInterface.h"
+#include "Elements/Common/TypedElementCommonTypes.h"
 
 #include "TypedElementPivotOffsetColumns.generated.h"
 
 USTRUCT()
-struct FTypedElementPivotOffset final : public FTypedElementDataStorageColumn
+struct FTypedElementPivotOffset final : public FEditorDataStorageColumn
 {
 	GENERATED_BODY()
 
diff --git a/Engine/Source/Runtime/TypedElementFramework/Public/Elements/Columns/TypedElementRevisionControlColumns.h b/Engine/Source/Runtime/TypedElementFramework/Public/Elements/Columns/TypedElementRevisionControlColumns.h
--- a/Engine/Source/Runtime/TypedElementFramework/Public/Elements/Columns/TypedElementRevisionControlColumns.h
+++ b/Engine/Source/Runtime/TypedElementFramework/Public/Elements/Columns/TypedElementRevisionControlColumns.h
@@ -2,36 +2,36 @@
 
 #pragma once
 
-#include "Elements/Interfaces/TypedElementDataStorageInterface.h"
+#include "Elements/Common/TypedElementCommonTypes.h"
 
 #include "TypedElementRevisionControlColumns.generated.h"
 
 USTRUCT(meta = (DisplayName = "In a changelist"))
-struct FSCCInChangelistTag final : public FTypedElementDataStorageTag
+struct FSCCInChangelistTag final : public FEditorDataStorageTag
 {
 	GENERATED_BODY()
 };
 
 USTRUCT(meta = (DisplayName = "Staged"))
-struct FSCCStagedTag final : public FTypedElementDataStorageTag
+struct FSCCStagedTag final : public FEditorDataStorageTag
 {
 	GENERATED_BODY()
 };
 
 USTRUCT(meta = (DisplayName = "Locked by you"))
-struct FSCCLockedTag final : public FTypedElementDataStorageTag
+struct FSCCLockedTag final : public FEditorDataStorageTag
 {
 	GENERATED_BODY()
 };
 
 USTRUCT(meta = (DisplayName = "Edited by others"))
-struct FSCCExternallyEditedTag final : public FTypedElementDataStorageTag
+struct FSCCExternallyEditedTag final : public FEditorDataStorageTag
 {
 	GENERATED_BODY()
 };
 
 USTRUCT(meta = (DisplayName = "Not at the latest revision"))
-struct FSCCNotCurrentTag final : public FTypedElementDataStorageTag
+struct FSCCNotCurrentTag final : public FEditorDataStorageTag
 {
 	GENERATED_BODY()
 };
@@ -46,7 +46,7 @@ enum class ESCCModification
 };
 
 USTRUCT(meta = (DisplayName = "Revision Control Status"))
-struct FSCCStatusColumn final : public FTypedElementDataStorageColumn
+struct FSCCStatusColumn final : public FEditorDataStorageColumn
 {
 	GENERATED_BODY()
 
@@ -64,7 +64,7 @@ struct FSCCRevisionId
 };
 
 USTRUCT(meta = (DisplayName = "Revision ID"))
-struct FSCCRevisionIdColumn final : public FTypedElementDataStorageColumn
+struct FSCCRevisionIdColumn final : public FEditorDataStorageColumn
 {
 	GENERATED_BODY()
 
@@ -73,7 +73,7 @@ struct FSCCRevisionIdColumn final : public FTypedElementDataStorageColumn
 };
 
 USTRUCT(meta = (DisplayName = "Revision ID from server"))
-struct FSCCExternalRevisionIdColumn final : public FTypedElementDataStorageColumn
+struct FSCCExternalRevisionIdColumn final : public FEditorDataStorageColumn
 {
 	GENERATED_BODY()
 
@@ -91,7 +91,7 @@ struct FSCCUserInfo
 };
 
 USTRUCT(meta = (DisplayName = "Locked by others"))
-struct FSCCExternallyLockedColumn final : public FTypedElementDataStorageColumn
+struct FSCCExternallyLockedColumn final : public FEditorDataStorageColumn
 {
 	GENERATED_BODY()
 
diff --git a/Engine/Source/Runtime/TypedElementFramework/Public/Elements/Columns/TypedElementSelectionColumns.h b/Engine/Source/Runtime/TypedElementFramework/Public/Elements/Columns/TypedElementSelectionColumns.h
--- a/Engine/Source/Runtime/TypedElementFramework/Public/Elements/Columns/TypedElementSelectionColumns.h
+++ b/Engine/Source/Runtime/TypedElementFramework/Public/Elements/Columns/TypedElementSelectionColumns.h
@@ -1,8 +1,8 @@
-﻿// Copyright Epic Games, Inc. All Rights Reserved.
+// Copyright Epic Games, Inc. All Rights Reserved.
 
 #pragma once
 
-#include "Elements/Interfaces/TypedElementDataStorageInterface.h"
+#include "Elements/Common/TypedElementCommonTypes.h"
 #include "UObject/ObjectMacros.h"
 
 #include "TypedElementSelectionColumns.generated.h"
@@ -11,7 +11,7 @@
  * Column to represent that a row is selected
  */
 USTRUCT(meta = (DisplayName = "Selected"))
-struct FTypedElementSelectionColumn final : public FTypedElementDataStorageColumn
+struct FTypedElementSelectionColumn final : public FEditorDataStorageColumn
 {
 	GENERATED_BODY()
 
diff --git a/Engine/Source/Runtime/TypedElementFramework/Public/Elements/Columns/TypedElementSlateWidgetColumns.h b/Engine/Source/Runtime/TypedElementFramework/Public/Elements/Columns/TypedElementSlateWidgetColumns.h
--- a/Engine/Source/Runtime/TypedElementFramework/Public/Elements/Columns/TypedElementSlateWidgetColumns.h
+++ b/Engine/Source/Runtime/TypedElementFramework/Public/Elements/Columns/TypedElementSlateWidgetColumns.h
@@ -2,9 +2,8 @@
 
 #pragma once
 
-#include "Elements/Interfaces/TypedElementDataStorageInterface.h"
+#include "Elements/Common/TypedElementCommonTypes.h"
 #include "Templates/SharedPointer.h"
-#include "UObject/ObjectMacros.h"
 
 #include "TypedElementSlateWidgetColumns.generated.h"
 
@@ -19,7 +18,7 @@ struct FTypedElementWidgetConstructor;
  * be deleted.
  */
 USTRUCT(meta = (DisplayName = "Slate widget reference"))
-struct FTypedElementSlateWidgetReferenceColumn final : public FTypedElementDataStorageColumn
+struct FTypedElementSlateWidgetReferenceColumn final : public FEditorDataStorageColumn
 {
 	GENERATED_BODY()
 
@@ -39,7 +38,7 @@ struct FTypedElementSlateWidgetReferenceColumn final : public FTypedElementDataS
  * the column will be removed.
  */
 USTRUCT(meta = (DisplayName = "Slate widget reference deletes row"))
-struct FTypedElementSlateWidgetReferenceDeletesRowTag final : public FTypedElementDataStorageTag
+struct FTypedElementSlateWidgetReferenceDeletesRowTag final : public FEditorDataStorageTag
 {
 	GENERATED_BODY()
 };
\ No newline at end of file
diff --git a/Engine/Source/Runtime/TypedElementFramework/Public/Elements/Columns/TypedElementTransformColumns.h b/Engine/Source/Runtime/TypedElementFramework/Public/Elements/Columns/TypedElementTransformColumns.h
--- a/Engine/Source/Runtime/TypedElementFramework/Public/Elements/Columns/TypedElementTransformColumns.h
+++ b/Engine/Source/Runtime/TypedElementFramework/Public/Elements/Columns/TypedElementTransformColumns.h
@@ -2,9 +2,8 @@
 
 #pragma once
 
-#include "Elements/Interfaces/TypedElementDataStorageInterface.h"
+#include "Elements/Common/TypedElementCommonTypes.h"
 #include "Math/MathFwd.h"
-#include "UObject/ObjectMacros.h"
 
 #include "TypedElementTransformColumns.generated.h"
 
@@ -13,12 +12,12 @@
  * Column that stores a local transform. 
  */
 USTRUCT(meta = (DisplayName = "Local Transform"))
-struct FTypedElementLocalTransformColumn final : public FTypedElementDataStorageColumn
+struct FTypedElementLocalTransformColumn final : public FEditorDataStorageColumn
 {
 	GENERATED_BODY()
 
-	// Transform is not being intialized to avoid spending time on initialization when the
-	// Transfrom will be updated the first and following ticks after it's creation. If this
+	// Transform is not being initialized to avoid spending time on initialization when the
+	// Transform will be updated the first and following ticks after it's creation. If this
 	// isn't initialized at the correct time, then the sync from source or the true initialization
 	// need to be moved to an earlier phase or group.
 	UPROPERTY(meta = (IgnoreForMemberInitializationTest))
diff --git a/Engine/Source/Runtime/TypedElementFramework/Public/Elements/Columns/TypedElementTypeInfoColumns.h b/Engine/Source/Runtime/TypedElementFramework/Public/Elements/Columns/TypedElementTypeInfoColumns.h
--- a/Engine/Source/Runtime/TypedElementFramework/Public/Elements/Columns/TypedElementTypeInfoColumns.h
+++ b/Engine/Source/Runtime/TypedElementFramework/Public/Elements/Columns/TypedElementTypeInfoColumns.h
@@ -2,8 +2,7 @@
 
 #pragma once
 
-#include "Elements/Interfaces/TypedElementDataStorageInterface.h"
-#include "UObject/ObjectMacros.h"
+#include "Elements/Common/TypedElementCommonTypes.h"
 #include "UObject/WeakObjectPtrTemplates.h"
 
 #include "TypedElementTypeInfoColumns.generated.h"
@@ -15,7 +14,7 @@ class UScriptStruct;
  * Column that stores type information for classes.
  */
 USTRUCT(meta = (DisplayName = "Type"))
-struct FTypedElementClassTypeInfoColumn final : public FTypedElementDataStorageColumn
+struct FTypedElementClassTypeInfoColumn final : public FEditorDataStorageColumn
 {
 	GENERATED_BODY()
 
@@ -26,7 +25,7 @@ struct FTypedElementClassTypeInfoColumn final : public FTypedElementDataStorageC
  * Column that stores type information for structs.
  */
 USTRUCT(meta = (DisplayName = "ScriptStruct type info"))
-struct FTypedElementScriptStructTypeInfoColumn final : public FTypedElementDataStorageColumn
+struct FTypedElementScriptStructTypeInfoColumn final : public FEditorDataStorageColumn
 {
 	GENERATED_BODY()
 
diff --git a/Engine/Source/Runtime/TypedElementFramework/Public/Elements/Columns/TypedElementValueCacheColumns.h b/Engine/Source/Runtime/TypedElementFramework/Public/Elements/Columns/TypedElementValueCacheColumns.h
--- a/Engine/Source/Runtime/TypedElementFramework/Public/Elements/Columns/TypedElementValueCacheColumns.h
+++ b/Engine/Source/Runtime/TypedElementFramework/Public/Elements/Columns/TypedElementValueCacheColumns.h
@@ -2,8 +2,7 @@
 
 #pragma once
 
-#include "Elements/Interfaces/TypedElementDataStorageInterface.h"
-#include "UObject/ObjectMacros.h"
+#include "Elements/Common/TypedElementCommonTypes.h"
 
 #include "TypedElementValueCacheColumns.generated.h"
 
@@ -12,7 +11,7 @@
  * Column that can be used to cache an unsigned 32-bit value in.
  */
 USTRUCT(meta = (DisplayName = "Uint32 value cache"))
-struct FTypedElementU32IntValueCacheColumn final : public FTypedElementDataStorageColumn
+struct FTypedElementU32IntValueCacheColumn final : public FEditorDataStorageColumn
 {
 	GENERATED_BODY()
 
@@ -24,7 +23,7 @@ struct FTypedElementU32IntValueCacheColumn final : public FTypedElementDataStora
  * Column that can be used to cache a signed 32-bit value in.
  */
 USTRUCT(meta = (DisplayName = "Int32 value cache"))
-struct FTypedElementI32IntValueCacheColumn final : public FTypedElementDataStorageColumn
+struct FTypedElementI32IntValueCacheColumn final : public FEditorDataStorageColumn
 {
 	GENERATED_BODY()
 
@@ -36,7 +35,7 @@ struct FTypedElementI32IntValueCacheColumn final : public FTypedElementDataStora
  * Column that can be used to cache an unsigned 64-bit value in.
  */
 USTRUCT(meta = (DisplayName = "Uint64 value cache"))
-struct FTypedElementU64IntValueCacheColumn final : public FTypedElementDataStorageColumn
+struct FTypedElementU64IntValueCacheColumn final : public FEditorDataStorageColumn
 {
 	GENERATED_BODY()
 
@@ -48,7 +47,7 @@ struct FTypedElementU64IntValueCacheColumn final : public FTypedElementDataStora
  * Column that can be used to cache a signed 64-bit value in.
  */
 USTRUCT(meta = (DisplayName = "Int64 value cache"))
-struct FTypedElementI64IntValueCacheColumn final : public FTypedElementDataStorageColumn
+struct FTypedElementI64IntValueCacheColumn final : public FEditorDataStorageColumn
 {
 	GENERATED_BODY()
 
@@ -60,7 +59,7 @@ struct FTypedElementI64IntValueCacheColumn final : public FTypedElementDataStora
  * Column that can be used to cache a 32-bit floating point value in.
  */
 USTRUCT(meta = (DisplayName = "float value cache"))
-struct FTypedElementFloatValueCacheColumn final : public FTypedElementDataStorageColumn
+struct FTypedElementFloatValueCacheColumn final : public FEditorDataStorageColumn
 {
 	GENERATED_BODY()
 
diff --git a/Engine/Source/Runtime/TypedElementFramework/Public/Elements/Columns/TypedElementViewportColumns.h b/Engine/Source/Runtime/TypedElementFramework/Public/Elements/Columns/TypedElementViewportColumns.h
--- a/Engine/Source/Runtime/TypedElementFramework/Public/Elements/Columns/TypedElementViewportColumns.h
+++ b/Engine/Source/Runtime/TypedElementFramework/Public/Elements/Columns/TypedElementViewportColumns.h
@@ -2,8 +2,7 @@
 
 #pragma once
 
-#include "Elements/Interfaces/TypedElementDataStorageInterface.h"
-#include "UObject/ObjectMacros.h"
+#include "Elements/Common/TypedElementCommonTypes.h"
 
 #include "TypedElementViewportColumns.generated.h"
 
@@ -11,7 +10,7 @@
  * Column to hold the color that the object is outlined with in the viewport
  */
 USTRUCT(meta = (DisplayName = "Viewport Outline Color"))
-struct FTypedElementViewportOutlineColorColumn final : public FTypedElementDataStorageColumn
+struct FTypedElementViewportOutlineColorColumn final : public FEditorDataStorageColumn
 {
 	GENERATED_BODY()
 
@@ -23,7 +22,7 @@ struct FTypedElementViewportOutlineColorColumn final : public FTypedElementDataS
  * Column to hold the color that the object is overlaid with in the viewport
  */
 USTRUCT(meta = (DisplayName = "Viewport Overlay Color"))
-struct FTypedElementViewportOverlayColorColumn final : public FTypedElementDataStorageColumn
+struct FTypedElementViewportOverlayColorColumn final : public FEditorDataStorageColumn
 {
 	GENERATED_BODY()
 
diff --git a/Engine/Source/Runtime/TypedElementFramework/Public/Elements/Common/TypedElementCommonTypes.h b/Engine/Source/Runtime/TypedElementFramework/Public/Elements/Common/TypedElementCommonTypes.h
--- a/Engine/Source/Runtime/TypedElementFramework/Public/Elements/Common/TypedElementCommonTypes.h
+++ b/Engine/Source/Runtime/TypedElementFramework/Public/Elements/Common/TypedElementCommonTypes.h
@@ -2,7 +2,6 @@
 
 #pragma once
 
-#include "CoreMinimal.h"
 #include "Elements/Common/TypedElementHandles.h"
 
 #include "TypedElementCommonTypes.generated.h"
@@ -11,7 +10,7 @@
  * Base for the data structures for a column.
  */
 USTRUCT()
-struct FTypedElementDataStorageColumn
+struct FEditorDataStorageColumn
 {
 	GENERATED_BODY()
 };
@@ -20,7 +19,7 @@ struct FTypedElementDataStorageColumn
  * Base for the data structures that act as tags to rows. Tags should not have any data.
  */
 USTRUCT()
-struct FTypedElementDataStorageTag
+struct FEditorDataStorageTag
 {
 	GENERATED_BODY()
 };
@@ -41,12 +40,12 @@ namespace UE
 
 	template<typename Derived, typename Base>
 	concept derived_from = std::is_base_of_v<Base, Derived> && std::is_convertible_v<const volatile Derived*, const volatile Base*>;
-}
 
-namespace UE
-{
 	namespace Editor::DataStorage
 	{
+		using FColumn = FEditorDataStorageColumn;
+		using FTag = FEditorDataStorageTag;
+
 		/**
 		 * Defines a dynamic type for a dynamic tag
 		 * Example:
@@ -83,35 +82,27 @@ namespace UE
 			return Other.Name == Name;
 		}
 
-	}
-}
-
-
-
-namespace TypedElementDataStorage
-{
-	// Standard callbacks.
-
-	using RowCreationCallbackRef = TFunctionRef<void(TypedElementDataStorage::RowHandle Row)>;
-	using ColumnCreationCallbackRef = TFunctionRef<void(void* Column, const UScriptStruct& ColumnType)>;
-	using ColumnListCallbackRef = TFunctionRef<void(const UScriptStruct& ColumnType)>;
-	using ColumnListWithDataCallbackRef = TFunctionRef<void(void* Column, const UScriptStruct& ColumnType)>;
-	using ColumnCopyOrMoveCallback = void (*)(const UScriptStruct& ColumnType, void* Destination, void* Source);
-
+		// Standard callbacks.
 
+		using RowCreationCallbackRef = TFunctionRef<void(TypedElementDataStorage::RowHandle Row)>;
+		using ColumnCreationCallbackRef = TFunctionRef<void(void* Column, const UScriptStruct& ColumnType)>;
+		using ColumnListCallbackRef = TFunctionRef<void(const UScriptStruct& ColumnType)>;
+		using ColumnListWithDataCallbackRef = TFunctionRef<void(void* Column, const UScriptStruct& ColumnType)>;
+		using ColumnCopyOrMoveCallback = void (*)(const UScriptStruct& ColumnType, void* Destination, void* Source);
 
-	// Template concepts to enforce type correctness.
 
-	template<typename T>
-	concept TDataColumnType = UE::derived_from<T, FTypedElementDataStorageColumn>;
 
-	template<typename T>
-	concept TTagColumnType = UE::derived_from<T, FTypedElementDataStorageTag>;
+		// Template concepts to enforce type correctness.
+		template<typename T>
+		concept TDataColumnType = UE::derived_from<T, FColumn>;
 
-	template<typename T>
-	concept TColumnType = TDataColumnType<T> || TTagColumnType<T>;
+		template<typename T>
+		concept TTagColumnType = UE::derived_from<T, FTag>;
 
-	template<typename T>
-	concept TEnumType = std::is_enum_v<T>;
+		template<typename T>
+		concept TColumnType = TDataColumnType<T> || TTagColumnType<T>;
 
-} // namespace TypedElementDataStorage
\ No newline at end of file
+		template<typename T>
+		concept TEnumType = std::is_enum_v<T>;
+	} // namespace Editor::DataStorage
+} // namespace UE
\ No newline at end of file
diff --git a/Engine/Source/Runtime/TypedElementFramework/Public/Elements/Common/TypedElementQueryTypes.h b/Engine/Source/Runtime/TypedElementFramework/Public/Elements/Common/TypedElementQueryTypes.h
--- a/Engine/Source/Runtime/TypedElementFramework/Public/Elements/Common/TypedElementQueryTypes.h
+++ b/Engine/Source/Runtime/TypedElementFramework/Public/Elements/Common/TypedElementQueryTypes.h
@@ -6,6 +6,8 @@
 
 namespace TypedElementDataStorage
 {
+	using IndexHash = uint64;
+
 	enum class EQueryTickPhase : uint8
 	{
 		PrePhysics, //< Queries are executed before physics simulation starts.
diff --git a/Engine/Source/Runtime/TypedElementFramework/Public/Elements/Framework/TypedElementAttributeBinding.h b/Engine/Source/Runtime/TypedElementFramework/Public/Elements/Framework/TypedElementAttributeBinding.h
--- a/Engine/Source/Runtime/TypedElementFramework/Public/Elements/Framework/TypedElementAttributeBinding.h
+++ b/Engine/Source/Runtime/TypedElementFramework/Public/Elements/Framework/TypedElementAttributeBinding.h
@@ -10,6 +10,8 @@ class ITypedElementDataStorageInterface;
 
 namespace UE::EditorDataStorage
 {
+	using namespace UE::Editor::DataStorage;
+
 	// Concept for a conversion function used by the attribute binder to bind a column data member to an attribute of a different type
 	template <typename FunctionType, typename ArgumentType>
 	concept AttributeBinderInvocable = std::is_invocable_v<std::decay_t<FunctionType>, const ArgumentType&>;
@@ -41,7 +43,7 @@ namespace UE::EditorDataStorage
 		 * TypedElementDataStorage::FAttributeBinder Binder(RowHandle);
 		 * TAttribute<FString> TestAttribute(Binder.BindData(&FTypedElementLabelColumn::Label))
 		 */
-		template <typename AttributeType, TypedElementDataStorage::TDataColumnType ColumnType>
+		template <typename AttributeType, TDataColumnType ColumnType>
 		TAttribute<AttributeType> BindData(AttributeType ColumnType::* InVariable, const AttributeType& InDefaultValue = AttributeType());
 
 		/**
@@ -62,7 +64,7 @@ namespace UE::EditorDataStorage
 		 *                                   }
 		 *                                 ));
 		 */
-		template <typename AttributeType, typename DataType, TypedElementDataStorage::TDataColumnType ColumnType>
+		template <typename AttributeType, typename DataType, TDataColumnType ColumnType>
 		TAttribute<AttributeType> BindData(DataType ColumnType::* InVariable, const TFunction<AttributeType(const DataType&)>& InConverter, const DataType& InDefaultValue = DataType());
 
 		/**
@@ -73,7 +75,7 @@ namespace UE::EditorDataStorage
 		 * @param InDefaultValue The default value to be used when the column isn't present on a row
 		 * @return A TAttribute bound to the row, column pair specified
 		 */
-		template <typename DataType, TypedElementDataStorage::TDataColumnType ColumnType, typename FunctionType>
+		template <typename DataType, TDataColumnType ColumnType, typename FunctionType>
 			requires AttributeBinderInvocable<FunctionType, DataType>
 		auto BindData(DataType ColumnType::* InVariable, FunctionType InConverter, const DataType& InDefaultValue = DataType());
 		
diff --git a/Engine/Source/Runtime/TypedElementFramework/Public/Elements/Framework/TypedElementAttributeBinding.inl b/Engine/Source/Runtime/TypedElementFramework/Public/Elements/Framework/TypedElementAttributeBinding.inl
--- a/Engine/Source/Runtime/TypedElementFramework/Public/Elements/Framework/TypedElementAttributeBinding.inl
+++ b/Engine/Source/Runtime/TypedElementFramework/Public/Elements/Framework/TypedElementAttributeBinding.inl
@@ -6,6 +6,8 @@
 
 namespace UE::EditorDataStorage
 {
+	using namespace UE::Editor::DataStorage;
+
 	// A property that can be directly be accessed from an object
 	template <typename PropertyType>
 	struct DirectProperty final
@@ -51,7 +53,7 @@ namespace UE::EditorDataStorage
 	public:
 		
 		// Bind this property directly
-		template <TypedElementDataStorage::TDataColumnType ObjectType>
+		template <TDataColumnType ObjectType>
 		void Bind(PropertyType ObjectType::* Variable)
 		{
 			DirectProperty<PropertyType> Result;
@@ -62,7 +64,7 @@ namespace UE::EditorDataStorage
 		}
 
 		// Bind this property using a conversion function
-		template <typename InputType, TypedElementDataStorage::TDataColumnType ObjectType>
+		template <typename InputType, TDataColumnType ObjectType>
 		void Bind(InputType ObjectType::* Variable, TFunction<PropertyType(const InputType&)> Converter)
 		{
 			ConvertibleProperty<PropertyType> Result;
@@ -76,14 +78,14 @@ namespace UE::EditorDataStorage
 		}
 
 		// Get the bound property for the specified object
-		template <TypedElementDataStorage::TDataColumnType ObjectType>
+		template <TDataColumnType ObjectType>
 		PropertyType& Get(ObjectType& Object) const
 		{
 			return Get(&Object, ObjectType::StaticStruct());
 		}
 
 		// Get the bound property for the specified object
-		template <TypedElementDataStorage::TDataColumnType ObjectType>
+		template <TDataColumnType ObjectType>
 		const PropertyType& Get(const ObjectType& Object) const
 		{
 			return Get(&Object, ObjectType::StaticStruct());
@@ -119,7 +121,7 @@ namespace UE::EditorDataStorage
 		UScriptStruct* ObjectTypeInfo = nullptr;
 	};
 	
-	template <typename AttributeType, TypedElementDataStorage::TDataColumnType ColumnType>
+	template <typename AttributeType, TDataColumnType ColumnType>
 	TAttribute<AttributeType> FAttributeBinder::BindData(AttributeType ColumnType::* InVariable, const AttributeType& InDefaultValue)
 	{
 		if(!DataStorage)
@@ -143,7 +145,7 @@ namespace UE::EditorDataStorage
 		});
 	}
 
-	template <typename AttributeType, typename DataType, TypedElementDataStorage::TDataColumnType ColumnType>
+	template <typename AttributeType, typename DataType, TDataColumnType ColumnType>
 	TAttribute<AttributeType> FAttributeBinder::BindData(DataType ColumnType::* InVariable, const TFunction<AttributeType(const DataType&)>& InConverter, const DataType& InDefaultValue)
 	{
 		if(!DataStorage)
@@ -166,7 +168,7 @@ namespace UE::EditorDataStorage
 		});
 	}
 	
-	template <typename DataType, TypedElementDataStorage::TDataColumnType ColumnType, typename FunctionType>
+	template <typename DataType, TDataColumnType ColumnType, typename FunctionType>
 		requires AttributeBinderInvocable<FunctionType, DataType>
 	auto FAttributeBinder::BindData(DataType ColumnType::* InVariable, FunctionType InConverter, const DataType& InDefaultValue)
 	{
diff --git a/Engine/Source/Runtime/TypedElementFramework/Public/Elements/Framework/TypedElementQueryBuilder.h b/Engine/Source/Runtime/TypedElementFramework/Public/Elements/Framework/TypedElementQueryBuilder.h
--- a/Engine/Source/Runtime/TypedElementFramework/Public/Elements/Framework/TypedElementQueryBuilder.h
+++ b/Engine/Source/Runtime/TypedElementFramework/Public/Elements/Framework/TypedElementQueryBuilder.h
@@ -152,16 +152,16 @@ namespace TypedElementQueryBuilder
 		TYPEDELEMENTFRAMEWORK_API FDependency DependsOn();
 		TYPEDELEMENTFRAMEWORK_API ITypedElementDataStorageInterface::FQueryDescription&& Compile();
 
-		template<TypedElementDataStorage::TColumnType... TargetTypes>
+		template<UE::Editor::DataStorage::TColumnType... TargetTypes>
 		FSimpleQuery& All();
 		TYPEDELEMENTFRAMEWORK_API FSimpleQuery& All(const UScriptStruct* Target);
 		TYPEDELEMENTFRAMEWORK_API FSimpleQuery& All(TConstArrayView<const UScriptStruct*> Targets);
 		
-		template<TypedElementDataStorage::TColumnType... TargetTypes>
+		template<UE::Editor::DataStorage::TColumnType... TargetTypes>
 		FSimpleQuery& Any();
 		TYPEDELEMENTFRAMEWORK_API FSimpleQuery& Any(const UScriptStruct* Target);
 		TYPEDELEMENTFRAMEWORK_API FSimpleQuery& Any(TConstArrayView<const UScriptStruct*> Targets);
-		template<TypedElementDataStorage::TColumnType... TargetTypes>
+		template<UE::Editor::DataStorage::TColumnType... TargetTypes>
 		FSimpleQuery& None();
 		TYPEDELEMENTFRAMEWORK_API FSimpleQuery& None(const UScriptStruct* Target);
 		TYPEDELEMENTFRAMEWORK_API FSimpleQuery& None(TConstArrayView<const UScriptStruct*> Targets);
@@ -192,13 +192,13 @@ namespace TypedElementQueryBuilder
 		template<>
 		TYPEDELEMENTFRAMEWORK_API FSimpleQuery& All<UE::Editor::DataStorage::FDynamicTag>(const FName& Tag, const FName& Value);
 		
-		template<TypedElementDataStorage::TEnumType EnumT>
+		template<UE::Editor::DataStorage::TEnumType EnumT>
 		FSimpleQuery& All();
 		
-		template<TypedElementDataStorage::TEnumType EnumT>
+		template<UE::Editor::DataStorage::TEnumType EnumT>
 		FSimpleQuery& All(EnumT EnumValue);
 
-		template<auto Value, TypedElementDataStorage::TEnumType EnumT = decltype(Value)>
+		template<auto Value, UE::Editor::DataStorage::TEnumType EnumT = decltype(Value)>
 		FSimpleQuery& All();
 
 	private:
@@ -237,14 +237,14 @@ namespace TypedElementQueryBuilder
 
 		TYPEDELEMENTFRAMEWORK_API FObserver(EEvent MonitorForEvent, const UScriptStruct* MonitoredColumn);
 
-		template<TypedElementDataStorage::TColumnType ColumnType>
+		template<UE::Editor::DataStorage::TColumnType ColumnType>
 		static FObserver OnAdd();
-		template<TypedElementDataStorage::TColumnType ColumnType>
+		template<UE::Editor::DataStorage::TColumnType ColumnType>
 		static FObserver OnRemove();
 
 		TYPEDELEMENTFRAMEWORK_API FObserver& SetEvent(EEvent MonitorForEvent);
 		TYPEDELEMENTFRAMEWORK_API FObserver& SetMonitoredColumn(const UScriptStruct* MonitoredColumn);
-		template<TypedElementDataStorage::TColumnType ColumnType>
+		template<UE::Editor::DataStorage::TColumnType ColumnType>
 		FObserver& SetMonitoredColumn();
 		TYPEDELEMENTFRAMEWORK_API FObserver& SetExecutionMode(TypedElementDataStorage::EExecutionMode Mode);
 		TYPEDELEMENTFRAMEWORK_API FObserver& MakeActivatable(FName Name);
@@ -347,7 +347,7 @@ namespace TypedElementQueryBuilder
 		Select(FName Name, const CallbackType& Type, Class* Instance, Function&& Callback);
 
 		/** Request read-only access to the listed columns. */
-		template<TypedElementDataStorage::TDataColumnType... TargetTypes>
+		template<UE::Editor::DataStorage::TDataColumnType... TargetTypes>
 		Select& ReadOnly();
 		/** Request read-only access to the listed columns. */
 		TYPEDELEMENTFRAMEWORK_API Select& ReadOnly(const UScriptStruct* Target);
@@ -357,7 +357,7 @@ namespace TypedElementQueryBuilder
 		 * Request read-only access to the listed columns. If optional is true read access will be given if the column is in the table but
 		 * it will not be used for finding matching tables. Columns bound with optional can not be bound to a query callback argument.
 		 */
-		template<TypedElementDataStorage::TDataColumnType... TargetTypes>
+		template<UE::Editor::DataStorage::TDataColumnType... TargetTypes>
 		Select& ReadOnly(EOptional Optional);
 		/**
 		 * Request read-only access to the listed columns. If optional is true read access will be given if the column is in the table but
@@ -370,7 +370,7 @@ namespace TypedElementQueryBuilder
 		 */
 		TYPEDELEMENTFRAMEWORK_API Select& ReadOnly(TConstArrayView<const UScriptStruct*> Targets, EOptional Optionall);
 		/** Request read and write access to the listed columns. */
-		template<TypedElementDataStorage::TDataColumnType... TargetTypes>
+		template<UE::Editor::DataStorage::TDataColumnType... TargetTypes>
 		Select& ReadWrite();
 		/** Request read and write access to the listed columns. */
 		TYPEDELEMENTFRAMEWORK_API Select& ReadWrite(const UScriptStruct* Target);
diff --git a/Engine/Source/Runtime/TypedElementFramework/Public/Elements/Framework/TypedElementQueryBuilder.inl b/Engine/Source/Runtime/TypedElementFramework/Public/Elements/Framework/TypedElementQueryBuilder.inl
--- a/Engine/Source/Runtime/TypedElementFramework/Public/Elements/Framework/TypedElementQueryBuilder.inl
+++ b/Engine/Source/Runtime/TypedElementFramework/Public/Elements/Framework/TypedElementQueryBuilder.inl
@@ -12,6 +12,8 @@
 
 namespace TypedElementQueryBuilder
 {
+	using namespace UE::Editor::DataStorage;
+
 	namespace Internal
 	{
 		// This assumes that the types are unique, but for queries this should be true and otherwise
@@ -90,19 +92,19 @@ namespace TypedElementQueryBuilder
 	// FObserver
 	//
 	
-	template<TypedElementDataStorage::TColumnType ColumnType>
+	template<TColumnType ColumnType>
 	FObserver FObserver::OnAdd()
 	{
 		return FObserver(FObserver::EEvent::Add, ColumnType::StaticStruct());
 	}
 
-	template<TypedElementDataStorage::TColumnType ColumnType>
+	template<TColumnType ColumnType>
 	FObserver FObserver::OnRemove()
 	{
 		return FObserver(FObserver::EEvent::Remove, ColumnType::StaticStruct());
 	}
 
-	template<TypedElementDataStorage::TColumnType ColumnType>
+	template<TColumnType ColumnType>
 	FObserver& FObserver::SetMonitoredColumn()
 	{
 		return SetMonitoredColumn(ColumnType::StaticStruct());
@@ -1117,21 +1119,21 @@ e.g. void(FCachedQueryContext<Subsystem1, const Subsystem2>& Context, TypedEleme
 		Internal::BindQueryFunction<TypedElementDataStorage::IQueryContext, ValidateColumns>(Query.Callback.Function, Instance, Callback);
 	}
 
-	template<TypedElementDataStorage::TDataColumnType... TargetTypes>
+	template<TDataColumnType... TargetTypes>
 	Select& Select::ReadOnly()
 	{
 		ReadOnly({ TargetTypes::StaticStruct()... });
 		return *this;
 	}
 
-	template<TypedElementDataStorage::TDataColumnType... TargetTypes>
+	template<TDataColumnType... TargetTypes>
 	Select& Select::ReadOnly(EOptional Optional)
 	{
 		ReadOnly({ TargetTypes::StaticStruct()... }, Optional);
 		return *this;
 	}
 
-	template<TypedElementDataStorage::TDataColumnType... TargetTypes>
+	template<TDataColumnType... TargetTypes>
 	Select& Select::ReadWrite()
 	{
 		ReadWrite({ TargetTypes::StaticStruct()... });
@@ -1143,40 +1145,40 @@ e.g. void(FCachedQueryContext<Subsystem1, const Subsystem2>& Context, TypedEleme
 	// FSimpleQuery
 	//
 
-	template<TypedElementDataStorage::TColumnType... TargetTypes>
+	template<TColumnType... TargetTypes>
 	FSimpleQuery& FSimpleQuery::All()
 	{
 		All({ TargetTypes::StaticStruct()... });
 		return *this;
 	}
 
-	template<TypedElementDataStorage::TColumnType... TargetTypes>
+	template<TColumnType... TargetTypes>
 	FSimpleQuery& FSimpleQuery::Any()
 	{
 		Any({ TargetTypes::StaticStruct()... });
 		return *this;
 	}
 
-	template<TypedElementDataStorage::TColumnType... TargetTypes>
+	template<TColumnType... TargetTypes>
 	FSimpleQuery& FSimpleQuery::None()
 	{
 		None({ TargetTypes::StaticStruct()... });
 		return *this;
 	}
 	
-	template <TypedElementDataStorage::TEnumType EnumT>
+	template <TEnumType EnumT>
     FSimpleQuery& FSimpleQuery::All()
     {
     	return All(static_cast<const UEnum&>(*StaticEnum<EnumT>()));
     }
 
-    template <TypedElementDataStorage::TEnumType EnumT>
+    template <TEnumType EnumT>
     FSimpleQuery& FSimpleQuery::All(EnumT EnumValue)
     {
     	return All(static_cast<const UEnum&>(*StaticEnum<EnumT>()), static_cast<int64>(EnumValue));
     }
 	
-	template<auto Value, TypedElementDataStorage::TEnumType EnumT = decltype(Value)>
+	template<auto Value, TEnumType EnumT = decltype(Value)>
 	FSimpleQuery& All()
 	{
 		return All<EnumT>(Value);
diff --git a/Engine/Source/Runtime/TypedElementFramework/Public/Elements/Interfaces/TypedElementDataStorageInterface.h b/Engine/Source/Runtime/TypedElementFramework/Public/Elements/Interfaces/TypedElementDataStorageInterface.h
--- a/Engine/Source/Runtime/TypedElementFramework/Public/Elements/Interfaces/TypedElementDataStorageInterface.h
+++ b/Engine/Source/Runtime/TypedElementFramework/Public/Elements/Interfaces/TypedElementDataStorageInterface.h
@@ -50,7 +50,7 @@ class UTypedElementDataStorageInterface : public UInterface
  * multiple column lists are used. Note that the returned array view is only available while
  * this object is constructed, so care must be taken with functions that return a const array view.
  */
-template<TypedElementDataStorage::TColumnType... Columns>
+template<UE::Editor::DataStorage::TColumnType... Columns>
 struct TTypedElementColumnTypeList
 {
 	const UScriptStruct* ColumnTypes[sizeof...(Columns)] = { Columns::StaticStruct()... };
@@ -62,6 +62,12 @@ class ITypedElementDataStorageInterface
 {
 	GENERATED_BODY()
 
+	using RowCreationCallbackRef = UE::Editor::DataStorage::RowCreationCallbackRef;
+	using ColumnCreationCallbackRef = UE::Editor::DataStorage::ColumnCreationCallbackRef;
+	using ColumnListCallbackRef = UE::Editor::DataStorage::ColumnListCallbackRef;
+	using ColumnListWithDataCallbackRef = UE::Editor::DataStorage::ColumnListWithDataCallbackRef;
+	using ColumnCopyOrMoveCallback = UE::Editor::DataStorage::ColumnCopyOrMoveCallback;
+
 public:
 	/**
 	 * @section Factories
@@ -89,7 +95,7 @@ public:
 
 	/** Creates a new table for with the provided columns. Optionally a name can be given which is useful for retrieval later. */
 	virtual TypedElementDataStorage::TableHandle RegisterTable(TConstArrayView<const UScriptStruct*> ColumnList, const FName Name) = 0;
-	template<TypedElementDataStorage::TColumnType... Columns>
+	template<UE::Editor::DataStorage::TColumnType... Columns>
 	TypedElementDataStorage::TableHandle RegisterTable(const FName Name);
 	/** 
 	 * Copies the column information from the provided table and creates a new table for with the provided columns. Optionally a 
@@ -97,7 +103,7 @@ public:
 	 */
 	virtual TypedElementDataStorage::TableHandle RegisterTable(TypedElementDataStorage::TableHandle SourceTable,
 		TConstArrayView<const UScriptStruct*> ColumnList, const FName Name) = 0;
-	template<TypedElementDataStorage::TColumnType... Columns>
+	template<UE::Editor::DataStorage::TColumnType... Columns>
 	TypedElementDataStorage::TableHandle RegisterTable(TypedElementDataStorage::TableHandle SourceTable, const FName Name);
 
 	/** Returns a previously created table with the provided name or TypedElementInvalidTableHandle if not found. */
@@ -134,7 +140,7 @@ public:
 	 * initialize the row if needed.
 	 */
 	virtual TypedElementDataStorage::RowHandle AddRow(TypedElementDataStorage::TableHandle Table,
-		TypedElementDataStorage::RowCreationCallbackRef OnCreated) = 0;
+		RowCreationCallbackRef OnCreated) = 0;
 	/** Adds a new row to the provided table using a previously reserved row. */
 	virtual bool AddRow(TypedElementDataStorage::RowHandle ReservedRow, TypedElementDataStorage::TableHandle Table) = 0;
 	/**
@@ -142,20 +148,20 @@ public:
 	 * initialize the row if needed.
 	 */
 	virtual bool AddRow(TypedElementDataStorage::RowHandle ReservedRow, TypedElementDataStorage::TableHandle Table,
-		TypedElementDataStorage::RowCreationCallbackRef OnCreated) = 0;
+		RowCreationCallbackRef OnCreated) = 0;
 
 	/**
 	 * Add multiple rows at once. For each new row the OnCreated callback is called. Callers are expected to use the callback to
 	 * initialize the row if needed.
 	 */
-	virtual bool BatchAddRow(TypedElementDataStorage::TableHandle Table, int32 Count, TypedElementDataStorage::RowCreationCallbackRef OnCreated) = 0;
+	virtual bool BatchAddRow(TypedElementDataStorage::TableHandle Table, int32 Count, RowCreationCallbackRef OnCreated) = 0;
 	/**
 	 * Add multiple rows at once. For each new row the OnCreated callback is called. Callers are expected to use the callback to
 	 * initialize the row if needed. This version uses a set of previously reserved rows. Any row that can't be used will be 
 	 * released.
 	 */
 	virtual bool BatchAddRow(TypedElementDataStorage::TableHandle Table, TConstArrayView<TypedElementDataStorage::RowHandle> ReservedHandles,
-		TypedElementDataStorage::RowCreationCallbackRef OnCreated) = 0;
+		RowCreationCallbackRef OnCreated) = 0;
 
 	/** Removes a previously reserved or added row. If the row handle is invalid or already removed, nothing happens */
 	virtual void RemoveRow(TypedElementDataStorage::RowHandle Row) = 0;
@@ -172,16 +178,16 @@ public:
 
 	/** Adds a column to a row or does nothing if already added. */
 	virtual void AddColumn(TypedElementRowHandle Row, const UScriptStruct* ColumnType) = 0;
-	template<TypedElementDataStorage::TColumnType ColumnType>
+	template<UE::Editor::DataStorage::TColumnType ColumnType>
 	void AddColumn(TypedElementRowHandle Row);
 	/**
 	 * Adds a new data column and initializes it. The relocator will be used to copy or move the column out of
 	 * its temporary location into the final table if the addition needs to be deferred.
 	 */
 	virtual void AddColumnData(TypedElementRowHandle Row, const UScriptStruct* ColumnType,
-		const TypedElementDataStorage::ColumnCreationCallbackRef& Initializer,
-		TypedElementDataStorage::ColumnCopyOrMoveCallback Relocator) = 0;
-	template<TypedElementDataStorage::TDataColumnType ColumnType>
+		const ColumnCreationCallbackRef& Initializer,
+		ColumnCopyOrMoveCallback Relocator) = 0;
+	template<UE::Editor::DataStorage::TDataColumnType ColumnType>
 	void AddColumn(TypedElementRowHandle Row, ColumnType&& Column);
 
 	/**
@@ -203,10 +209,10 @@ public:
 	template<>
 	void AddColumn<UE::Editor::DataStorage::FDynamicTag>(TypedElementDataStorage::RowHandle Row, const FName& Tag, const FName& Value);
 
-	template<TypedElementDataStorage::TEnumType EnumT>
+	template<UE::Editor::DataStorage::TEnumType EnumT>
 	void AddColumn(TypedElementDataStorage::RowHandle Row, EnumT Value);
 	
-	template<auto Value, TypedElementDataStorage::TEnumType EnumT = decltype(Value)>
+	template<auto Value, UE::Editor::DataStorage::TEnumType EnumT = decltype(Value)>
 	void AddColumn(TypedElementDataStorage::RowHandle Row);
 
 	/**
@@ -214,15 +220,15 @@ public:
 	 * at a time.
 	 */
 	virtual void AddColumns(TypedElementRowHandle Row, TConstArrayView<const UScriptStruct*> Columns) = 0;
-	template<TypedElementDataStorage::TColumnType... Columns>
+	template<UE::Editor::DataStorage::TColumnType... Columns>
 	void AddColumns(TypedElementRowHandle Row);
 
 	/** Removes a column from a row or does nothing if already removed. */
 	virtual void RemoveColumn(TypedElementRowHandle Row, const UScriptStruct* ColumnType) = 0;
-	template<TypedElementDataStorage::TColumnType Column>
+	template<UE::Editor::DataStorage::TColumnType Column>
 	void RemoveColumn(TypedElementRowHandle Row);
 
-	template<TypedElementDataStorage::TEnumType EnumT>
+	template<UE::Editor::DataStorage::TEnumType EnumT>
 	void RemoveColumn(TypedElementDataStorage::RowHandle Row);
 
 	/**
@@ -242,7 +248,7 @@ public:
 	 * at a time.
 	 */
 	virtual void RemoveColumns(TypedElementRowHandle Row, TConstArrayView<const UScriptStruct*> Columns) = 0;
-	template<TypedElementDataStorage::TColumnType... Columns>
+	template<UE::Editor::DataStorage::TColumnType... Columns>
 	void RemoveColumns(TypedElementRowHandle Row);
 
 	/** 
@@ -263,25 +269,25 @@ public:
 	virtual void* GetColumnData(TypedElementRowHandle Row, const UScriptStruct* ColumnType) = 0;
 	virtual const void* GetColumnData(TypedElementRowHandle Row, const UScriptStruct* ColumnType) const = 0;
 	/** Returns a pointer to the column of the given row or a nullptr if the type couldn't be found or the row doesn't exist. */
-	template<TypedElementDataStorage::TDataColumnType ColumnType>
+	template<UE::Editor::DataStorage::TDataColumnType ColumnType>
 	ColumnType* GetColumn(TypedElementRowHandle Row);
-	template<TypedElementDataStorage::TDataColumnType ColumnType>
+	template<UE::Editor::DataStorage::TDataColumnType ColumnType>
 	const ColumnType* GetColumn(TypedElementRowHandle Row) const;
 	
 	/** Determines if the provided row contains the collection of columns and tags. */
 	virtual bool HasColumns(TypedElementRowHandle Row, TConstArrayView<const UScriptStruct*> ColumnTypes) const = 0;
 	virtual bool HasColumns(TypedElementRowHandle Row, TConstArrayView<TWeakObjectPtr<const UScriptStruct>> ColumnTypes) const = 0;
-	template<TypedElementDataStorage::TColumnType... ColumnTypes>
+	template<UE::Editor::DataStorage::TColumnType... ColumnTypes>
 	bool HasColumns(TypedElementRowHandle Row) const;
 
 	/** Lists the columns on a row. This includes data and tag columns. */
-	virtual void ListColumns(TypedElementDataStorage::RowHandle Row, TypedElementDataStorage::ColumnListCallbackRef Callback) const = 0;
+	virtual void ListColumns(TypedElementDataStorage::RowHandle Row, ColumnListCallbackRef Callback) const = 0;
 
 	/** 
 	 * Lists the column type and data on a row. This includes data and tag columns. Not all columns may have data so the data pointer in 
 	 * the callback can be null.
 	 */
-	virtual void ListColumns(TypedElementDataStorage::RowHandle Row, TypedElementDataStorage::ColumnListWithDataCallbackRef Callback) = 0;
+	virtual void ListColumns(TypedElementDataStorage::RowHandle Row, ColumnListWithDataCallbackRef Callback) = 0;
 
 	/** Determines if the columns in the row match the query conditions. */
 	virtual bool MatchesColumns(TypedElementDataStorage::RowHandle Row, const TypedElementDataStorage::FQueryConditions& Conditions) const = 0;
@@ -418,32 +424,32 @@ const FactoryT* ITypedElementDataStorageInterface::FindFactory() const
 	return static_cast<const FactoryT*>(FindFactory(FactoryT::StaticClass()));
 }
 
-template<TypedElementDataStorage::TColumnType... Columns>
+template<UE::Editor::DataStorage::TColumnType... Columns>
 TypedElementDataStorage::TableHandle ITypedElementDataStorageInterface::RegisterTable(const FName Name)
 {
 	return RegisterTable({ Columns::StaticStruct()... }, Name);
 }
 
-template<TypedElementDataStorage::TColumnType... Columns>
+template<UE::Editor::DataStorage::TColumnType... Columns>
 TypedElementDataStorage::TableHandle ITypedElementDataStorageInterface::RegisterTable(
 	TypedElementDataStorage::TableHandle SourceTable, const FName Name)
 {
 	return RegisterTable(SourceTable, { Columns::StaticStruct()... }, Name);
 }
 
-template<TypedElementDataStorage::TColumnType Column>
+template<UE::Editor::DataStorage::TColumnType Column>
 void ITypedElementDataStorageInterface::AddColumn(TypedElementRowHandle Row)
 {
 	AddColumn(Row, Column::StaticStruct());
 }
 
-template<TypedElementDataStorage::TColumnType Column>
+template<UE::Editor::DataStorage::TColumnType Column>
 void ITypedElementDataStorageInterface::RemoveColumn(TypedElementRowHandle Row)
 {
 	RemoveColumn(Row, Column::StaticStruct());
 }
 
-template<TypedElementDataStorage::TColumnType... Columns>
+template<UE::Editor::DataStorage::TColumnType... Columns>
 void ITypedElementDataStorageInterface::AddColumns(TypedElementRowHandle Row)
 {
 	AddColumns(Row, { Columns::StaticStruct()...});
@@ -462,7 +468,7 @@ inline void ITypedElementDataStorageInterface::RemoveColumn<UE::Editor::DataStor
 	RemoveColumn(Row, FDynamicTag(Tag));
 }
 
-template<TypedElementDataStorage::TEnumType EnumT>
+template<UE::Editor::DataStorage::TEnumType EnumT>
 void ITypedElementDataStorageInterface::AddColumn(TypedElementDataStorage::RowHandle Row, EnumT Value)
 {
 	const UEnum* Enum = StaticEnum<EnumT>();
@@ -473,26 +479,26 @@ void ITypedElementDataStorageInterface::AddColumn(TypedElementDataStorage::RowHa
 	}
 }
 
-template<auto Value, TypedElementDataStorage::TEnumType EnumT>
+template<auto Value, UE::Editor::DataStorage::TEnumType EnumT>
 void ITypedElementDataStorageInterface::AddColumn(TypedElementDataStorage::RowHandle Row)
 {
 	AddColumn<EnumT>(Row, Value);
 }
 
-template<TypedElementDataStorage::TEnumType EnumT>
+template<UE::Editor::DataStorage::TEnumType EnumT>
 void ITypedElementDataStorageInterface::RemoveColumn(TypedElementDataStorage::RowHandle Row)
 {
 	const UEnum* Enum = StaticEnum<EnumT>();
 	RemoveColumn(Row, UE::Editor::DataStorage::FDynamicTag(Enum->GetFName()));
 }
 
-template<TypedElementDataStorage::TColumnType... Columns>
+template<UE::Editor::DataStorage::TColumnType... Columns>
 void ITypedElementDataStorageInterface::RemoveColumns(TypedElementRowHandle Row)
 {
 	RemoveColumns(Row, { Columns::StaticStruct()...});
 }
 
-template<TypedElementDataStorage::TDataColumnType ColumnType>
+template<UE::Editor::DataStorage::TDataColumnType ColumnType>
 void ITypedElementDataStorageInterface::AddColumn(TypedElementRowHandle Row, ColumnType&& Column)
 {
 	AddColumnData(Row, ColumnType::StaticStruct(),
@@ -520,19 +526,19 @@ void ITypedElementDataStorageInterface::AddColumn(TypedElementRowHandle Row, Col
 		});
 }
 
-template<TypedElementDataStorage::TDataColumnType ColumnType>
+template<UE::Editor::DataStorage::TDataColumnType ColumnType>
 ColumnType* ITypedElementDataStorageInterface::GetColumn(TypedElementRowHandle Row)
 {
 	return reinterpret_cast<ColumnType*>(GetColumnData(Row, ColumnType::StaticStruct()));
 }
 
-template<TypedElementDataStorage::TDataColumnType ColumnType>
+template<UE::Editor::DataStorage::TDataColumnType ColumnType>
 const ColumnType* ITypedElementDataStorageInterface::GetColumn(TypedElementRowHandle Row) const
 {
 	return reinterpret_cast<const ColumnType*>(GetColumnData(Row, ColumnType::StaticStruct()));
 }
 
-template<TypedElementDataStorage::TColumnType... ColumnType>
+template<UE::Editor::DataStorage::TColumnType... ColumnType>
 bool ITypedElementDataStorageInterface::HasColumns(TypedElementRowHandle Row) const
 {
 	return HasColumns(Row, TConstArrayView<const UScriptStruct*>({ ColumnType::StaticStruct()... }));
diff --git a/Engine/Source/Runtime/TypedElementFramework/Public/Elements/Interfaces/TypedElementQueryStorageInterfaces.h b/Engine/Source/Runtime/TypedElementFramework/Public/Elements/Interfaces/TypedElementQueryStorageInterfaces.h
--- a/Engine/Source/Runtime/TypedElementFramework/Public/Elements/Interfaces/TypedElementQueryStorageInterfaces.h
+++ b/Engine/Source/Runtime/TypedElementFramework/Public/Elements/Interfaces/TypedElementQueryStorageInterfaces.h
@@ -21,8 +21,6 @@ namespace TypedElementDataStorage
 	using SubqueryCallback = TFunction<void(const FQueryDescription&, ISubqueryContext&)>;
 	using SubqueryCallbackRef = TFunctionRef<void(const FQueryDescription&, ISubqueryContext&)>;
 
-	using IndexHash = uint64;
-
 	/**
 	 * Base interface for any contexts provided to query callbacks.
 	 */
diff --git a/Engine/Source/Runtime/TypedElementFramework/Tests/Elements/Framework/TypedElementDataStoragePerformanceTestCommands.cpp b/Engine/Source/Runtime/TypedElementFramework/Tests/Elements/Framework/TypedElementDataStoragePerformanceTestCommands.cpp
--- a/Engine/Source/Runtime/TypedElementFramework/Tests/Elements/Framework/TypedElementDataStoragePerformanceTestCommands.cpp
+++ b/Engine/Source/Runtime/TypedElementFramework/Tests/Elements/Framework/TypedElementDataStoragePerformanceTestCommands.cpp
@@ -2,6 +2,7 @@
 
 #include "TypedElementDataStoragePerformanceTestCommands.h"
 
+#include "Elements/Interfaces/TypedElementDataStorageInterface.h"
 #include "Elements/Framework/TypedElementQueryBuilder.h"
 #include "Elements/Framework/TypedElementRegistry.h"
 #include "HAL/IConsoleManager.h"
diff --git a/Engine/Source/Runtime/TypedElementFramework/Tests/Elements/Framework/TypedElementDataStoragePerformanceTestCommands.h b/Engine/Source/Runtime/TypedElementFramework/Tests/Elements/Framework/TypedElementDataStoragePerformanceTestCommands.h
--- a/Engine/Source/Runtime/TypedElementFramework/Tests/Elements/Framework/TypedElementDataStoragePerformanceTestCommands.h
+++ b/Engine/Source/Runtime/TypedElementFramework/Tests/Elements/Framework/TypedElementDataStoragePerformanceTestCommands.h
@@ -4,7 +4,7 @@
 
 #include "CoreMinimal.h"
 #include "Elements/Interfaces/TypedElementDataStorageFactory.h"
-#include "Elements/Interfaces/TypedElementDataStorageInterface.h"
+#include "Elements/Common/TypedElementCommonTypes.h"
 #include "Modules/ModuleManager.h"
 
 #include "TypedElementDataStoragePerformanceTestCommands.generated.h"
@@ -13,7 +13,7 @@
  * Column to represent that a row is selected
  */
 USTRUCT(meta = (DisplayName = ""))
-struct FTest_PingPongPrePhys final : public FTypedElementDataStorageColumn
+struct FTest_PingPongPrePhys final : public FEditorDataStorageColumn
 {
 	GENERATED_BODY()
 	
@@ -22,7 +22,7 @@ struct FTest_PingPongPrePhys final : public FTypedElementDataStorageColumn
 };
 
 USTRUCT(meta = (DisplayName = ""))
-struct FTest_PingPongDurPhys final : public FTypedElementDataStorageColumn
+struct FTest_PingPongDurPhys final : public FEditorDataStorageColumn
 {
 	GENERATED_BODY()
 
@@ -31,7 +31,7 @@ struct FTest_PingPongDurPhys final : public FTypedElementDataStorageColumn
 };
 
 USTRUCT(meta = (DisplayName = ""))
-struct FTest_PingPongPostPhys final : public FTypedElementDataStorageColumn
+struct FTest_PingPongPostPhys final : public FEditorDataStorageColumn
 {
 	GENERATED_BODY()
 	
diff --git a/Engine/Source/Runtime/TypedElementFramework/Tests/Elements/Framework/TypedElementTestColumns.h b/Engine/Source/Runtime/TypedElementFramework/Tests/Elements/Framework/TypedElementTestColumns.h
--- a/Engine/Source/Runtime/TypedElementFramework/Tests/Elements/Framework/TypedElementTestColumns.h
+++ b/Engine/Source/Runtime/TypedElementFramework/Tests/Elements/Framework/TypedElementTestColumns.h
@@ -2,55 +2,55 @@
 
 #pragma once
 
-#include "Elements/Interfaces/TypedElementDataStorageInterface.h"
-#include "UObject/ObjectMacros.h"
+#include "Elements/Common/TypedElementCommonTypes.h"
+#include "Elements/Common/TypedElementHandles.h"
 
 #include "TypedElementTestColumns.generated.h"
 
 USTRUCT(meta = (DisplayName = "ColumnA"))
-struct FTestColumnA final : public FTypedElementDataStorageColumn
+struct FTestColumnA final : public FEditorDataStorageColumn
 {
 	GENERATED_BODY()
 };
 
 USTRUCT(meta = (DisplayName = "ColumnB"))
-struct FTestColumnB final : public FTypedElementDataStorageColumn
+struct FTestColumnB final : public FEditorDataStorageColumn
 {
 	GENERATED_BODY()
 };
 
 USTRUCT(meta = (DisplayName = "ColumnC"))
-struct FTestColumnC final : public FTypedElementDataStorageColumn
+struct FTestColumnC final : public FEditorDataStorageColumn
 {
 	GENERATED_BODY()
 };
 
 USTRUCT(meta = (DisplayName = "ColumnD"))
-struct FTestColumnD final : public FTypedElementDataStorageColumn
+struct FTestColumnD final : public FEditorDataStorageColumn
 {
 	GENERATED_BODY()
 };
 
 USTRUCT(meta = (DisplayName = "ColumnE"))
-struct FTestColumnE final : public FTypedElementDataStorageColumn
+struct FTestColumnE final : public FEditorDataStorageColumn
 {
 	GENERATED_BODY()
 };
 
 USTRUCT(meta = (DisplayName = "ColumnF"))
-struct FTestColumnF final : public FTypedElementDataStorageColumn
+struct FTestColumnF final : public FEditorDataStorageColumn
 {
 	GENERATED_BODY()
 };
 
 USTRUCT(meta = (DisplayName = "ColumnG"))
-struct FTestColumnG final : public FTypedElementDataStorageColumn
+struct FTestColumnG final : public FEditorDataStorageColumn
 {
 	GENERATED_BODY()
 };
 
 USTRUCT(meta = (DisplayName = "ColumnInt"))
-struct FTestColumnInt final : public FTypedElementDataStorageColumn
+struct FTestColumnInt final : public FEditorDataStorageColumn
 {
 	GENERATED_BODY()
 	
@@ -59,7 +59,7 @@ struct FTestColumnInt final : public FTypedElementDataStorageColumn
 };
 
 USTRUCT(meta = (DisplayName = "ColumnString"))
-struct FTestColumnString final : public FTypedElementDataStorageColumn
+struct FTestColumnString final : public FEditorDataStorageColumn
 {
 	GENERATED_BODY()
 	
@@ -69,31 +69,31 @@ struct FTestColumnString final : public FTypedElementDataStorageColumn
 
 
 USTRUCT(meta = (DisplayName = "TagA"))
-struct FTestTagColumnA final : public FTypedElementDataStorageTag
+struct FTestTagColumnA final : public FEditorDataStorageTag
 {
 	GENERATED_BODY()
 };
 
 USTRUCT(meta = (DisplayName = "TagB"))
-struct FTestTagColumnB final : public FTypedElementDataStorageTag
+struct FTestTagColumnB final : public FEditorDataStorageTag
 {
 	GENERATED_BODY()
 };
 
 USTRUCT(meta = (DisplayName = "TagC"))
-struct FTestTagColumnC final : public FTypedElementDataStorageTag
+struct FTestTagColumnC final : public FEditorDataStorageTag
 {
 	GENERATED_BODY()
 };
 
 USTRUCT(meta = (DisplayName = "TagD"))
-struct FTestTagColumnD final : public FTypedElementDataStorageTag
+struct FTestTagColumnD final : public FEditorDataStorageTag
 {
 	GENERATED_BODY()
 };
 
 USTRUCT(meta = (DisplayName = "TestReferenceColumn"))
-struct FTEDSProcessorTestsReferenceColumn final : public FTypedElementDataStorageColumn
+struct FTEDSProcessorTestsReferenceColumn final : public FEditorDataStorageColumn
 {
 	GENERATED_BODY()
 
@@ -104,19 +104,19 @@ struct FTEDSProcessorTestsReferenceColumn final : public FTypedElementDataStorag
 };
 
 USTRUCT(meta = (DisplayName = "TEDSProcessorTests_PrimaryTag"))
-struct FTEDSProcessorTests_PrimaryTag final : public FTypedElementDataStorageTag
+struct FTEDSProcessorTests_PrimaryTag final : public FEditorDataStorageTag
 {
 	GENERATED_BODY()
 };
 
 USTRUCT(meta = (DisplayName = "TEDSProcessorTestsSecondaryTag"))
-struct FTEDSProcessorTests_SecondaryTag final : public FTypedElementDataStorageTag
+struct FTEDSProcessorTests_SecondaryTag final : public FEditorDataStorageTag
 {
 	GENERATED_BODY()
 };
 
 USTRUCT(meta = (DisplayName = "TEDSProcessorTests_Linked"))
-struct FTEDSProcessorTests_Linked final : public FTypedElementDataStorageTag
+struct FTEDSProcessorTests_Linked final : public FEditorDataStorageTag
 {
 	GENERATED_BODY()
 };
\ No newline at end of file
