[Teds] Update naming for RowHandle, QueryHandle, and TableHandle

diff --git a/Engine/Plugins/Experimental/EditorDataStorageFeatures/Source/TedsTableViewer/Private/Widgets/STedsTableViewerRow.cpp b/Engine/Plugins/Experimental/EditorDataStorageFeatures/Source/TedsTableViewer/Private/Widgets/STedsTableViewerRow.cpp
--- a/Engine/Plugins/Experimental/EditorDataStorageFeatures/Source/TedsTableViewer/Private/Widgets/STedsTableViewerRow.cpp
+++ b/Engine/Plugins/Experimental/EditorDataStorageFeatures/Source/TedsTableViewer/Private/Widgets/STedsTableViewerRow.cpp
@@ -6,7 +6,7 @@
 #include "TedsTableViewerColumn.h"
 #include "Widgets/STedsTableViewer.h"
 
-namespace UE::EditorDataStorage
+namespace UE::Editor::DataStorage
 {
 	void STedsTableViewerRow::Construct(const FArguments& InArgs, const TSharedRef<STableViewBase>& OwnerTableView, const TSharedRef<FTedsTableViewerModel>& InTableViewerModel)
 	{
@@ -36,4 +36,4 @@ namespace UE::EditorDataStorage
 
 		return SNullWidget::NullWidget;
 	}
-}
\ No newline at end of file
+} // namespace UE::Editor::DataStorage
diff --git a/Engine/Plugins/Experimental/EditorDataStorageFeatures/Source/TedsTableViewer/Public/QueryStack/FQueryStackNode_RowView.h b/Engine/Plugins/Experimental/EditorDataStorageFeatures/Source/TedsTableViewer/Public/QueryStack/FQueryStackNode_RowView.h
--- a/Engine/Plugins/Experimental/EditorDataStorageFeatures/Source/TedsTableViewer/Public/QueryStack/FQueryStackNode_RowView.h
+++ b/Engine/Plugins/Experimental/EditorDataStorageFeatures/Source/TedsTableViewer/Public/QueryStack/FQueryStackNode_RowView.h
@@ -4,7 +4,7 @@
 
 #include "QueryStack/IQueryStackNode_Row.h"
 
-namespace UE::EditorDataStorage
+namespace UE::Editor::DataStorage
 {
 	/*
 	 * A very simple row query stack that views a list of rows.
@@ -15,19 +15,18 @@ namespace UE::EditorDataStorage
 	{
 	public:
 		// IQueryStackNode_Row interface
-		TEDSTABLEVIEWER_API virtual TConstArrayView<TypedElementDataStorage::RowHandle> GetOrderedRowList() override;
+		TEDSTABLEVIEWER_API virtual TConstArrayView<RowHandle> GetOrderedRowList() override;
 		TEDSTABLEVIEWER_API virtual uint32 GetRevisionId() const override;
 		// IQueryStackNode_Row interface
 
-		TEDSTABLEVIEWER_API FQueryStackNode_RowView(TArray<TypedElementDataStorage::RowHandle>* InRows);
+		TEDSTABLEVIEWER_API FQueryStackNode_RowView(TArray<RowHandle>* InRows);
 
 		// Increment the revision ID to update the table viewer when the list of rows changes
 		TEDSTABLEVIEWER_API void MarkDirty();
 
 	private:
 
-		TArray<TypedElementDataStorage::RowHandle>* Rows;
+		TArray<RowHandle>* Rows;
 		uint32 RevisionId = 0;
 	};
-}
-
+} // namespace UE::Editor::DataStorage
diff --git a/Engine/Plugins/Experimental/EditorDataStorageFeatures/Source/TedsTableViewer/Public/QueryStack/IQueryStackNode_Row.h b/Engine/Plugins/Experimental/EditorDataStorageFeatures/Source/TedsTableViewer/Public/QueryStack/IQueryStackNode_Row.h
--- a/Engine/Plugins/Experimental/EditorDataStorageFeatures/Source/TedsTableViewer/Public/QueryStack/IQueryStackNode_Row.h
+++ b/Engine/Plugins/Experimental/EditorDataStorageFeatures/Source/TedsTableViewer/Public/QueryStack/IQueryStackNode_Row.h
@@ -5,7 +5,7 @@
 #include "Containers/ArrayView.h"
 #include "Elements/Common/TypedElementHandles.h"
 
-namespace UE::EditorDataStorage
+namespace UE::Editor::DataStorage
 {
 	/*
 	 * A row node in the query stack is used by the Teds Table Viewer to determine which rows to visualize in the table viewer
@@ -18,10 +18,9 @@ namespace UE::EditorDataStorage
 		virtual ~IQueryStackNode_Row() = default;
 
 		// Get the rows to visualize
-		virtual TConstArrayView<TypedElementDataStorage::RowHandle> GetOrderedRowList() = 0;
+		virtual TConstArrayView<RowHandle> GetOrderedRowList() = 0;
 
 		// Returns an id that updates whenever the locally stored rows have changed.
 		virtual uint32 GetRevisionId() const = 0;
 	};
-}
-
+} // namespace UE::Editor::DataStorage
diff --git a/Engine/Plugins/Experimental/EditorDataStorageFeatures/Source/TedsTableViewer/Public/TedsTableViewerColumn.h b/Engine/Plugins/Experimental/EditorDataStorageFeatures/Source/TedsTableViewer/Public/TedsTableViewerColumn.h
--- a/Engine/Plugins/Experimental/EditorDataStorageFeatures/Source/TedsTableViewer/Public/TedsTableViewerColumn.h
+++ b/Engine/Plugins/Experimental/EditorDataStorageFeatures/Source/TedsTableViewer/Public/TedsTableViewerColumn.h
@@ -13,7 +13,7 @@ class ITypedElementDataStorageInterface;
 struct FTypedElementWidgetConstructor;
 class SWidget;
 
-namespace UE::EditorDataStorage
+namespace UE::Editor::DataStorage
 {
 	/*
 	 * Class representing a column in the UI of the table viewer. Can be constructed using a NameID and a WidgetConstructor to create the actual
@@ -24,7 +24,7 @@ namespace UE::EditorDataStorage
 	public:
 
 		// Delegate to check if a row is currently visible in the owning table viewer's UI
-		DECLARE_DELEGATE_RetVal_OneParam(bool, FIsRowVisible, const TypedElementDataStorage::RowHandle);
+		DECLARE_DELEGATE_RetVal_OneParam(bool, FIsRowVisible, const RowHandle);
 
 		TEDSTABLEVIEWER_API FTedsTableViewerColumn(
 			const FName& ColumnName, // The unique ID of this column
@@ -35,7 +35,7 @@ namespace UE::EditorDataStorage
 
 		TEDSTABLEVIEWER_API ~FTedsTableViewerColumn();
 		
-		TEDSTABLEVIEWER_API	TSharedPtr<SWidget> ConstructRowWidget(TypedElementDataStorage::RowHandle RowHandle) const;
+		TEDSTABLEVIEWER_API	TSharedPtr<SWidget> ConstructRowWidget(RowHandle InRowHandle) const;
 		
 		TEDSTABLEVIEWER_API SHeaderRow::FColumn::FArguments ConstructHeaderRowColumn() const;
 		
@@ -52,7 +52,7 @@ namespace UE::EditorDataStorage
 
 		void RegisterQueries();
 		void UnRegisterQueries();
-		bool IsRowVisible(const TypedElementDataStorage::RowHandle InRowHandle) const;
+		bool IsRowVisible(const RowHandle InRowHandle) const;
 		void UpdateWidgets();
 		
 	private:
@@ -75,11 +75,11 @@ namespace UE::EditorDataStorage
 		ITypedElementDataStorageCompatibilityInterface* StorageCompatibility;
 
 		// Queries used to virtualize widgets when a column is added to/remove from a row
-		TArray<TypedElementDataStorage::QueryHandle> InternalObserverQueries;
-		TypedElementDataStorage::QueryHandle WidgetQuery;
-		TMap<TypedElementDataStorage::RowHandle, bool> RowsToUpdate;
+		TArray<QueryHandle> InternalObserverQueries;
+		QueryHandle WidgetQuery;
+		TMap<RowHandle, bool> RowsToUpdate;
 
 		// Delegate to check if a row is visible in the owning table viewer
 		FIsRowVisible IsRowVisibleDelegate;
 	};
-}
+} // namespace UE::Editor::DataStorage
diff --git a/Engine/Plugins/Experimental/EditorDataStorageFeatures/Source/TedsTableViewer/Public/TedsTableViewerModel.h b/Engine/Plugins/Experimental/EditorDataStorageFeatures/Source/TedsTableViewer/Public/TedsTableViewerModel.h
--- a/Engine/Plugins/Experimental/EditorDataStorageFeatures/Source/TedsTableViewer/Public/TedsTableViewerModel.h
+++ b/Engine/Plugins/Experimental/EditorDataStorageFeatures/Source/TedsTableViewer/Public/TedsTableViewerModel.h
@@ -13,7 +13,7 @@ class ITypedElementDataStorageCompatibilityInterface;
 class ITypedElementDataStorageUiInterface;
 class ITypedElementDataStorageInterface;
 
-namespace UE::EditorDataStorage
+namespace UE::Editor::DataStorage
 {
 	class FTedsTableViewerColumn;
 	class IQueryStackNode_Row;
@@ -21,7 +21,7 @@ namespace UE::EditorDataStorage
 	// Typedef for an item in the table viewer
 	using TableViewerItemPtr = UIRowType;
 
-	// Model class for the Teds Table Viewer that can be plugged into any widget that is a UI representation of data in Teds
+	// Model class for the TEDS Table Viewer that can be plugged into any widget that is a UI representation of data in TEDS
 	// @see STedsTableViewer
 	class FTedsTableViewerModel
 	{
@@ -65,14 +65,14 @@ namespace UE::EditorDataStorage
 		void GenerateColumns();
 
 		// Check if the given row is currently visible in the UI
-		bool IsRowVisible(TypedElementDataStorage::RowHandle InRowHandle) const;
+		bool IsRowVisible(RowHandle InRowHandle) const;
 
 		bool Tick(float DeltaTime);
 
 		void Refresh();
 
 		// Check whether a row is allowed to be displayed in the table viewer
-		bool IsRowDisplayable(TypedElementDataStorage::RowHandle InRowHandle) const;
+		bool IsRowDisplayable(RowHandle InRowHandle) const;
 
 	private:
 
@@ -107,4 +107,4 @@ namespace UE::EditorDataStorage
 		ITypedElementDataStorageUiInterface* StorageUi = nullptr;
 		ITypedElementDataStorageCompatibilityInterface* StorageCompatibility = nullptr;
 	};
-}
\ No newline at end of file
+} // namespace UE::Editor::DataStorage
diff --git a/Engine/Plugins/Experimental/EditorDataStorageFeatures/Source/TedsTableViewer/Public/TedsTableViewerUtils.h b/Engine/Plugins/Experimental/EditorDataStorageFeatures/Source/TedsTableViewer/Public/TedsTableViewerUtils.h
--- a/Engine/Plugins/Experimental/EditorDataStorageFeatures/Source/TedsTableViewer/Public/TedsTableViewerUtils.h
+++ b/Engine/Plugins/Experimental/EditorDataStorageFeatures/Source/TedsTableViewer/Public/TedsTableViewerUtils.h
@@ -17,7 +17,7 @@ class ITypedElementDataStorageUiInterface;
 struct FTypedElementWidgetConstructor;
 
 // Util library for functions shared by the Teds Table Viewer and the Teds Outliner
-namespace UE::EditorDataStorage::TableViewerUtils
+namespace UE::Editor::DataStorage::TableViewerUtils
 {
 	TEDSTABLEVIEWER_API FName GetWidgetTableName();
 	
@@ -31,7 +31,7 @@ namespace UE::EditorDataStorage::TableViewerUtils
 
 	// Create a copy of the provided column types array after discarding invalid entries
 	TEDSTABLEVIEWER_API TArray<TWeakObjectPtr<const UScriptStruct>> CreateVerifiedColumnTypeArray(TConstArrayView<TWeakObjectPtr<const UScriptStruct>> ColumnTypes);
-}
+} // namespace UE::Editor::DataStorage::TableViewerUtils
 
 UCLASS()
 class UTypedElementTableViewerFactory : public UTypedElementDataStorageFactory
diff --git a/Engine/Plugins/Experimental/EditorDataStorageFeatures/Source/TedsTableViewer/Public/TypedElementUITypes.h b/Engine/Plugins/Experimental/EditorDataStorageFeatures/Source/TedsTableViewer/Public/TypedElementUITypes.h
--- a/Engine/Plugins/Experimental/EditorDataStorageFeatures/Source/TedsTableViewer/Public/TypedElementUITypes.h
+++ b/Engine/Plugins/Experimental/EditorDataStorageFeatures/Source/TedsTableViewer/Public/TypedElementUITypes.h
@@ -5,35 +5,35 @@
 #include "Elements/Common/TypedElementHandles.h"
 #include "Framework/Views/TableViewTypeTraits.h"
 
-namespace UE::EditorDataStorage
+namespace UE::Editor::DataStorage
 {
-	// Wrapper struct around TypedElementDataStorage::RowHandle so we can template speicalize for it without also specializing for uint64
+	// Wrapper struct around RowHandle so we can template specialize for it without also specializing for uint64
 	struct UIRowType
 	{
-		TypedElementDataStorage::RowHandle RowHandle;
+		RowHandle Row;
 		
 		UIRowType()
-			: RowHandle(TypedElementDataStorage::InvalidRowHandle)
+			: Row(InvalidRowHandle)
 		{}
 
-		UIRowType(TypedElementDataStorage::RowHandle InRowHandle)
-			: RowHandle(InRowHandle)
+		UIRowType(RowHandle InRowHandle)
+			: Row(InRowHandle)
 		{}
 
-		operator TypedElementDataStorage::RowHandle () const
+		operator RowHandle() const
 		{
-			return RowHandle;
+			return Row;
 		}
 
-		UIRowType& operator=(TypedElementDataStorage::RowHandle InRowHandle)
+		UIRowType& operator=(RowHandle InRowHandle)
 		{
-			RowHandle = InRowHandle;
+			Row = InRowHandle;
 			return *this;
 		}
 		
 		friend uint32 GetTypeHash(const UIRowType& Key)
 		{
-			return GetTypeHash(Key.RowHandle);
+			return GetTypeHash(Key.Row);
 		}
 
 	};
@@ -44,45 +44,45 @@ namespace UE::EditorDataStorage
  * SListView<UE::EditorDataStorage::UIRowType>
  */
 template <>
-struct TListTypeTraits<UE::EditorDataStorage::UIRowType>
+struct TListTypeTraits<UE::Editor::DataStorage::UIRowType>
 {
-	typedef UE::EditorDataStorage::UIRowType NullableType;
+	using NullableType = UE::Editor::DataStorage::UIRowType;
 
-	using MapKeyFuncs = TDefaultMapHashableKeyFuncs<UE::EditorDataStorage::UIRowType, TSharedRef<ITableRow>, false>;
-	using MapKeyFuncsSparse = TDefaultMapHashableKeyFuncs<UE::EditorDataStorage::UIRowType, FSparseItemInfo, false>;
-	using SetKeyFuncs = DefaultKeyFuncs<UE::EditorDataStorage::UIRowType>;
+	using MapKeyFuncs = TDefaultMapHashableKeyFuncs<UE::Editor::DataStorage::UIRowType, TSharedRef<ITableRow>, false>;
+	using MapKeyFuncsSparse = TDefaultMapHashableKeyFuncs<UE::Editor::DataStorage::UIRowType, FSparseItemInfo, false>;
+	using SetKeyFuncs = DefaultKeyFuncs<UE::Editor::DataStorage::UIRowType>;
 
 	template<typename U>
 	static void AddReferencedObjects(FReferenceCollector&,
-		TArray<UE::EditorDataStorage::UIRowType>&,
-		TSet<UE::EditorDataStorage::UIRowType>&,
-		TMap<const U*, UE::EditorDataStorage::UIRowType>&)
+		TArray<UE::Editor::DataStorage::UIRowType>&,
+		TSet<UE::Editor::DataStorage::UIRowType>&,
+		TMap<const U*, UE::Editor::DataStorage::UIRowType>&)
 	{
 	}
 
-	static bool IsPtrValid(const UE::EditorDataStorage::UIRowType& InPtr)
+	static bool IsPtrValid(const UE::Editor::DataStorage::UIRowType& InPtr)
 	{
-		return InPtr != TypedElementDataStorage::InvalidRowHandle;
+		return InPtr != UE::Editor::DataStorage::InvalidRowHandle;
 	}
 
-	static void ResetPtr(UE::EditorDataStorage::UIRowType& InPtr)
+	static void ResetPtr(UE::Editor::DataStorage::UIRowType& InPtr)
 	{
-		InPtr = TypedElementDataStorage::InvalidRowHandle;
+		InPtr = UE::Editor::DataStorage::InvalidRowHandle;
 	}
 
-	static UE::EditorDataStorage::UIRowType MakeNullPtr()
+	static UE::Editor::DataStorage::UIRowType MakeNullPtr()
 	{
-		return TypedElementDataStorage::InvalidRowHandle;
+		return UE::Editor::DataStorage::InvalidRowHandle;
 	}
 
-	static UE::EditorDataStorage::UIRowType NullableItemTypeConvertToItemType(const UE::EditorDataStorage::UIRowType& InPtr)
+	static UE::Editor::DataStorage::UIRowType NullableItemTypeConvertToItemType(const UE::Editor::DataStorage::UIRowType& InPtr)
 	{
 		return InPtr;
 	}
 
-	static FString DebugDump(UE::EditorDataStorage::UIRowType InPtr)
+	static FString DebugDump(UE::Editor::DataStorage::UIRowType InPtr)
 	{
-		return FString::Printf(TEXT("%llu"), InPtr.RowHandle);
+		return FString::Printf(TEXT("%llu"), InPtr.Row);
 	}
 
 	class SerializerType {};
@@ -90,10 +90,10 @@ struct TListTypeTraits<UE::EditorDataStorage::UIRowType>
 
 // Template declaration to enable using row handles inside of slate widgets like SListView
 template <>
-struct TIsValidListItem<UE::EditorDataStorage::UIRowType>
+struct TIsValidListItem<UE::Editor::DataStorage::UIRowType>
 {
 	enum
 	{
 		Value = true
 	};
-};
\ No newline at end of file
+};
diff --git a/Engine/Plugins/Experimental/EditorDataStorageFeatures/Source/TedsTableViewer/Public/Widgets/RowHandleWidget.h b/Engine/Plugins/Experimental/EditorDataStorageFeatures/Source/TedsTableViewer/Public/Widgets/RowHandleWidget.h
--- a/Engine/Plugins/Experimental/EditorDataStorageFeatures/Source/TedsTableViewer/Public/Widgets/RowHandleWidget.h
+++ b/Engine/Plugins/Experimental/EditorDataStorageFeatures/Source/TedsTableViewer/Public/Widgets/RowHandleWidget.h
@@ -40,6 +40,6 @@ protected:
 	virtual bool FinalizeWidget(
 		ITypedElementDataStorageInterface* DataStorage,
 		ITypedElementDataStorageUiInterface* DataStorageUi,
-		TypedElementRowHandle Row,
+		UE::Editor::DataStorage::RowHandle Row,
 		const TSharedPtr<SWidget>& Widget) override;
 };
diff --git a/Engine/Plugins/Experimental/EditorDataStorageFeatures/Source/TedsTableViewer/Public/Widgets/SRowDetails.h b/Engine/Plugins/Experimental/EditorDataStorageFeatures/Source/TedsTableViewer/Public/Widgets/SRowDetails.h
--- a/Engine/Plugins/Experimental/EditorDataStorageFeatures/Source/TedsTableViewer/Public/Widgets/SRowDetails.h
+++ b/Engine/Plugins/Experimental/EditorDataStorageFeatures/Source/TedsTableViewer/Public/Widgets/SRowDetails.h
@@ -10,7 +10,7 @@ class ITypedElementDataStorageInterface;
 class ITypedElementDataStorageUiInterface;
 struct FTypedElementWidgetConstructor;
 
-namespace UE::EditorDataStorage
+namespace UE::Editor::DataStorage
 {
 	// A row in the SRowDetails widget that represents a column on the TEDS row we are viewing
 	struct FRowDetailsItem
@@ -21,11 +21,11 @@ namespace UE::EditorDataStorage
 		// Widget for the column
 		TUniquePtr<FTypedElementWidgetConstructor> WidgetConstructor;
 		
-		TypedElementDataStorage::RowHandle Row = TypedElementDataStorage::InvalidRowHandle;
-		TypedElementDataStorage::RowHandle WidgetRow = TypedElementDataStorage::InvalidRowHandle;
+		RowHandle Row = InvalidRowHandle;
+		RowHandle WidgetRow = InvalidRowHandle;
 
 		FRowDetailsItem(const TWeakObjectPtr<const UScriptStruct>& InColumnType, TUniquePtr<FTypedElementWidgetConstructor> InWidgetConstructor,
-			TypedElementDataStorage::RowHandle InRow);
+			RowHandle InRow);
 	};
 	
 	using RowDetailsItemPtr = TSharedPtr<FRowDetailsItem>;
@@ -53,7 +53,7 @@ namespace UE::EditorDataStorage
 		void Construct(const FArguments& InArgs);
 
 		// Set the row to view
-		void SetRow(TypedElementDataStorage::RowHandle Row);
+		void SetRow(RowHandle Row);
 
 		// Clear the row to view
 		void ClearRow();
@@ -95,4 +95,4 @@ namespace UE::EditorDataStorage
 		ITypedElementDataStorageInterface* DataStorage = nullptr;
 		ITypedElementDataStorageUiInterface* DataStorageUi = nullptr;
 	};
-} // namespace UE::EditorDataStorage
\ No newline at end of file
+} // namespace UE::Editor::DataStorage
\ No newline at end of file
diff --git a/Engine/Plugins/Experimental/EditorDataStorageFeatures/Source/TedsTableViewer/Public/Widgets/STedsTableViewer.h b/Engine/Plugins/Experimental/EditorDataStorageFeatures/Source/TedsTableViewer/Public/Widgets/STedsTableViewer.h
--- a/Engine/Plugins/Experimental/EditorDataStorageFeatures/Source/TedsTableViewer/Public/Widgets/STedsTableViewer.h
+++ b/Engine/Plugins/Experimental/EditorDataStorageFeatures/Source/TedsTableViewer/Public/Widgets/STedsTableViewer.h
@@ -12,7 +12,7 @@ class ITableRow;
 class STableViewBase;
 class SHeaderRow;
 
-namespace UE::EditorDataStorage
+namespace UE::Editor::DataStorage
 {
 	class FTedsTableViewerColumn;
 
@@ -30,7 +30,7 @@ namespace UE::EditorDataStorage
 	public:
 		
 		// Delegate fired when the selection in the table viewer changes
-		DECLARE_DELEGATE_OneParam(FOnSelectionChanged, TypedElementDataStorage::RowHandle)
+		DECLARE_DELEGATE_OneParam(FOnSelectionChanged, RowHandle)
 
 		SLATE_BEGIN_ARGS(STedsTableViewer)
 			: _CellWidgetPurposes({TEXT("General.Cell")})
@@ -91,4 +91,4 @@ namespace UE::EditorDataStorage
 		// Delegate fired when the selection changes
 		FOnSelectionChanged OnSelectionChanged;
 	};
-}
\ No newline at end of file
+} // namespace UE::Editor::DataStorage
\ No newline at end of file
diff --git a/Engine/Plugins/Experimental/InstanceDataObjectFixupTool/Source/Private/InstanceDataObjectFixupToolTedsQueries.cpp b/Engine/Plugins/Experimental/InstanceDataObjectFixupTool/Source/Private/InstanceDataObjectFixupToolTedsQueries.cpp
--- a/Engine/Plugins/Experimental/InstanceDataObjectFixupTool/Source/Private/InstanceDataObjectFixupToolTedsQueries.cpp
+++ b/Engine/Plugins/Experimental/InstanceDataObjectFixupTool/Source/Private/InstanceDataObjectFixupToolTedsQueries.cpp
@@ -20,12 +20,13 @@ void UInstanceDataObjectFixupToolTedsQueryFactory::RegisterQueries(ITypedElement
 	using namespace UE;
 	using namespace TypedElementDataStorage;
 	using namespace TypedElementQueryBuilder;
+	using namespace UE::Editor::DataStorage;
 	
 	DataStorage.RegisterQuery(
 		Select(
 			TEXT("Add fix-up tool to serialization placeholder alerts"),
 			FProcessor(EQueryTickPhase::PrePhysics, DataStorage.GetQueryTickGroupName(EQueryTickGroups::SyncExternalToDataStorage)),
-			[](IQueryContext& Context, TypedElementRowHandle Row, const FTypedElementUObjectColumn& Object)
+			[](IQueryContext& Context, RowHandle Row, const FTypedElementUObjectColumn& Object)
 			{
 				Context.AddColumn(Row, FTypedElementAlertActionColumn{ .Action = ShowFixUpToolForPlaceholders });
 			}
@@ -39,7 +40,7 @@ void UInstanceDataObjectFixupToolTedsQueryFactory::RegisterQueries(ITypedElement
 		Select(
 			TEXT("Add fix-up tool to serialization loose property alerts"),
 			FProcessor(EQueryTickPhase::PrePhysics, DataStorage.GetQueryTickGroupName(EQueryTickGroups::SyncExternalToDataStorage)),
-			[](IQueryContext& Context, TypedElementRowHandle Row, const FTypedElementUObjectColumn& Object)
+			[](IQueryContext& Context, RowHandle Row, const FTypedElementUObjectColumn& Object)
 			{
 				Context.AddColumn(Row, FTypedElementAlertActionColumn{ .Action = ShowFixUpToolForLooseProperties });
 			}
@@ -50,7 +51,7 @@ void UInstanceDataObjectFixupToolTedsQueryFactory::RegisterQueries(ITypedElement
 		.Compile());
 }
 
-void UInstanceDataObjectFixupToolTedsQueryFactory::ShowFixUpToolForPlaceholders(TypedElementDataStorage::RowHandle Row)
+void UInstanceDataObjectFixupToolTedsQueryFactory::ShowFixUpToolForPlaceholders(UE::Editor::DataStorage::RowHandle Row)
 {
 	FNotificationInfo Info(LOCTEXT("PlaceholderResolutionSuggestion", "Please fix your Verse code and/or rename the Verse class back to the original name."));
 	
@@ -58,12 +59,12 @@ void UInstanceDataObjectFixupToolTedsQueryFactory::ShowFixUpToolForPlaceholders(
 	FSlateNotificationManager::Get().AddNotification(Info);
 }
 
-void UInstanceDataObjectFixupToolTedsQueryFactory::ShowFixUpToolForLooseProperties(TypedElementDataStorage::RowHandle Row)
+void UInstanceDataObjectFixupToolTedsQueryFactory::ShowFixUpToolForLooseProperties(UE::Editor::DataStorage::RowHandle Row)
 {
 	ShowFixUpTool(Row, true);
 }
 
-void UInstanceDataObjectFixupToolTedsQueryFactory::ShowFixUpTool(TypedElementDataStorage::RowHandle Row, bool bRecurseIntoObject)
+void UInstanceDataObjectFixupToolTedsQueryFactory::ShowFixUpTool(UE::Editor::DataStorage::RowHandle Row, bool bRecurseIntoObject)
 {
 	ITypedElementDataStorageInterface* DataStorage = UTypedElementRegistry::GetInstance()->GetMutableDataStorage();
 	if (FTypedElementUObjectColumn* ObjectColumn = DataStorage->GetColumn<FTypedElementUObjectColumn>(Row))
diff --git a/Engine/Plugins/Experimental/InstanceDataObjectFixupTool/Source/Private/InstanceDataObjectFixupToolTedsQueries.h b/Engine/Plugins/Experimental/InstanceDataObjectFixupTool/Source/Private/InstanceDataObjectFixupToolTedsQueries.h
--- a/Engine/Plugins/Experimental/InstanceDataObjectFixupTool/Source/Private/InstanceDataObjectFixupToolTedsQueries.h
+++ b/Engine/Plugins/Experimental/InstanceDataObjectFixupTool/Source/Private/InstanceDataObjectFixupToolTedsQueries.h
@@ -19,7 +19,7 @@ public:
 	void RegisterQueries(ITypedElementDataStorageInterface& DataStorage) override;
 
 private:
-	static void ShowFixUpToolForPlaceholders(TypedElementDataStorage::RowHandle Row);
-	static void ShowFixUpToolForLooseProperties(TypedElementDataStorage::RowHandle Row);
-	static void ShowFixUpTool(TypedElementDataStorage::RowHandle Row, bool bRecurseIntoObject);
+	static void ShowFixUpToolForPlaceholders(UE::Editor::DataStorage::RowHandle Row);
+	static void ShowFixUpToolForLooseProperties(UE::Editor::DataStorage::RowHandle Row);
+	static void ShowFixUpTool(UE::Editor::DataStorage::RowHandle Row, bool bRecurseIntoObject);
 };
diff --git a/Engine/Source/Editor/UnrealEd/Private/Elements/Actor/ActorElementEditorSelectionInterface.cpp b/Engine/Source/Editor/UnrealEd/Private/Elements/Actor/ActorElementEditorSelectionInterface.cpp
--- a/Engine/Source/Editor/UnrealEd/Private/Elements/Actor/ActorElementEditorSelectionInterface.cpp
+++ b/Engine/Source/Editor/UnrealEd/Private/Elements/Actor/ActorElementEditorSelectionInterface.cpp
@@ -48,13 +48,15 @@ private:
  
 bool UActorElementEditorSelectionInterface::SelectElement(const FTypedElementHandle& InElementHandle, const FTypedElementListPtr& InSelectionSet, const FTypedElementSelectionOptions& InSelectionOptions)
 {
+	using namespace UE::Editor::DataStorage;
+
 	const AActor* Actor = ActorElementDataUtil::GetActorFromHandle(InElementHandle);
 	UTypedElementRegistry* Registry = InSelectionSet->GetRegistry();
 	// Add a selection column in TEDS
 	if (const ITypedElementDataStorageCompatibilityInterface* Compatibility = Registry->GetDataStorageCompatibility())
 	{
-		TypedElementRowHandle Row = Compatibility->FindRowWithCompatibleObject(Actor);
-		if (Row != TypedElementInvalidRowHandle)
+		RowHandle Row = Compatibility->FindRowWithCompatibleObject(Actor);
+		if (Row != InvalidRowHandle)
 		{
 			if (ITypedElementDataStorageInterface* DataStorage = UTypedElementRegistry::GetInstance()->GetMutableDataStorage())
 			{
@@ -75,12 +77,14 @@ bool UActorElementEditorSelectionInterface::SelectElement(const FTypedElementHan
  
 bool UActorElementEditorSelectionInterface::DeselectElement(const FTypedElementHandle& InElementHandle, const FTypedElementListPtr& InSelectionSet, const FTypedElementSelectionOptions& InSelectionOptions)
 {
+	using namespace UE::Editor::DataStorage;
+
 	const AActor* Actor = ActorElementDataUtil::GetActorFromHandle(InElementHandle);
 	UTypedElementRegistry* Registry = InSelectionSet->GetRegistry();
 	if (const ITypedElementDataStorageCompatibilityInterface* Compatibility = Registry->GetDataStorageCompatibility())
 	{
-		TypedElementRowHandle Row = Compatibility->FindRowWithCompatibleObject(Actor);
-		if (Row != TypedElementInvalidRowHandle)
+		RowHandle Row = Compatibility->FindRowWithCompatibleObject(Actor);
+		if (Row != InvalidRowHandle)
 		{
 			if (ITypedElementDataStorageInterface* DataStorage = UTypedElementRegistry::GetInstance()->GetMutableDataStorage())
             {
diff --git a/Engine/Source/Runtime/TypedElementFramework/Private/Elements/Framework/TypedElementQueryBuilder.cpp b/Engine/Source/Runtime/TypedElementFramework/Private/Elements/Framework/TypedElementQueryBuilder.cpp
--- a/Engine/Source/Runtime/TypedElementFramework/Private/Elements/Framework/TypedElementQueryBuilder.cpp
+++ b/Engine/Source/Runtime/TypedElementFramework/Private/Elements/Framework/TypedElementQueryBuilder.cpp
@@ -89,7 +89,7 @@ namespace TypedElementQueryBuilder
 		return *this;
 	}
 
-	FDependency& FDependency::SubQuery(TypedElementQueryHandle Handle)
+	FDependency& FDependency::SubQuery(QueryHandle Handle)
 	{
 		checkf(Query->Callback.ExecutionMode != TypedElementDataStorage::EExecutionMode::ThreadedChunks,
 			TEXT("TEDS sub-queries can not be added to queries with a callback that process chunks in parallel."));
@@ -97,7 +97,7 @@ namespace TypedElementQueryBuilder
 		return *this;
 	}
 	
-	FDependency& FDependency::SubQuery(TConstArrayView<TypedElementQueryHandle> Handles)
+	FDependency& FDependency::SubQuery(TConstArrayView<QueryHandle> Handles)
 	{
 		checkf(Query->Callback.ExecutionMode != TypedElementDataStorage::EExecutionMode::ThreadedChunks,
 			TEXT("TEDS sub-queries can not be added to queries with a callback that process chunks in parallel."));
diff --git a/Engine/Source/Runtime/TypedElementFramework/Private/Elements/Interfaces/TypedElementDataStorageUiInterface.cpp b/Engine/Source/Runtime/TypedElementFramework/Private/Elements/Interfaces/TypedElementDataStorageUiInterface.cpp
--- a/Engine/Source/Runtime/TypedElementFramework/Private/Elements/Interfaces/TypedElementDataStorageUiInterface.cpp
+++ b/Engine/Source/Runtime/TypedElementFramework/Private/Elements/Interfaces/TypedElementDataStorageUiInterface.cpp
@@ -94,7 +94,7 @@ FString FTypedElementWidgetConstructor::CreateWidgetDisplayName(
 }
 
 TSharedPtr<SWidget> FTypedElementWidgetConstructor::ConstructFinalWidget(
-	TypedElementRowHandle Row,
+	RowHandle Row,
 	ITypedElementDataStorageInterface* DataStorage,
 	ITypedElementDataStorageUiInterface* DataStorageUi,
 	const TypedElementDataStorage::FMetaDataView& Arguments)
@@ -136,7 +136,7 @@ TSharedPtr<SWidget> FTypedElementWidgetConstructor::ConstructFinalWidget(
 }
 
 TSharedPtr<SWidget> FTypedElementWidgetConstructor::Construct(
-	TypedElementRowHandle Row,
+	RowHandle Row,
 	ITypedElementDataStorageInterface* DataStorage,
 	ITypedElementDataStorageUiInterface* DataStorageUi,
 	const TypedElementDataStorage::FMetaDataView& Arguments)
@@ -162,7 +162,7 @@ TSharedPtr<SWidget> FTypedElementWidgetConstructor::CreateWidget(const TypedElem
 	return nullptr;
 }
 
-bool FTypedElementWidgetConstructor::SetColumns(ITypedElementDataStorageInterface* DataStorage, TypedElementRowHandle Row)
+bool FTypedElementWidgetConstructor::SetColumns(ITypedElementDataStorageInterface* DataStorage, RowHandle Row)
 {
 	return true;
 }
@@ -187,13 +187,13 @@ FString FTypedElementWidgetConstructor::DescribeColumnType(const UScriptStruct*
 bool FTypedElementWidgetConstructor::FinalizeWidget(
 	ITypedElementDataStorageInterface* DataStorage,
 	ITypedElementDataStorageUiInterface* DataStorageUi,
-	TypedElementRowHandle Row,
+	RowHandle Row,
 	const TSharedPtr<SWidget>& Widget)
 {
 	return true;
 }
 
-void FTypedElementWidgetConstructor::AddDefaultWidgetColumns(TypedElementRowHandle Row, ITypedElementDataStorageInterface* DataStorage) const
+void FTypedElementWidgetConstructor::AddDefaultWidgetColumns(RowHandle Row, ITypedElementDataStorageInterface* DataStorage) const
 {
 	const FString WidgetLabel(CreateWidgetDisplayName(DataStorage, Row));
 	DataStorage->AddColumn(Row, FTypedElementLabelColumn{.Label = WidgetLabel} );
diff --git a/Engine/Source/Runtime/TypedElementFramework/Public/Elements/Columns/TypedElementMiscColumns.h b/Engine/Source/Runtime/TypedElementFramework/Public/Elements/Columns/TypedElementMiscColumns.h
--- a/Engine/Source/Runtime/TypedElementFramework/Public/Elements/Columns/TypedElementMiscColumns.h
+++ b/Engine/Source/Runtime/TypedElementFramework/Public/Elements/Columns/TypedElementMiscColumns.h
@@ -53,5 +53,5 @@ struct FTypedElementRowReferenceColumn final : public FEditorDataStorageColumn
 {
 	GENERATED_BODY()
 
-	TypedElementDataStorage::RowHandle Row;
+	UE::Editor::DataStorage::RowHandle Row;
 };
\ No newline at end of file
diff --git a/Engine/Source/Runtime/TypedElementFramework/Public/Elements/Columns/TypedElementPackageColumns.h b/Engine/Source/Runtime/TypedElementFramework/Public/Elements/Columns/TypedElementPackageColumns.h
--- a/Engine/Source/Runtime/TypedElementFramework/Public/Elements/Columns/TypedElementPackageColumns.h
+++ b/Engine/Source/Runtime/TypedElementFramework/Public/Elements/Columns/TypedElementPackageColumns.h
@@ -33,7 +33,7 @@ struct FTypedElementPackageReference final : public FEditorDataStorageColumn
 {
 	GENERATED_BODY()
 	
-	TypedElementDataStorage::RowHandle Row;
+	UE::Editor::DataStorage::RowHandle Row;
 };
 
 /**
@@ -53,8 +53,8 @@ struct FTypedElementPackageUpdateColumn final : public FEditorDataStorageColumn
 {
 	GENERATED_BODY()
 	
-	TypedElementDataStorage::RowHandle ObjectRow;
-	TypedElementDataStorage::RowHandle PackageRow;
+	UE::Editor::DataStorage::RowHandle ObjectRow;
+	UE::Editor::DataStorage::RowHandle PackageRow;
 };
 
 /**
diff --git a/Engine/Source/Runtime/TypedElementFramework/Public/Elements/Common/TypedElementCommonTypes.h b/Engine/Source/Runtime/TypedElementFramework/Public/Elements/Common/TypedElementCommonTypes.h
--- a/Engine/Source/Runtime/TypedElementFramework/Public/Elements/Common/TypedElementCommonTypes.h
+++ b/Engine/Source/Runtime/TypedElementFramework/Public/Elements/Common/TypedElementCommonTypes.h
@@ -78,7 +78,7 @@ namespace UE
 		};
 		// Standard callbacks.
 
-		using RowCreationCallbackRef = TFunctionRef<void(TypedElementDataStorage::RowHandle Row)>;
+		using RowCreationCallbackRef = TFunctionRef<void(RowHandle Row)>;
 		using ColumnCreationCallbackRef = TFunctionRef<void(void* Column, const UScriptStruct& ColumnType)>;
 		using ColumnListCallbackRef = TFunctionRef<void(const UScriptStruct& ColumnType)>;
 		using ColumnListWithDataCallbackRef = TFunctionRef<void(void* Column, const UScriptStruct& ColumnType)>;
diff --git a/Engine/Source/Runtime/TypedElementFramework/Public/Elements/Common/TypedElementHandles.h b/Engine/Source/Runtime/TypedElementFramework/Public/Elements/Common/TypedElementHandles.h
--- a/Engine/Source/Runtime/TypedElementFramework/Public/Elements/Common/TypedElementHandles.h
+++ b/Engine/Source/Runtime/TypedElementFramework/Public/Elements/Common/TypedElementHandles.h
@@ -4,7 +4,7 @@
 
 #include "Math/NumericLimits.h"
 
-namespace TypedElementDataStorage
+namespace UE::Editor::DataStorage
 {
 	using TableHandle = uint64;
 	static constexpr TableHandle InvalidTableHandle = TNumericLimits<TableHandle>::Max();
diff --git a/Engine/Source/Runtime/TypedElementFramework/Public/Elements/Framework/TypedElementQueryBuilder.h b/Engine/Source/Runtime/TypedElementFramework/Public/Elements/Framework/TypedElementQueryBuilder.h
--- a/Engine/Source/Runtime/TypedElementFramework/Public/Elements/Framework/TypedElementQueryBuilder.h
+++ b/Engine/Source/Runtime/TypedElementFramework/Public/Elements/Framework/TypedElementQueryBuilder.h
@@ -43,15 +43,15 @@ class USubsystem;
  * callback once if the query matches. The following function signatures are accepted by "Select":
  *	- void([const]Column&...) 
  *	- void([const]Column*...) 
- *	- void(TypedElementRowHandle, [const]Column&...) 
+ *	- void(RowHandle, [const]Column&...) 
  *	- void(<Context>&, [const]Column&...) 
- *	- void(<Context>&, TypedElementRowHandle, [const]Column&...) 
+ *	- void(<Context>&, RowHandle, [const]Column&...) 
  *	- void(<Context>&, [const]Column*...) 
- *	- void(<Context>&, const TypedElementRowHandle*, [const]Column*...) 
+ *	- void(<Context>&, const RowHandle*, [const]Column*...) 
  *	Where <Context> is ITypedElementDataStorageInterface::IQueryContext or FCachedQueryContext<...>	e.g.:
  *		void(
  *			FCachedQueryContext<Subsystem1, const Subsystem2>& Context, 
- *			TypedElementRowHandle Row, 
+ *			RowHandle Row, 
  *			ColumnType0& ColumnA, 
  *			const ColumnType1& ColumnB) 
  *			{...}
@@ -106,6 +106,8 @@ class USubsystem;
 
 namespace TypedElementQueryBuilder
 {
+	using namespace UE::Editor::DataStorage;
+
 	TYPEDELEMENTFRAMEWORK_API const UScriptStruct* Type(FTopLevelAssetPath Name);
 	TYPEDELEMENTFRAMEWORK_API const UScriptStruct* TypeOptional(FTopLevelAssetPath Name);
 	TYPEDELEMENTFRAMEWORK_API const UScriptStruct* operator""_Type(const char* Name, std::size_t NameSize);
@@ -132,8 +134,8 @@ namespace TypedElementQueryBuilder
 		TYPEDELEMENTFRAMEWORK_API FDependency& ReadWrite(const UClass* Target);
 		TYPEDELEMENTFRAMEWORK_API FDependency& ReadWrite(TConstArrayView<const UClass*> Targets);
 
-		TYPEDELEMENTFRAMEWORK_API FDependency& SubQuery(TypedElementQueryHandle Handle);
-		TYPEDELEMENTFRAMEWORK_API FDependency& SubQuery(TConstArrayView<TypedElementQueryHandle> Handles);
+		TYPEDELEMENTFRAMEWORK_API FDependency& SubQuery(QueryHandle Handle);
+		TYPEDELEMENTFRAMEWORK_API FDependency& SubQuery(TConstArrayView<QueryHandle> Handles);
 
 		TYPEDELEMENTFRAMEWORK_API ITypedElementDataStorageInterface::FQueryDescription&& Compile();
 
@@ -152,12 +154,12 @@ namespace TypedElementQueryBuilder
 		TYPEDELEMENTFRAMEWORK_API FDependency DependsOn();
 		TYPEDELEMENTFRAMEWORK_API ITypedElementDataStorageInterface::FQueryDescription&& Compile();
 
-		template<UE::Editor::DataStorage::TColumnType... TargetTypes>
+		template<TColumnType... TargetTypes>
 		FSimpleQuery& All();
 		TYPEDELEMENTFRAMEWORK_API FSimpleQuery& All(const UScriptStruct* Target);
 		TYPEDELEMENTFRAMEWORK_API FSimpleQuery& All(TConstArrayView<const UScriptStruct*> Targets);
 		
-		template<UE::Editor::DataStorage::TColumnType... TargetTypes>
+		template<TColumnType... TargetTypes>
 		FSimpleQuery& Any();
 		TYPEDELEMENTFRAMEWORK_API FSimpleQuery& Any(const UScriptStruct* Target);
 		TYPEDELEMENTFRAMEWORK_API FSimpleQuery& Any(TConstArrayView<const UScriptStruct*> Targets);
@@ -173,7 +175,7 @@ namespace TypedElementQueryBuilder
 		template<UE::Editor::DataStorage::TValueTagType>
 		FSimpleQuery& Any(const FName&, const FName&) = delete;
 		
-		template<UE::Editor::DataStorage::TColumnType... TargetTypes>
+		template<TColumnType... TargetTypes>
 		FSimpleQuery& None();
 		
 		TYPEDELEMENTFRAMEWORK_API FSimpleQuery& None(const UScriptStruct* Target);
@@ -193,10 +195,10 @@ namespace TypedElementQueryBuilder
 		// Dynamic Tags
 		// ============
 		// Adds a filter to the query which must match a given DynamicTag.  The value of the DynamicTag will not be checked.
-		TYPEDELEMENTFRAMEWORK_API FSimpleQuery& All(const UE::Editor::DataStorage::FDynamicTag& Tag);
+		TYPEDELEMENTFRAMEWORK_API FSimpleQuery& All(const FDynamicTag& Tag);
 		// Adds a filter to the query which must match a given DynamicTag.  The value of the DynamicTag must also match.
 		// Note: The query can only match a single value.  Multiple value queries are not supported at this time.
-		TYPEDELEMENTFRAMEWORK_API FSimpleQuery& All(const UE::Editor::DataStorage::FDynamicTag& Tag, const FName& Value);
+		TYPEDELEMENTFRAMEWORK_API FSimpleQuery& All(const FDynamicTag& Tag, const FName& Value);
 		
 		TYPEDELEMENTFRAMEWORK_API FSimpleQuery& All(const UEnum& Enum);
 		TYPEDELEMENTFRAMEWORK_API FSimpleQuery& All(const UEnum& Enum, int64 Value);
@@ -214,13 +216,13 @@ namespace TypedElementQueryBuilder
 		template<UE::Editor::DataStorage::TValueTagType>
 		FSimpleQuery& All(const FName& Tag, const FName& Value);
 		
-		template<UE::Editor::DataStorage::TEnumType EnumT>
+		template<TEnumType EnumT>
 		FSimpleQuery& All();
 		
-		template<UE::Editor::DataStorage::TEnumType EnumT>
+		template<TEnumType EnumT>
 		FSimpleQuery& All(EnumT EnumValue);
 
-		template<auto Value, UE::Editor::DataStorage::TEnumType EnumT = decltype(Value)>
+		template<auto Value, TEnumType EnumT = decltype(Value)>
 		FSimpleQuery& All();
 
 	private:
@@ -259,14 +261,14 @@ namespace TypedElementQueryBuilder
 
 		TYPEDELEMENTFRAMEWORK_API FObserver(EEvent MonitorForEvent, const UScriptStruct* MonitoredColumn);
 
-		template<UE::Editor::DataStorage::TColumnType ColumnType>
+		template<TColumnType ColumnType>
 		static FObserver OnAdd();
-		template<UE::Editor::DataStorage::TColumnType ColumnType>
+		template<TColumnType ColumnType>
 		static FObserver OnRemove();
 
 		TYPEDELEMENTFRAMEWORK_API FObserver& SetEvent(EEvent MonitorForEvent);
 		TYPEDELEMENTFRAMEWORK_API FObserver& SetMonitoredColumn(const UScriptStruct* MonitoredColumn);
-		template<UE::Editor::DataStorage::TColumnType ColumnType>
+		template<TColumnType ColumnType>
 		FObserver& SetMonitoredColumn();
 		TYPEDELEMENTFRAMEWORK_API FObserver& SetExecutionMode(TypedElementDataStorage::EExecutionMode Mode);
 		TYPEDELEMENTFRAMEWORK_API FObserver& MakeActivatable(FName Name);
@@ -321,19 +323,19 @@ namespace TypedElementQueryBuilder
 			TConstArrayView<ITypedElementDataStorageInterface::EQueryAccessType> AccessTypes) override;
 
 		inline uint32 GetRowCount() const override;
-		inline TConstArrayView<TypedElementRowHandle> GetRowHandles() const override;
-		inline void RemoveRow(TypedElementRowHandle Row) override;
-		inline void RemoveRows(TConstArrayView<TypedElementRowHandle> Rows) override;
+		inline TConstArrayView<RowHandle> GetRowHandles() const override;
+		inline void RemoveRow(RowHandle Row) override;
+		inline void RemoveRows(TConstArrayView<RowHandle> Rows) override;
 
-		inline void AddColumns(TypedElementRowHandle Row, TConstArrayView<const UScriptStruct*> ColumnTypes) override;
-		inline void AddColumns(TConstArrayView<TypedElementRowHandle> Rows, TConstArrayView<const UScriptStruct*> ColumnTypes) override;
-		inline void RemoveColumns(TypedElementRowHandle Row, TConstArrayView<const UScriptStruct*> ColumnTypes) override;
-		inline void RemoveColumns(TConstArrayView<TypedElementRowHandle> Rows, TConstArrayView<const UScriptStruct*> ColumnTypes) override;
+		inline void AddColumns(RowHandle Row, TConstArrayView<const UScriptStruct*> ColumnTypes) override;
+		inline void AddColumns(TConstArrayView<RowHandle> Rows, TConstArrayView<const UScriptStruct*> ColumnTypes) override;
+		inline void RemoveColumns(RowHandle Row, TConstArrayView<const UScriptStruct*> ColumnTypes) override;
+		inline void RemoveColumns(TConstArrayView<RowHandle> Rows, TConstArrayView<const UScriptStruct*> ColumnTypes) override;
 
-		inline TypedElementDataStorage::FQueryResult RunQuery(TypedElementQueryHandle Query) override;
+		inline TypedElementDataStorage::FQueryResult RunQuery(QueryHandle Query) override;
 		inline TypedElementDataStorage::FQueryResult RunSubquery(int32 SubqueryIndex) override;
 		inline TypedElementDataStorage::FQueryResult RunSubquery(int32 SubqueryIndex, TypedElementDataStorage::SubqueryCallbackRef Callback) override;
-		inline TypedElementDataStorage::FQueryResult RunSubquery(int32 SubqueryIndex, TypedElementDataStorage::RowHandle Row,
+		inline TypedElementDataStorage::FQueryResult RunSubquery(int32 SubqueryIndex, RowHandle Row,
 			TypedElementDataStorage::SubqueryCallbackRef Callback) override;
 
 		ITypedElementDataStorageInterface::IQueryContext& ParentContext;
@@ -369,7 +371,7 @@ namespace TypedElementQueryBuilder
 		Select(FName Name, const CallbackType& Type, Class* Instance, Function&& Callback);
 
 		/** Request read-only access to the listed columns. */
-		template<UE::Editor::DataStorage::TDataColumnType... TargetTypes>
+		template<TDataColumnType... TargetTypes>
 		Select& ReadOnly();
 		/** Request read-only access to the listed columns. */
 		TYPEDELEMENTFRAMEWORK_API Select& ReadOnly(const UScriptStruct* Target);
@@ -380,7 +382,7 @@ namespace TypedElementQueryBuilder
 		 * Request read-only access to the listed columns. If optional is true read access will be given if the column is in the table but
 		 * it will not be used for finding matching tables. Columns bound with optional can not be bound to a query callback argument.
 		 */
-		template<UE::Editor::DataStorage::TDataColumnType... TargetTypes>
+		template<TDataColumnType... TargetTypes>
 		Select& ReadOnly(EOptional Optional);
 		/**
 		 * Request read-only access to the listed columns. If optional is true read access will be given if the column is in the table but
@@ -395,7 +397,7 @@ namespace TypedElementQueryBuilder
 		template<UE::Editor::DataStorage::TDataColumnType Target>
 		Select& ReadOnly(const FName& Identifier);
 		/** Request read and write access to the listed columns. */
-		template<UE::Editor::DataStorage::TDataColumnType... TargetTypes>
+		template<TDataColumnType... TargetTypes>
 		Select& ReadWrite();
 		/** Request read and write access to the listed columns. */
 		TYPEDELEMENTFRAMEWORK_API Select& ReadWrite(const UScriptStruct* Target);
diff --git a/Engine/Source/Runtime/TypedElementFramework/Public/Elements/Framework/TypedElementQueryBuilder.inl b/Engine/Source/Runtime/TypedElementFramework/Public/Elements/Framework/TypedElementQueryBuilder.inl
--- a/Engine/Source/Runtime/TypedElementFramework/Public/Elements/Framework/TypedElementQueryBuilder.inl
+++ b/Engine/Source/Runtime/TypedElementFramework/Public/Elements/Framework/TypedElementQueryBuilder.inl
@@ -170,42 +170,42 @@ namespace TypedElementQueryBuilder
 		return ParentContext.GetRowCount();
 	}
 
-	TConstArrayView<TypedElementRowHandle> FQueryContextForwarder::GetRowHandles() const
+	TConstArrayView<RowHandle> FQueryContextForwarder::GetRowHandles() const
 	{
 		return ParentContext.GetRowHandles();
 	}
 
-	void FQueryContextForwarder::RemoveRow(TypedElementRowHandle Row)
+	void FQueryContextForwarder::RemoveRow(RowHandle Row)
 	{
 		ParentContext.RemoveRow(Row);
 	}
 
-	void FQueryContextForwarder::RemoveRows(TConstArrayView<TypedElementRowHandle> Rows)
+	void FQueryContextForwarder::RemoveRows(TConstArrayView<RowHandle> Rows)
 	{
 		ParentContext.RemoveRows(Rows);
 	}
 
-	void FQueryContextForwarder::AddColumns(TypedElementRowHandle Row, TConstArrayView<const UScriptStruct*> ColumnTypes)
+	void FQueryContextForwarder::AddColumns(RowHandle Row, TConstArrayView<const UScriptStruct*> ColumnTypes)
 	{
 		ParentContext.AddColumns(Row, ColumnTypes);
 	}
 
-	void FQueryContextForwarder::AddColumns(TConstArrayView<TypedElementRowHandle> Rows, TConstArrayView<const UScriptStruct*> ColumnTypes)
+	void FQueryContextForwarder::AddColumns(TConstArrayView<RowHandle> Rows, TConstArrayView<const UScriptStruct*> ColumnTypes)
 	{
 		ParentContext.AddColumns(Rows, ColumnTypes);
 	}
 
-	void FQueryContextForwarder::RemoveColumns(TypedElementRowHandle Row, TConstArrayView<const UScriptStruct*> ColumnTypes)
+	void FQueryContextForwarder::RemoveColumns(RowHandle Row, TConstArrayView<const UScriptStruct*> ColumnTypes)
 	{
 		ParentContext.RemoveColumns(Row, ColumnTypes);
 	}
 
-	void FQueryContextForwarder::RemoveColumns(TConstArrayView<TypedElementRowHandle> Rows, TConstArrayView<const UScriptStruct*> ColumnTypes)
+	void FQueryContextForwarder::RemoveColumns(TConstArrayView<RowHandle> Rows, TConstArrayView<const UScriptStruct*> ColumnTypes)
 	{
 		ParentContext.RemoveColumns(Rows, ColumnTypes);
 	}
 
-	TypedElementDataStorage::FQueryResult FQueryContextForwarder::RunQuery(TypedElementQueryHandle Query)
+	TypedElementDataStorage::FQueryResult FQueryContextForwarder::RunQuery(QueryHandle Query)
 	{
 		return ParentContext.RunQuery(Query);
 	}
@@ -222,7 +222,7 @@ namespace TypedElementQueryBuilder
 	}
 
 	TypedElementDataStorage::FQueryResult FQueryContextForwarder::RunSubquery(int32 SubqueryIndex,
-		TypedElementDataStorage::RowHandle Row, TypedElementDataStorage::SubqueryCallbackRef Callback)
+		RowHandle Row, TypedElementDataStorage::SubqueryCallbackRef Callback)
 	{
 		return ParentContext.RunSubquery(SubqueryIndex, Row, Callback);
 	}
@@ -294,7 +294,7 @@ namespace TypedElementQueryBuilder
 			const TWeakObjectPtr<const UScriptStruct>*ColumnTypes, const TypedElementDataStorage::EQueryAccessType * AccessTypes)
 		{
 			{ Ctx.GetColumnsUnguarded(TypeCount, RetrievedAddresses, ColumnTypes, AccessTypes) };
-			{ Ctx.GetRowHandles() } -> UE::convertible_to<TConstArrayView<TypedElementRowHandle>>;
+			{ Ctx.GetRowHandles() } -> UE::convertible_to<TConstArrayView<RowHandle>>;
 			{ Ctx.GetRowCount() } -> UE::convertible_to<uint32>;
 		};
 
@@ -362,7 +362,7 @@ namespace TypedElementQueryBuilder
 		}
 
 		template<typename RowType>
-		concept IsRowHandleType = std::is_same_v<std::remove_cv_t<std::remove_pointer_t<RowType>>, TypedElementRowHandle>;
+		concept IsRowHandleType = std::is_same_v<std::remove_cv_t<std::remove_pointer_t<RowType>>, RowHandle>;
 
 		template<typename RowType, typename... Columns>
 		constexpr bool IsRowTypeCompatibleWithColumns()
@@ -501,7 +501,7 @@ namespace TypedElementQueryBuilder
 			template<typename CallerType>
 			void Call(SourceContext& Context, const CallerType& Caller)
 			{
-				TConstArrayView<TypedElementRowHandle> Rows = Context.GetRowHandles();
+				TConstArrayView<RowHandle> Rows = Context.GetRowHandles();
 				if constexpr (SuperColumn::bArePointerColumns || std::is_pointer_v<RowHandleType>)
 				{
 					if constexpr (sizeof...(Args) > 0)
@@ -518,7 +518,7 @@ namespace TypedElementQueryBuilder
 				}
 				else
 				{
-					for (TypedElementRowHandle Row : Rows)
+					for (RowHandle Row : Rows)
 					{
 						if constexpr (sizeof...(Args) > 0)
 						{
@@ -601,7 +601,7 @@ namespace TypedElementQueryBuilder
 			template<typename CallerType>
 			void Call(SourceContext& Context, const CallerType& Caller)
 			{
-				TConstArrayView<TypedElementRowHandle> Rows = Context.GetRowHandles();
+				TConstArrayView<RowHandle> Rows = Context.GetRowHandles();
 				if constexpr (Super::bArePointerColumns || std::is_pointer_v<RowHandleType>)
 				{
 					if constexpr (sizeof...(Columns) > 0)
@@ -618,7 +618,7 @@ namespace TypedElementQueryBuilder
 				}
 				else
 				{
-					for (TypedElementRowHandle Row : Rows)
+					for (RowHandle Row : Rows)
 					{
 						if constexpr (sizeof...(Columns) > 0)
 						{
@@ -1087,13 +1087,13 @@ namespace TypedElementQueryBuilder
 				R"(The function provided to the Query Builder's Select call wasn't invocable or doesn't contain a supported combination of arguments.
 The following options are supported:
 - void([const]Column&...) 
-- void(TypedElementRowHandle, [const]Column&...) 
+- void(RowHandle, [const]Column&...) 
 - void(<Context>&, [const]Column&...) 
-- void(<Context>&, TypedElementRowHandle, [const]Column&...) 
+- void(<Context>&, RowHandle, [const]Column&...) 
 - void(<Context>&, [const]Column*...) 
-- void(<Context>&, const TypedElementRowHandle*, [const]Column*...) 
+- void(<Context>&, const RowHandle*, [const]Column*...) 
 Where <Context> is TypedElementDataStorage::IQueryContext or FCachedQueryContext<...>
-e.g. void(FCachedQueryContext<Subsystem1, const Subsystem2>& Context, TypedElementRowHandle Row, ColumnType0& ColumnA, const ColumnType1& ColumnB) {...}
+e.g. void(FCachedQueryContext<Subsystem1, const Subsystem2>& Context, RowHandle Row, ColumnType0& ColumnA, const ColumnType1& ColumnB) {...}
 )");
 			RegisterFunctionArguments<TypedElementDataStorage::IQueryContext>(Query, Target, Callback);
 			PrepareForQueryBinding(Query, Type);
@@ -1258,13 +1258,13 @@ e.g. void(FCachedQueryContext<Subsystem1, const Subsystem2>& Context, TypedEleme
 			R"(The function provided to the Query Builder's CreateDirectQueryCallbackBinding call wasn't invocable or doesn't contain a supported combination of arguments.
 The following options are supported:
 - void([const]Column&...) 
-- void(TypedElementRowHandle, [const]Column&...) 
+- void(RowHandle, [const]Column&...) 
 - void(<Context>&, [const]Column&...) 
-- void(<Context>&, TypedElementRowHandle, [const]Column&...) 
+- void(<Context>&, RowHandle, [const]Column&...) 
 - void(<Context>&, [const]Column*...) 
-- void(<Context>&, const TypedElementRowHandle*, [const]Column*...) 
+- void(<Context>&, const RowHandle*, [const]Column*...) 
 Where <Context> is TypedElementDataStorage::IDirectQueryContext
-e.g. void(IDirectQueryContext& Context, TypedElementRowHandle Row, ColumnType0& ColumnA, const ColumnType1& ColumnB) {...}
+e.g. void(IDirectQueryContext& Context, RowHandle Row, ColumnType0& ColumnA, const ColumnType1& ColumnB) {...}
 )");
 
 		return [Callback = Forward<Function>(Callback)](
@@ -1283,13 +1283,13 @@ e.g. void(IDirectQueryContext& Context, TypedElementRowHandle Row, ColumnType0&
 			R"(The function provided to the Query Builder's CreateSubqueryCallbackBinding call wasn't invocable or doesn't contain a supported combination of arguments.
 The following options are supported:
 - void([const]Column&...) 
-- void(TypedElementRowHandle, [const]Column&...) 
+- void(RowHandle, [const]Column&...) 
 - void(<Context>&, [const]Column&...) 
-- void(<Context>&, TypedElementRowHandle, [const]Column&...) 
+- void(<Context>&, RowHandle, [const]Column&...) 
 - void(<Context>&, [const]Column*...) 
-- void(<Context>&, const TypedElementRowHandle*, [const]Column*...) 
+- void(<Context>&, const RowHandle*, [const]Column*...) 
 Where <Context> is TypedElementDataStorage::ISubqueryContext
-e.g. void(ISubqueryContext& Context, TypedElementRowHandle Row, ColumnType0& ColumnA, const ColumnType1& ColumnB) {...}
+e.g. void(ISubqueryContext& Context, RowHandle Row, ColumnType0& ColumnA, const ColumnType1& ColumnB) {...}
 )");
 
 		return [Callback = Forward<Function>(Callback)](
diff --git a/Engine/Source/Runtime/TypedElementFramework/Public/Elements/Interfaces/TypedElementDataStorageCompatibilityInterface.h b/Engine/Source/Runtime/TypedElementFramework/Public/Elements/Interfaces/TypedElementDataStorageCompatibilityInterface.h
--- a/Engine/Source/Runtime/TypedElementFramework/Public/Elements/Interfaces/TypedElementDataStorageCompatibilityInterface.h
+++ b/Engine/Source/Runtime/TypedElementFramework/Public/Elements/Interfaces/TypedElementDataStorageCompatibilityInterface.h
@@ -29,7 +29,7 @@ class ITypedElementDataStorageCompatibilityInterface
 
 public:
 	using ObjectRegistrationFilter = TFunction<bool(const ITypedElementDataStorageCompatibilityInterface&, const UObject*)>;
-	using ObjectToRowDealiaser = TFunction<TypedElementRowHandle(const ITypedElementDataStorageCompatibilityInterface&, const UObject*)>;
+	using ObjectToRowDealiaser = TFunction<UE::Editor::DataStorage::RowHandle(const ITypedElementDataStorageCompatibilityInterface&, const UObject*)>;
 	
 	/**
 	 * @section Type-agnostic functions
@@ -44,14 +44,14 @@ public:
 	 * with a row and to setup the initial row data.
 	 */
 	template<typename ObjectType>
-	TypedElementRowHandle AddCompatibleObject(ObjectType&& Object);
+	UE::Editor::DataStorage::RowHandle AddCompatibleObject(ObjectType&& Object);
 	
 	/** Removes a previously registered object from the data storage. */
 	template<typename ObjectType>
 	void RemoveCompatibleObject(ObjectType&& Object);
 
 	template<typename ObjectType>
-	TypedElementRowHandle FindRowWithCompatibleObject(ObjectType&& Object) const;
+	UE::Editor::DataStorage::RowHandle FindRowWithCompatibleObject(ObjectType&& Object) const;
 
 	/**
 	 * @section Callback registration
@@ -76,7 +76,7 @@ public:
 	 * will be used to find the closest match in the registered types and use the associated table. E.g. actors derive from uobjects so
 	 * if the type information of an actor is registered the actor table will be used instead of the uobject table.
 	 */
-	virtual void RegisterTypeTableAssociation(TObjectPtr<UStruct> TypeInfo, TypedElementDataStorage::TableHandle Table) = 0;
+	virtual void RegisterTypeTableAssociation(TObjectPtr<UStruct> TypeInfo, UE::Editor::DataStorage::TableHandle Table) = 0;
 
 	/**
 	 * @section Explicit functions
@@ -84,9 +84,9 @@ public:
 	 */
 
 	/** Adds a UObject to the data storage. */
-	virtual TypedElementRowHandle AddCompatibleObjectExplicit(UObject* Object) = 0;
+	virtual UE::Editor::DataStorage::RowHandle AddCompatibleObjectExplicit(UObject* Object) = 0;
 	/** Adds an FStruct to the data storage. */
-	virtual TypedElementRowHandle AddCompatibleObjectExplicit(void* Object, TWeakObjectPtr<UScriptStruct> TypeInfo) = 0;
+	virtual UE::Editor::DataStorage::RowHandle AddCompatibleObjectExplicit(void* Object, TWeakObjectPtr<UScriptStruct> TypeInfo) = 0;
 	
 	/** Removes a UObject from the data storage. */
 	virtual void RemoveCompatibleObjectExplicit(UObject* Object) = 0;
@@ -94,9 +94,9 @@ public:
 	virtual void RemoveCompatibleObjectExplicit(void* Object) = 0;
 
 	/** Finds a previously stored UObject. If not found an invalid row handle will be returned. */
-	virtual TypedElementRowHandle FindRowWithCompatibleObjectExplicit(const UObject* Object) const = 0;
+	virtual UE::Editor::DataStorage::RowHandle FindRowWithCompatibleObjectExplicit(const UObject* Object) const = 0;
 	/** Finds a previously stored FStruct. If not found an invalid row handle will be returned. */
-	virtual TypedElementRowHandle FindRowWithCompatibleObjectExplicit(const void* Object) const = 0;
+	virtual UE::Editor::DataStorage::RowHandle FindRowWithCompatibleObjectExplicit(const void* Object) const = 0;
 
 	/**
 	 * @section Miscellaneous functions
@@ -118,7 +118,7 @@ template<typename Type> Type* GetRawPointer(Type* Object)						{ return Object;
 template<typename Type> Type* GetRawPointer(Type& Object)						{ return &Object; }
 
 template<typename ObjectType>
-TypedElementRowHandle ITypedElementDataStorageCompatibilityInterface::AddCompatibleObject(ObjectType&& Object)
+UE::Editor::DataStorage::RowHandle ITypedElementDataStorageCompatibilityInterface::AddCompatibleObject(ObjectType&& Object)
 {
 	auto RawPointer = GetRawPointer(Forward<ObjectType>(Object));
 	using BaseType = std::remove_cv_t<std::remove_pointer_t<decltype(RawPointer)>>;
@@ -140,7 +140,7 @@ void ITypedElementDataStorageCompatibilityInterface::RemoveCompatibleObject(Obje
 }
 
 template<typename ObjectType>
-TypedElementRowHandle ITypedElementDataStorageCompatibilityInterface::FindRowWithCompatibleObject(ObjectType&& Object) const
+UE::Editor::DataStorage::RowHandle ITypedElementDataStorageCompatibilityInterface::FindRowWithCompatibleObject(ObjectType&& Object) const
 {
 	return FindRowWithCompatibleObjectExplicit(GetRawPointer(Forward<ObjectType>(Object)));
 }
diff --git a/Engine/Source/Runtime/TypedElementFramework/Public/Elements/Interfaces/TypedElementDataStorageInterface.h b/Engine/Source/Runtime/TypedElementFramework/Public/Elements/Interfaces/TypedElementDataStorageInterface.h
--- a/Engine/Source/Runtime/TypedElementFramework/Public/Elements/Interfaces/TypedElementDataStorageInterface.h
+++ b/Engine/Source/Runtime/TypedElementFramework/Public/Elements/Interfaces/TypedElementDataStorageInterface.h
@@ -27,13 +27,6 @@ class USubsystem;
 class UScriptStruct;
 class UTypedElementDataStorageFactory;
 
-using TypedElementTableHandle = TypedElementDataStorage::TableHandle;
-static constexpr auto TypedElementInvalidTableHandle = TypedElementDataStorage::InvalidTableHandle;
-using TypedElementRowHandle = TypedElementDataStorage::RowHandle;
-static constexpr auto TypedElementInvalidRowHandle = TypedElementDataStorage::InvalidRowHandle;
-using TypedElementQueryHandle = TypedElementDataStorage::QueryHandle;
-static constexpr auto TypedElementInvalidQueryHandle = TypedElementDataStorage::InvalidQueryHandle;
-
 using FTypedElementOnDataStorageCreation = FSimpleMulticastDelegate;
 using FTypedElementOnDataStorageDestruction = FSimpleMulticastDelegate;
 using FTypedElementOnDataStorageUpdate = FSimpleMulticastDelegate;
@@ -67,6 +60,9 @@ class ITypedElementDataStorageInterface
 	using ColumnListCallbackRef = UE::Editor::DataStorage::ColumnListCallbackRef;
 	using ColumnListWithDataCallbackRef = UE::Editor::DataStorage::ColumnListWithDataCallbackRef;
 	using ColumnCopyOrMoveCallback = UE::Editor::DataStorage::ColumnCopyOrMoveCallback;
+	using RowHandle = UE::Editor::DataStorage::RowHandle;
+	using TableHandle = UE::Editor::DataStorage::TableHandle;
+	using QueryHandle = UE::Editor::DataStorage::QueryHandle;
 
 public:
 	/**
@@ -94,20 +90,20 @@ public:
 	 */
 
 	/** Creates a new table for with the provided columns. Optionally a name can be given which is useful for retrieval later. */
-	virtual TypedElementDataStorage::TableHandle RegisterTable(TConstArrayView<const UScriptStruct*> ColumnList, const FName Name) = 0;
+	virtual TableHandle RegisterTable(TConstArrayView<const UScriptStruct*> ColumnList, const FName Name) = 0;
 	template<UE::Editor::DataStorage::TColumnType... Columns>
-	TypedElementDataStorage::TableHandle RegisterTable(const FName Name);
+	TableHandle RegisterTable(const FName Name);
 	/** 
 	 * Copies the column information from the provided table and creates a new table for with the provided columns. Optionally a 
 	 * name can be given which is useful for retrieval later.
 	 */
-	virtual TypedElementDataStorage::TableHandle RegisterTable(TypedElementDataStorage::TableHandle SourceTable,
+	virtual TableHandle RegisterTable(TableHandle SourceTable,
 		TConstArrayView<const UScriptStruct*> ColumnList, const FName Name) = 0;
 	template<UE::Editor::DataStorage::TColumnType... Columns>
-	TypedElementDataStorage::TableHandle RegisterTable(TypedElementDataStorage::TableHandle SourceTable, const FName Name);
+	TableHandle RegisterTable(TableHandle SourceTable, const FName Name);
 
 	/** Returns a previously created table with the provided name or TypedElementInvalidTableHandle if not found. */
-	virtual TypedElementDataStorage::TableHandle FindTable(const FName Name) = 0;
+	virtual TableHandle FindTable(const FName Name) = 0;
 	
 	/**
 	 * @section Row management
@@ -117,59 +113,59 @@ public:
 	 * Reserves a row to be assigned to a table at a later point. If the row is no longer needed before it's been assigned
 	 * to a table, it should still be released with RemoveRow.
 	 */
-	virtual TypedElementDataStorage::RowHandle ReserveRow() = 0;
+	virtual RowHandle ReserveRow() = 0;
 	/**
 	 * Reserve multiple rows at once to be assigned to a table at a later point. If multiple rows are needed, the batch version will
 	 * generally have better performance. If a row is no longer needed before it's been assigned to a table, it should still be released 
 	 * with RemoveRow.
 	 * The reservation callback will be called once per reserved row.
 	 */
-	virtual void BatchReserveRows(int32 Count, TFunctionRef<void(TypedElementDataStorage::RowHandle)> ReservationCallback) = 0;
+	virtual void BatchReserveRows(int32 Count, TFunctionRef<void(RowHandle)> ReservationCallback) = 0;
 	/**
 	 * Reserve multiple rows at once to be assigned to a table at a later point. If multiple rows are needed, the batch version will
 	 * generally have better performance. If a row is no longer needed before it's been assigned to a table, it should still be released
 	 * with RemoveRow.
 	 * The provided range will be have its values set to the reserved row handles.
 	 */
-	virtual void BatchReserveRows(TArrayView<TypedElementDataStorage::RowHandle> ReservedRows) = 0;
+	virtual void BatchReserveRows(TArrayView<RowHandle> ReservedRows) = 0;
 
 	/** Adds a new row to the provided table. */
-	virtual TypedElementDataStorage::RowHandle AddRow(TypedElementDataStorage::TableHandle Table) = 0;
+	virtual RowHandle AddRow(TableHandle Table) = 0;
 	/**
 	 * Adds a new row to the provided table. Callers are expected to use the callback to
 	 * initialize the row if needed.
 	 */
-	virtual TypedElementDataStorage::RowHandle AddRow(TypedElementDataStorage::TableHandle Table,
+	virtual RowHandle AddRow(TableHandle Table,
 		RowCreationCallbackRef OnCreated) = 0;
 	/** Adds a new row to the provided table using a previously reserved row. */
-	virtual bool AddRow(TypedElementDataStorage::RowHandle ReservedRow, TypedElementDataStorage::TableHandle Table) = 0;
+	virtual bool AddRow(RowHandle ReservedRow, TableHandle Table) = 0;
 	/**
 	 * Adds a new row to the provided table using a previously reserved row. Callers are expected to use the callback to
 	 * initialize the row if needed.
 	 */
-	virtual bool AddRow(TypedElementDataStorage::RowHandle ReservedRow, TypedElementDataStorage::TableHandle Table,
+	virtual bool AddRow(RowHandle ReservedRow, TableHandle Table,
 		RowCreationCallbackRef OnCreated) = 0;
 
 	/**
 	 * Add multiple rows at once. For each new row the OnCreated callback is called. Callers are expected to use the callback to
 	 * initialize the row if needed.
 	 */
-	virtual bool BatchAddRow(TypedElementDataStorage::TableHandle Table, int32 Count, RowCreationCallbackRef OnCreated) = 0;
+	virtual bool BatchAddRow(TableHandle Table, int32 Count, RowCreationCallbackRef OnCreated) = 0;
 	/**
 	 * Add multiple rows at once. For each new row the OnCreated callback is called. Callers are expected to use the callback to
 	 * initialize the row if needed. This version uses a set of previously reserved rows. Any row that can't be used will be 
 	 * released.
 	 */
-	virtual bool BatchAddRow(TypedElementDataStorage::TableHandle Table, TConstArrayView<TypedElementDataStorage::RowHandle> ReservedHandles,
+	virtual bool BatchAddRow(TableHandle Table, TConstArrayView<RowHandle> ReservedHandles,
 		RowCreationCallbackRef OnCreated) = 0;
 
 	/** Removes a previously reserved or added row. If the row handle is invalid or already removed, nothing happens */
-	virtual void RemoveRow(TypedElementDataStorage::RowHandle Row) = 0;
+	virtual void RemoveRow(RowHandle Row) = 0;
 
 	/** Checks whether or not a row is in use. This is true even if the row has only been reserved. */
-	virtual bool IsRowAvailable(TypedElementDataStorage::RowHandle Row) const = 0;
+	virtual bool IsRowAvailable(RowHandle Row) const = 0;
 	/** Checks whether or not a row has been reserved but not yet assigned to a table. */
-	virtual bool IsRowAssigned(TypedElementDataStorage::RowHandle Row) const = 0;
+	virtual bool IsRowAssigned(RowHandle Row) const = 0;
 
 	
 	/**
@@ -177,18 +173,18 @@ public:
 	 */
 
 	/** Adds a column to a row or does nothing if already added. */
-	virtual void AddColumn(TypedElementRowHandle Row, const UScriptStruct* ColumnType) = 0;
+	virtual void AddColumn(RowHandle Row, const UScriptStruct* ColumnType) = 0;
 	template<UE::Editor::DataStorage::TColumnType ColumnType>
-	void AddColumn(TypedElementRowHandle Row);
+	void AddColumn(RowHandle Row);
 	/**
 	 * Adds a new data column and initializes it. The relocator will be used to copy or move the column out of
 	 * its temporary location into the final table if the addition needs to be deferred.
 	 */
-	virtual void AddColumnData(TypedElementRowHandle Row, const UScriptStruct* ColumnType,
+	virtual void AddColumnData(RowHandle Row, const UScriptStruct* ColumnType,
 		const ColumnCreationCallbackRef& Initializer,
 		ColumnCopyOrMoveCallback Relocator) = 0;
 	template<UE::Editor::DataStorage::TDataColumnType ColumnType>
-	void AddColumn(TypedElementRowHandle Row, ColumnType&& Column);
+	void AddColumn(RowHandle Row, ColumnType&& Column);
 
 	/**
 	 * Adds a DynamicTag with the given value to a row
@@ -201,22 +197,22 @@ public:
 	 *       is added.  This will likely change in the future to transparently replace the tag to have consistent behaviour with other usages
 	 *       of AddColumn
 	 */
-	virtual void AddColumn(TypedElementDataStorage::RowHandle Row, const UE::Editor::DataStorage::FDynamicTag& Tag, const FName& Value) = 0;
+	virtual void AddColumn(RowHandle Row, const UE::Editor::DataStorage::FDynamicTag& Tag, const FName& Value) = 0;
 
 	template<typename T>
 	void AddColumn(TypedElementDataStorage::RowHandle Row, const FName& Tag) = delete;
 	
 	template<typename T>
-	void AddColumn(TypedElementDataStorage::RowHandle Row, const FName& Tag, const FName& Value) = delete;
+	void AddColumn(RowHandle Row, const FName& Tag, const FName& Value) = delete;
 	
 	template<>
-	void AddColumn<UE::Editor::DataStorage::FDynamicTag>(TypedElementDataStorage::RowHandle Row, const FName& Tag, const FName& Value);
+	void AddColumn<UE::Editor::DataStorage::FDynamicTag>(RowHandle Row, const FName& Tag, const FName& Value);
 
 	template<UE::Editor::DataStorage::TEnumType EnumT>
-	void AddColumn(TypedElementDataStorage::RowHandle Row, EnumT Value);
+	void AddColumn(RowHandle Row, EnumT Value);
 	
 	template<auto Value, UE::Editor::DataStorage::TEnumType EnumT = decltype(Value)>
-	void AddColumn(TypedElementDataStorage::RowHandle Row);
+	void AddColumn(RowHandle Row);
 
 	template<UE::Editor::DataStorage::TColumnType DynamicColumnTemplate>
 	void AddColumn(TypedElementDataStorage::RowHandle Row, const FName& Identifier);
@@ -228,86 +224,86 @@ public:
 	 * Adds multiple columns from a row. This is typically more efficient than adding columns one 
 	 * at a time.
 	 */
-	virtual void AddColumns(TypedElementRowHandle Row, TConstArrayView<const UScriptStruct*> Columns) = 0;
+	virtual void AddColumns(RowHandle Row, TConstArrayView<const UScriptStruct*> Columns) = 0;
 	template<UE::Editor::DataStorage::TColumnType... Columns>
-	void AddColumns(TypedElementRowHandle Row);
+	void AddColumns(RowHandle Row);
 
 	/** Removes a column from a row or does nothing if already removed. */
-	virtual void RemoveColumn(TypedElementRowHandle Row, const UScriptStruct* ColumnType) = 0;
+	virtual void RemoveColumn(RowHandle Row, const UScriptStruct* ColumnType) = 0;
 	template<UE::Editor::DataStorage::TColumnType Column>
-	void RemoveColumn(TypedElementRowHandle Row);
+	void RemoveColumn(RowHandle Row);
 
 	template<UE::Editor::DataStorage::TEnumType EnumT>
-	void RemoveColumn(TypedElementDataStorage::RowHandle Row);
+	void RemoveColumn(RowHandle Row);
 
 	/**
 	 * Removes a dynamic tag from the given row
 	 * If tag does not exist on row, operation will do nothing.
 	 */
-	virtual void RemoveColumn(TypedElementDataStorage::RowHandle Row, const UE::Editor::DataStorage::FDynamicTag& Tag) = 0;
+	virtual void RemoveColumn(RowHandle Row, const UE::Editor::DataStorage::FDynamicTag& Tag) = 0;
 
 	template<typename T>
-	void RemoveColumn(TypedElementDataStorage::RowHandle Row, const FName& Tag) = delete;
+	void RemoveColumn(RowHandle Row, const FName& Tag) = delete;
 
 	template<UE::Editor::DataStorage::TColumnType DynamicColumnTemplateType>
 	void RemoveColumn(TypedElementDataStorage::RowHandle Row, const FName& Identifier);
 	
 	template<>
-	void RemoveColumn<UE::Editor::DataStorage::FDynamicTag>(TypedElementDataStorage::RowHandle Row, const FName& Tag);
+	void RemoveColumn<UE::Editor::DataStorage::FDynamicTag>(RowHandle Row, const FName& Tag);
 
 	/**
 	 * Removes multiple columns from a row. This is typically more efficient than adding columns one
 	 * at a time.
 	 */
-	virtual void RemoveColumns(TypedElementRowHandle Row, TConstArrayView<const UScriptStruct*> Columns) = 0;
+	virtual void RemoveColumns(RowHandle Row, TConstArrayView<const UScriptStruct*> Columns) = 0;
 	template<UE::Editor::DataStorage::TColumnType... Columns>
-	void RemoveColumns(TypedElementRowHandle Row);
+	void RemoveColumns(RowHandle Row);
 
 	/** 
 	 * Adds and removes the provided column types from the provided row. This is typically more efficient 
 	 * than individually adding and removing columns as well as being faster than adding and removing
 	 * columns separately.
 	 */
-	virtual void AddRemoveColumns(TypedElementRowHandle Row, TConstArrayView<const UScriptStruct*> ColumnsToAdd,
+	virtual void AddRemoveColumns(RowHandle Row, TConstArrayView<const UScriptStruct*> ColumnsToAdd,
 		TConstArrayView<const UScriptStruct*> ColumnsToRemove) = 0;
 	
 	/** Adds and removes the provided column types from the provided list of rows. */
 	virtual void BatchAddRemoveColumns(
-		TConstArrayView<TypedElementRowHandle> Rows,
+		TConstArrayView<RowHandle> Rows,
 		TConstArrayView<const UScriptStruct*> ColumnsToAdd,
 		TConstArrayView<const UScriptStruct*> ColumnsToRemove) = 0;
 	
 	/** Retrieves a pointer to the column of the given row or a nullptr if not found or if the column type is a tag. */
-	virtual void* GetColumnData(TypedElementRowHandle Row, const UScriptStruct* ColumnType) = 0;
-	virtual const void* GetColumnData(TypedElementRowHandle Row, const UScriptStruct* ColumnType) const = 0;
+	virtual void* GetColumnData(RowHandle Row, const UScriptStruct* ColumnType) = 0;
+	virtual const void* GetColumnData(RowHandle Row, const UScriptStruct* ColumnType) const = 0;
 	/** Returns a pointer to the column of the given row or a nullptr if the type couldn't be found or the row doesn't exist. */
 	template<UE::Editor::DataStorage::TDataColumnType ColumnType>
-	ColumnType* GetColumn(TypedElementRowHandle Row);
+	ColumnType* GetColumn(RowHandle Row);
 	template<UE::Editor::DataStorage::TDataColumnType ColumnType>
-	const ColumnType* GetColumn(TypedElementRowHandle Row) const;
+	const ColumnType* GetColumn(RowHandle Row) const;
 	// Gets a dynamic column identified by the ColumnTypeTemplate and Identifier
 	template<UE::Editor::DataStorage::TDataColumnType ColumnTypeTemplate>
-	ColumnTypeTemplate* GetColumn(TypedElementRowHandle Row, const FName& Identifer);
+	ColumnTypeTemplate* GetColumn(RowHandle Row, const FName& Identifer);
 	template<UE::Editor::DataStorage::TDataColumnType ColumnTypeTemplate>
-	const ColumnTypeTemplate* GetColumn(TypedElementRowHandle Row, const FName& Identifer) const;
+	const ColumnTypeTemplate* GetColumn(RowHandle Row, const FName& Identifer) const;
 	
 	/** Determines if the provided row contains the collection of columns and tags. */
-	virtual bool HasColumns(TypedElementRowHandle Row, TConstArrayView<const UScriptStruct*> ColumnTypes) const = 0;
-	virtual bool HasColumns(TypedElementRowHandle Row, TConstArrayView<TWeakObjectPtr<const UScriptStruct>> ColumnTypes) const = 0;
+	virtual bool HasColumns(RowHandle Row, TConstArrayView<const UScriptStruct*> ColumnTypes) const = 0;
+	virtual bool HasColumns(RowHandle Row, TConstArrayView<TWeakObjectPtr<const UScriptStruct>> ColumnTypes) const = 0;
 	template<UE::Editor::DataStorage::TColumnType... ColumnTypes>
-	bool HasColumns(TypedElementRowHandle Row) const;
+	bool HasColumns(RowHandle Row) const;
 
 	/** Lists the columns on a row. This includes data and tag columns. */
-	virtual void ListColumns(TypedElementDataStorage::RowHandle Row, ColumnListCallbackRef Callback) const = 0;
+	virtual void ListColumns(RowHandle Row, ColumnListCallbackRef Callback) const = 0;
 
 	/** 
 	 * Lists the column type and data on a row. This includes data and tag columns. Not all columns may have data so the data pointer in 
 	 * the callback can be null.
 	 */
-	virtual void ListColumns(TypedElementDataStorage::RowHandle Row, ColumnListWithDataCallbackRef Callback) = 0;
+	virtual void ListColumns(RowHandle Row, ColumnListWithDataCallbackRef Callback) = 0;
 
 	/** Determines if the columns in the row match the query conditions. */
-	virtual bool MatchesColumns(TypedElementDataStorage::RowHandle Row, const TypedElementDataStorage::FQueryConditions& Conditions) const = 0;
+	virtual bool MatchesColumns(RowHandle Row, const TypedElementDataStorage::FQueryConditions& Conditions) const = 0;
 
 	/**
 	 * Finds the type information for a dynamic column.
@@ -356,11 +352,11 @@ public:
 	 * could be created an invalid query handle will be returned. It's recommended to use the Query Builder for a more convenient
 	 * and safer construction of a query.
 	 */
-	virtual TypedElementQueryHandle RegisterQuery(FQueryDescription&& Query) = 0;
+	virtual QueryHandle RegisterQuery(FQueryDescription&& Query) = 0;
 	/** Removes a previous registered. If the query handle is invalid or the query has already been deleted nothing will happen. */
-	virtual void UnregisterQuery(TypedElementQueryHandle Query) = 0;
+	virtual void UnregisterQuery(QueryHandle Query) = 0;
 	/** Returns the description of a previously registered query. If the query no longer exists an empty description will be returned. */
-	virtual const FQueryDescription& GetQueryDescription(TypedElementQueryHandle Query) const = 0;
+	virtual const FQueryDescription& GetQueryDescription(QueryHandle Query) const = 0;
 	/**
 	 * Tick groups for queries can be given any name and the Data Storage will figure out the order of execution based on found
 	 * dependencies. However keeping processors within the same query group can help promote better performance through parallelization.
@@ -368,17 +364,17 @@ public:
 	 */
 	virtual FName GetQueryTickGroupName(EQueryTickGroups Group) const = 0;
 	/** Directly runs a query. If the query handle is invalid or has been deleted nothing will happen. */
-	virtual FQueryResult RunQuery(TypedElementQueryHandle Query) = 0;
+	virtual FQueryResult RunQuery(QueryHandle Query) = 0;
 	/**
 	 * Directly runs a query. The callback will be called for batches of matching rows. During a single call to RunQuery the callback
 	 * may be called multiple times. If the query handle is invalid or has been deleted nothing happens and the callback won't be called.
 	 */
-	virtual FQueryResult RunQuery(TypedElementQueryHandle Query, DirectQueryCallbackRef Callback) = 0;
+	virtual FQueryResult RunQuery(QueryHandle Query, DirectQueryCallbackRef Callback) = 0;
 	/**
 	 * Directly runs a query. The callback will be called for batches of matching rows. During a single call to RunQuery the callback
 	 * may be called multiple times. If the query handle is invalid or has been deleted nothing happens and the callback won't be called.
 	 */
-	virtual FQueryResult RunQuery(TypedElementQueryHandle Query, TypedElementDataStorage::EDirectQueryExecutionFlags Flags, 
+	virtual FQueryResult RunQuery(QueryHandle Query, TypedElementDataStorage::EDirectQueryExecutionFlags Flags, 
 		DirectQueryCallbackRef Callback) = 0;
 	/**
 	 * Triggers all queries registered under the activation name to run for one update cycle. The activatable queries will be activated at
@@ -396,21 +392,21 @@ public:
 	 */
 
 	/** Retrieves the row for an indexed object. Returns an invalid row handle if the hash wasn't found. */
-	virtual TypedElementDataStorage::RowHandle FindIndexedRow(TypedElementDataStorage::IndexHash Index) const = 0;
+	virtual RowHandle FindIndexedRow(TypedElementDataStorage::IndexHash Index) const = 0;
 	/** 
 	 * Registers a row under the index hash. The same row can be registered multiple, but an index hash can only be associated 
 	 * with a single row.
 	 */
-	virtual void IndexRow(TypedElementDataStorage::IndexHash Index, TypedElementDataStorage::RowHandle Row) = 0;
+	virtual void IndexRow(TypedElementDataStorage::IndexHash Index, RowHandle Row) = 0;
 	/**
 	 * Register multiple rows under their index hash. The same row can be registered multiple times,
 	 * but an index hash can only be associated with a single row.
 	 */
 	virtual void BatchIndexRows(
-		TConstArrayView<TPair<TypedElementDataStorage::IndexHash, TypedElementDataStorage::RowHandle>> IndexRowPairs) = 0;
+		TConstArrayView<TPair<TypedElementDataStorage::IndexHash, RowHandle>> IndexRowPairs) = 0;
 	/** Updates the index of a row to a new value. Effectively this is the same as removing an index and adding a new one. */
 	virtual void ReindexRow(
-		TypedElementDataStorage::IndexHash OriginalIndex, TypedElementDataStorage::IndexHash NewIndex, TypedElementDataStorage::RowHandle Row) = 0;
+		TypedElementDataStorage::IndexHash OriginalIndex, TypedElementDataStorage::IndexHash NewIndex, RowHandle Row) = 0;
 	/** Removes a previously registered index hash from the index lookup table or does nothing if the hash no longer exists. */
 	virtual void RemoveIndex(TypedElementDataStorage::IndexHash Index) = 0;
 
@@ -456,44 +452,44 @@ const FactoryT* ITypedElementDataStorageInterface::FindFactory() const
 }
 
 template<UE::Editor::DataStorage::TColumnType... Columns>
-TypedElementDataStorage::TableHandle ITypedElementDataStorageInterface::RegisterTable(const FName Name)
+UE::Editor::DataStorage::TableHandle ITypedElementDataStorageInterface::RegisterTable(const FName Name)
 {
 	return RegisterTable({ Columns::StaticStruct()... }, Name);
 }
 
 template<UE::Editor::DataStorage::TColumnType... Columns>
-TypedElementDataStorage::TableHandle ITypedElementDataStorageInterface::RegisterTable(
-	TypedElementDataStorage::TableHandle SourceTable, const FName Name)
+UE::Editor::DataStorage::TableHandle ITypedElementDataStorageInterface::RegisterTable(
+	TableHandle SourceTable, const FName Name)
 {
 	return RegisterTable(SourceTable, { Columns::StaticStruct()... }, Name);
 }
 
 template<UE::Editor::DataStorage::TColumnType Column>
-void ITypedElementDataStorageInterface::AddColumn(TypedElementRowHandle Row)
+void ITypedElementDataStorageInterface::AddColumn(RowHandle Row)
 {
 	AddColumn(Row, Column::StaticStruct());
 }
 
 template<UE::Editor::DataStorage::TColumnType Column>
-void ITypedElementDataStorageInterface::RemoveColumn(TypedElementRowHandle Row)
+void ITypedElementDataStorageInterface::RemoveColumn(RowHandle Row)
 {
 	RemoveColumn(Row, Column::StaticStruct());
 }
 
 template<UE::Editor::DataStorage::TColumnType... Columns>
-void ITypedElementDataStorageInterface::AddColumns(TypedElementRowHandle Row)
+void ITypedElementDataStorageInterface::AddColumns(RowHandle Row)
 {
 	AddColumns(Row, { Columns::StaticStruct()...});
 }
 
 template <>
-inline void ITypedElementDataStorageInterface::AddColumn<UE::Editor::DataStorage::FDynamicTag>(TypedElementDataStorage::RowHandle Row, const FName& Tag, const FName& Value)
+inline void ITypedElementDataStorageInterface::AddColumn<UE::Editor::DataStorage::FDynamicTag>(RowHandle Row, const FName& Tag, const FName& Value)
 {
 	AddColumn(Row, UE::Editor::DataStorage::FDynamicTag(Tag), Value);
 }
 
 template <>
-inline void ITypedElementDataStorageInterface::RemoveColumn<UE::Editor::DataStorage::FDynamicTag>(TypedElementDataStorage::RowHandle Row, const FName& Tag)
+inline void ITypedElementDataStorageInterface::RemoveColumn<UE::Editor::DataStorage::FDynamicTag>(RowHandle Row, const FName& Tag)
 {
 	using namespace UE::Editor::DataStorage;
 	RemoveColumn(Row, FDynamicTag(Tag));
@@ -512,7 +508,7 @@ void ITypedElementDataStorageInterface::RemoveColumn(TypedElementDataStorage::Ro
 }
 
 template<UE::Editor::DataStorage::TEnumType EnumT>
-void ITypedElementDataStorageInterface::AddColumn(TypedElementDataStorage::RowHandle Row, EnumT Value)
+void ITypedElementDataStorageInterface::AddColumn(RowHandle Row, EnumT Value)
 {
 	const UEnum* Enum = StaticEnum<EnumT>();
 	const FName ValueAsFName = *Enum->GetNameStringByValue(static_cast<int64>(Value));
@@ -523,7 +519,7 @@ void ITypedElementDataStorageInterface::AddColumn(TypedElementDataStorage::RowHa
 }
 
 template<auto Value, UE::Editor::DataStorage::TEnumType EnumT>
-void ITypedElementDataStorageInterface::AddColumn(TypedElementDataStorage::RowHandle Row)
+void ITypedElementDataStorageInterface::AddColumn(RowHandle Row)
 {
 	AddColumn<EnumT>(Row, Value);
 }
@@ -580,20 +576,20 @@ void ITypedElementDataStorageInterface::AddColumn(TypedElementDataStorage::RowHa
 }
 
 template<UE::Editor::DataStorage::TEnumType EnumT>
-void ITypedElementDataStorageInterface::RemoveColumn(TypedElementDataStorage::RowHandle Row)
+void ITypedElementDataStorageInterface::RemoveColumn(RowHandle Row)
 {
 	const UEnum* Enum = StaticEnum<EnumT>();
 	RemoveColumn(Row, UE::Editor::DataStorage::FDynamicTag(Enum->GetFName()));
 }
 
 template<UE::Editor::DataStorage::TColumnType... Columns>
-void ITypedElementDataStorageInterface::RemoveColumns(TypedElementRowHandle Row)
+void ITypedElementDataStorageInterface::RemoveColumns(RowHandle Row)
 {
 	RemoveColumns(Row, { Columns::StaticStruct()...});
 }
 
 template<UE::Editor::DataStorage::TDataColumnType ColumnType>
-void ITypedElementDataStorageInterface::AddColumn(TypedElementRowHandle Row, ColumnType&& Column)
+void ITypedElementDataStorageInterface::AddColumn(RowHandle Row, ColumnType&& Column)
 {
 	AddColumnData(Row, ColumnType::StaticStruct(),
 		[&Column](void* ColumnData, const UScriptStruct&)
@@ -621,19 +617,19 @@ void ITypedElementDataStorageInterface::AddColumn(TypedElementRowHandle Row, Col
 }
 
 template<UE::Editor::DataStorage::TDataColumnType ColumnType>
-ColumnType* ITypedElementDataStorageInterface::GetColumn(TypedElementRowHandle Row)
+ColumnType* ITypedElementDataStorageInterface::GetColumn(RowHandle Row)
 {
 	return reinterpret_cast<ColumnType*>(GetColumnData(Row, ColumnType::StaticStruct()));
 }
 
 template<UE::Editor::DataStorage::TDataColumnType ColumnType>
-const ColumnType* ITypedElementDataStorageInterface::GetColumn(TypedElementRowHandle Row) const
+const ColumnType* ITypedElementDataStorageInterface::GetColumn(RowHandle Row) const
 {
 	return reinterpret_cast<const ColumnType*>(GetColumnData(Row, ColumnType::StaticStruct()));
 }
 
 template <UE::Editor::DataStorage::TDataColumnType DynamicColumnTemplate>
-DynamicColumnTemplate* ITypedElementDataStorageInterface::GetColumn(TypedElementRowHandle Row, const FName& Identifier)
+DynamicColumnTemplate* ITypedElementDataStorageInterface::GetColumn(RowHandle Row, const FName& Identifier)
 {
 	const UE::Editor::DataStorage::FDynamicColumnDescription Description
 	{
@@ -649,7 +645,7 @@ DynamicColumnTemplate* ITypedElementDataStorageInterface::GetColumn(TypedElement
 }
 
 template <UE::Editor::DataStorage::TDataColumnType DynamicColumnTemplate>
-const DynamicColumnTemplate* ITypedElementDataStorageInterface::GetColumn(TypedElementRowHandle Row, const FName& Identifier) const
+const DynamicColumnTemplate* ITypedElementDataStorageInterface::GetColumn(RowHandle Row, const FName& Identifier) const
 {
 	const UE::Editor::DataStorage::FDynamicColumnDescription Description
 	{
@@ -665,7 +661,7 @@ const DynamicColumnTemplate* ITypedElementDataStorageInterface::GetColumn(TypedE
 }
 
 template<UE::Editor::DataStorage::TColumnType... ColumnType>
-bool ITypedElementDataStorageInterface::HasColumns(TypedElementRowHandle Row) const
+bool ITypedElementDataStorageInterface::HasColumns(RowHandle Row) const
 {
 	return HasColumns(Row, TConstArrayView<const UScriptStruct*>({ ColumnType::StaticStruct()... }));
 }
diff --git a/Engine/Source/Runtime/TypedElementFramework/Public/Elements/Interfaces/TypedElementDataStorageUiInterface.h b/Engine/Source/Runtime/TypedElementFramework/Public/Elements/Interfaces/TypedElementDataStorageUiInterface.h
--- a/Engine/Source/Runtime/TypedElementFramework/Public/Elements/Interfaces/TypedElementDataStorageUiInterface.h
+++ b/Engine/Source/Runtime/TypedElementFramework/Public/Elements/Interfaces/TypedElementDataStorageUiInterface.h
@@ -25,6 +25,8 @@ struct FTypedElementWidgetConstructor
 {
 	GENERATED_BODY()
 
+	using RowHandle = UE::Editor::DataStorage::RowHandle;
+
 public:
 	TYPEDELEMENTFRAMEWORK_API explicit FTypedElementWidgetConstructor(const UScriptStruct* InTypeInfo);
 	explicit FTypedElementWidgetConstructor(EForceInit) {} //< For compatibility and shouldn't be directly used.
@@ -63,7 +65,7 @@ public:
 	 *	@see Construct
 	 */
 	TYPEDELEMENTFRAMEWORK_API TSharedPtr<SWidget> ConstructFinalWidget(
-		TypedElementRowHandle Row, /** The row the widget will be stored in. */
+		RowHandle Row, /** The row the widget will be stored in. */
 		ITypedElementDataStorageInterface* DataStorage,
 		ITypedElementDataStorageUiInterface* DataStorageUi,
 		const TypedElementDataStorage::FMetaDataView& Arguments);
@@ -76,7 +78,7 @@ public:
 	 * In most cases, you want to call ConstructFinalWidget to create the actual widget.
 	 */
 	TYPEDELEMENTFRAMEWORK_API virtual TSharedPtr<SWidget> Construct(
-		TypedElementRowHandle Row, /** The row the widget will be stored in. */
+		RowHandle Row, /** The row the widget will be stored in. */
 		ITypedElementDataStorageInterface* DataStorage,
 		ITypedElementDataStorageUiInterface* DataStorageUi,
 		const TypedElementDataStorage::FMetaDataView& Arguments);
@@ -85,7 +87,7 @@ protected:
 	/** Create a new instance of the target widget. This is a required function. */
 	TYPEDELEMENTFRAMEWORK_API virtual TSharedPtr<SWidget> CreateWidget(const TypedElementDataStorage::FMetaDataView& Arguments);
 	/** Set any values in columns if needed. The columns provided through GetAdditionalColumnsList() will have already been created. */
-	TYPEDELEMENTFRAMEWORK_API virtual bool SetColumns(ITypedElementDataStorageInterface* DataStorage, TypedElementRowHandle Row);
+	TYPEDELEMENTFRAMEWORK_API virtual bool SetColumns(ITypedElementDataStorageInterface* DataStorage, RowHandle Row);
 	
 	/** Creates a (friendly) name for the provided column type. */
 	TYPEDELEMENTFRAMEWORK_API virtual FString DescribeColumnType(const UScriptStruct* ColumnType) const;
@@ -97,11 +99,11 @@ protected:
 	TYPEDELEMENTFRAMEWORK_API virtual bool FinalizeWidget(
 		ITypedElementDataStorageInterface* DataStorage,
 		ITypedElementDataStorageUiInterface* DataStorageUi,
-		TypedElementRowHandle Row,
+		RowHandle Row,
 		const TSharedPtr<SWidget>& Widget);
 
 	/** Add the default misc columns we want a widget row to have. */
-	TYPEDELEMENTFRAMEWORK_API void AddDefaultWidgetColumns(TypedElementRowHandle Row, ITypedElementDataStorageInterface* DataStorage) const;
+	TYPEDELEMENTFRAMEWORK_API void AddDefaultWidgetColumns(RowHandle Row, ITypedElementDataStorageInterface* DataStorage) const;
 protected:
 
 	TArray<TWeakObjectPtr<const UScriptStruct>> MatchedColumnTypes;
@@ -129,6 +131,8 @@ class ITypedElementDataStorageUiInterface
 {
 	GENERATED_BODY()
 
+	using RowHandle = UE::Editor::DataStorage::RowHandle;
+
 public:
 	enum class EPurposeType : uint8
 	{
@@ -159,7 +163,7 @@ public:
 		SingleMatch
 	};
 
-	using WidgetCreatedCallback = TFunctionRef<void(const TSharedRef<SWidget>& NewWidget, TypedElementRowHandle Row)>;
+	using WidgetCreatedCallback = TFunctionRef<void(const TSharedRef<SWidget>& NewWidget, RowHandle Row)>;
 	using WidgetConstructorCallback = TFunctionRef<bool(TUniquePtr<FTypedElementWidgetConstructor>, TConstArrayView<TWeakObjectPtr<const UScriptStruct>>)>;
 	using WidgetPurposeCallback = TFunctionRef<void(FName, EPurposeType, const FText&)>;
 
@@ -258,7 +262,7 @@ public:
 	 * The provided row will be used to store the widget information on. If columns have already been added to the row, the 
 	 * constructor is free to use that to configure the widget. Arguments are used by the constructor to configure the widget.
 	 */
-	virtual TSharedPtr<SWidget> ConstructWidget(TypedElementRowHandle Row, FTypedElementWidgetConstructor& Constructor,
+	virtual TSharedPtr<SWidget> ConstructWidget(RowHandle Row, FTypedElementWidgetConstructor& Constructor,
 		const TypedElementDataStorage::FMetaDataView& Arguments) = 0;
 
 	/** Calls the provided callback for all known registered widget purposes. */
diff --git a/Engine/Source/Runtime/TypedElementFramework/Tests/Elements/Framework/TypedElementDataStoragePerformanceTestCommands.cpp b/Engine/Source/Runtime/TypedElementFramework/Tests/Elements/Framework/TypedElementDataStoragePerformanceTestCommands.cpp
--- a/Engine/Source/Runtime/TypedElementFramework/Tests/Elements/Framework/TypedElementDataStoragePerformanceTestCommands.cpp
+++ b/Engine/Source/Runtime/TypedElementFramework/Tests/Elements/Framework/TypedElementDataStoragePerformanceTestCommands.cpp
@@ -7,9 +7,12 @@
 #include "Elements/Framework/TypedElementRegistry.h"
 #include "HAL/IConsoleManager.h"
 
-namespace Private
+namespace UE::Editor::DataStorage
 {
-	TypedElementTableHandle PerformanceTestCommandTable = TypedElementInvalidTableHandle;
+	namespace Private
+	{
+		TableHandle PerformanceTestCommandTable = InvalidTableHandle;
+	}
 }
 
 FAutoConsoleCommand CVarAddDebugRows = FAutoConsoleCommand(
@@ -24,7 +27,7 @@ FAutoConsoleCommand CVarAddDebugRows = FAutoConsoleCommand(
 
 			ITypedElementDataStorageInterface* DataStorageInterface = UTypedElementRegistry::GetInstance()->GetMutableDataStorage();
 			
-			DataStorageInterface->BatchAddRow(Private::PerformanceTestCommandTable, EntitiesToAdd, [DataStorageInterface](TypedElementDataStorage::RowHandle Row)
+			DataStorageInterface->BatchAddRow(UE::Editor::DataStorage::Private::PerformanceTestCommandTable, EntitiesToAdd, [DataStorageInterface](TypedElementDataStorage::RowHandle Row)
 				{
 					FTest_PingPongPrePhys* Column = DataStorageInterface->GetColumn<FTest_PingPongPrePhys>(Row);
 					Column->Value = 0;
@@ -67,7 +70,7 @@ FAutoConsoleCommand CVarResetDebugEntities = FAutoConsoleCommand(
 
 void UTest_PingPongBetweenPhaseFactory::RegisterTables(ITypedElementDataStorageInterface& DataStorage)
 {
-	Private::PerformanceTestCommandTable = DataStorage.RegisterTable({FTest_PingPongPrePhys::StaticStruct()}, TEXT("Test_PingPongPrePhys"));
+	UE::Editor::DataStorage::Private::PerformanceTestCommandTable = DataStorage.RegisterTable({FTest_PingPongPrePhys::StaticStruct()}, TEXT("Test_PingPongPrePhys"));
 }
 
 // Performance test is a small benchmark to better characterize the performance of adding and removing columns during the processor phases
@@ -86,10 +89,10 @@ void UTest_PingPongBetweenPhaseFactory::RegisterQueries(ITypedElementDataStorage
 		Select(TEXT("PingPong PrePhysics->DurPhysics"),
 			FProcessor(EQueryTickPhase::PrePhysics, DataStorage.GetQueryTickGroupName(EQueryTickGroups::Default))
 				.SetExecutionMode(EExecutionMode::GameThread),
-			[](IQueryContext& Context, const TypedElementRowHandle* RowPtr, FTest_PingPongPrePhys* PingPongAPtr)
+			[](IQueryContext& Context, const RowHandle* RowPtr, FTest_PingPongPrePhys* PingPongAPtr)
 			{
 				QUICK_SCOPE_CYCLE_COUNTER(PingPong_Pre_During);
-				TArrayView<const TypedElementRowHandle> Rows(RowPtr, Context.GetRowCount());
+				TArrayView<const RowHandle> Rows(RowPtr, Context.GetRowCount());
 				TArrayView<FTest_PingPongPrePhys> PingPongAs(PingPongAPtr, Context.GetRowCount());
 
 				UScriptStruct* RemovedColumnStruct = FTest_PingPongPrePhys::StaticStruct();
@@ -108,10 +111,10 @@ void UTest_PingPongBetweenPhaseFactory::RegisterQueries(ITypedElementDataStorage
 	Select(TEXT("PingPong DurPhysics->PostPhysics"),
 		FProcessor(EQueryTickPhase::DuringPhysics, DataStorage.GetQueryTickGroupName(EQueryTickGroups::Default))
 			.SetExecutionMode(EExecutionMode::GameThread),
-		[](IQueryContext& Context, const TypedElementRowHandle* RowPtr, FTest_PingPongDurPhys* PingPongBPtr)
+		[](IQueryContext& Context, const RowHandle* RowPtr, FTest_PingPongDurPhys* PingPongBPtr)
 		{
 			QUICK_SCOPE_CYCLE_COUNTER(PingPong_During_Post);
-			TArrayView<const TypedElementRowHandle> Rows(RowPtr, Context.GetRowCount());
+			TArrayView<const RowHandle> Rows(RowPtr, Context.GetRowCount());
 			TArrayView<FTest_PingPongDurPhys> PingPongAs(PingPongBPtr, Context.GetRowCount());
 
 			UScriptStruct* RemovedColumnStruct = FTest_PingPongDurPhys::StaticStruct();
@@ -130,10 +133,10 @@ void UTest_PingPongBetweenPhaseFactory::RegisterQueries(ITypedElementDataStorage
 	Select(TEXT("PingPong PostPhysics->PrePhysics"),
 	FProcessor(EQueryTickPhase::PostPhysics, DataStorage.GetQueryTickGroupName(EQueryTickGroups::Default))
 		.SetExecutionMode(EExecutionMode::GameThread),
-	[](IQueryContext& Context, const TypedElementRowHandle* RowPtr, FTest_PingPongPostPhys* PingPongBPtr)
+	[](IQueryContext& Context, const RowHandle* RowPtr, FTest_PingPongPostPhys* PingPongBPtr)
 	{
 		QUICK_SCOPE_CYCLE_COUNTER(PingPong_Post_Pre);
-		TArrayView<const TypedElementRowHandle> Rows(RowPtr, Context.GetRowCount());
+		TArrayView<const RowHandle> Rows(RowPtr, Context.GetRowCount());
 		TArrayView<FTest_PingPongPostPhys> PingPongAs(PingPongBPtr, Context.GetRowCount());
 
 		UScriptStruct* RemovedColumnStruct = FTest_PingPongPostPhys::StaticStruct();
