[Teds] Update naming for RowHandle, QueryHandle, and TableHandle

diff --git a/Engine/Source/Runtime/TypedElementFramework/Tests/Elements/Framework/TypedElementDataStorageProcessorTests.cpp b/Engine/Source/Runtime/TypedElementFramework/Tests/Elements/Framework/TypedElementDataStorageProcessorTests.cpp
--- a/Engine/Source/Runtime/TypedElementFramework/Tests/Elements/Framework/TypedElementDataStorageProcessorTests.cpp
+++ b/Engine/Source/Runtime/TypedElementFramework/Tests/Elements/Framework/TypedElementDataStorageProcessorTests.cpp
@@ -1,183 +1,186 @@
-﻿// Copyright Epic Games, Inc. All Rights Reserved.
+// Copyright Epic Games, Inc. All Rights Reserved.
 
 #if WITH_TESTS
 #include "Elements/Framework/TypedElementRegistry.h"
 #include "HAL/IConsoleManager.h"
 #include "TypedElementTestColumns.h"
 
-namespace TEDS::Debug::ProcessorTests
+namespace UE::Editor::DataStorage
 {
-	void OnProcessorTestsEnabled(IConsoleVariable*);
-
-	bool GProcessorTestsEnabled = false;
-	FAutoConsoleVariableRef CVarAllowUnversionedContentInEditor(
-		TEXT("TEDS.Tests.ProcessorTestsEnabled"),
-		GProcessorTestsEnabled,
-		TEXT("If true, registers processors and additional commands with TEDS to test processors."),
-		FConsoleVariableDelegate::CreateStatic(&OnProcessorTestsEnabled),
-		ECVF_Default
-	);
+	namespace Debug::ProcessorTests
+	{
+		void OnProcessorTestsEnabled(IConsoleVariable*);
 
-	TypedElementDataStorage::TableHandle PrimaryTable = TypedElementDataStorage::InvalidTableHandle;
-	TypedElementDataStorage::TableHandle SecondaryTable = TypedElementDataStorage::InvalidTableHandle;
+		bool GProcessorTestsEnabled = false;
+		FAutoConsoleVariableRef CVarAllowUnversionedContentInEditor(
+			TEXT("TEDS.Tests.ProcessorTestsEnabled"),
+			GProcessorTestsEnabled,
+			TEXT("If true, registers processors and additional commands with TEDS to test processors."),
+			FConsoleVariableDelegate::CreateStatic(&OnProcessorTestsEnabled),
+			ECVF_Default
+		);
 
-	TArray<TypedElementQueryHandle> RegisteredQueries;
-	TArray<IConsoleCommand*> RegisteredCommands;
+		TableHandle PrimaryTable = InvalidTableHandle;
+		TableHandle SecondaryTable = InvalidTableHandle;
 
-	void RegisterProcessors()
-	{
-		ITypedElementDataStorageInterface* DataStorage = UTypedElementRegistry::GetInstance()->GetMutableDataStorage();
-		using namespace TypedElementQueryBuilder;
-		using DSI = ITypedElementDataStorageInterface;
-		namespace DS = TypedElementDataStorage;
+		TArray<QueryHandle> RegisteredQueries;
+		TArray<IConsoleCommand*> RegisteredCommands;
 
-		if (PrimaryTable == TypedElementDataStorage::InvalidTableHandle)
+		void RegisterProcessors()
 		{
-			PrimaryTable =
-				DataStorage->RegisterTable<FTEDSProcessorTestsReferenceColumn, FTEDSProcessorTests_PrimaryTag>(FName(TEXT("ProcessorTests Primary Table")));
-		}
-		if (SecondaryTable == TypedElementDataStorage::InvalidTableHandle)
-		{
-			SecondaryTable =
-				DataStorage->RegisterTable<FTEDSProcessorTestsReferenceColumn, FTEDSProcessorTests_SecondaryTag>(FName(TEXT("ProcessorTests Secondary Table")));
-		}
+			ITypedElementDataStorageInterface* DataStorage = UTypedElementRegistry::GetInstance()->GetMutableDataStorage();
+			using namespace TypedElementQueryBuilder;
+			using DSI = ITypedElementDataStorageInterface;
+			namespace DS = TypedElementDataStorage;
 
-		// Test creation of a row from within a query processor
-		TypedElementQueryHandle PrimaryRowQuery = DataStorage->RegisterQuery(
-			Select(TEXT("TEST: Creating a row for primary reference column"),
-			FProcessor(DSI::EQueryTickPhase::PrePhysics, DataStorage->GetQueryTickGroupName(DSI::EQueryTickGroups::Default)),
-			[](DS::IQueryContext& Context, const TypedElementRowHandle* Rows, FTEDSProcessorTestsReferenceColumn* ReferenceColumns)
+			if (PrimaryTable == InvalidTableHandle)
 			{
-				const int32 RowCount = Context.GetRowCount();
-				TConstArrayView<TypedElementRowHandle> RowsView = MakeArrayView(Rows, RowCount);
-				TArrayView<FTEDSProcessorTestsReferenceColumn> ReferenceColumnsView = MakeArrayView(ReferenceColumns, RowCount);
+				PrimaryTable =
+					DataStorage->RegisterTable<FTEDSProcessorTestsReferenceColumn, FTEDSProcessorTests_PrimaryTag>(FName(TEXT("ProcessorTests Primary Table")));
+			}
+			if (SecondaryTable == InvalidTableHandle)
+			{
+				SecondaryTable =
+					DataStorage->RegisterTable<FTEDSProcessorTestsReferenceColumn, FTEDSProcessorTests_SecondaryTag>(FName(TEXT("ProcessorTests Secondary Table")));
+			}
 
-				for (int32 Index = 0; Index < RowCount; ++Index)
+			// Test creation of a row from within a query processor
+			QueryHandle PrimaryRowQuery = DataStorage->RegisterQuery(
+				Select(TEXT("TEST: Creating a row for primary reference column"),
+				FProcessor(DSI::EQueryTickPhase::PrePhysics, DataStorage->GetQueryTickGroupName(DSI::EQueryTickGroups::Default)),
+				[](DS::IQueryContext& Context, const RowHandle* Rows, FTEDSProcessorTestsReferenceColumn* ReferenceColumns)
 				{
-					FTEDSProcessorTestsReferenceColumn& PrimaryReferenceColumn = ReferenceColumnsView[Index];
-					// Auto-create a secondary row if this points to no row
-					if (!Context.IsRowAvailable(PrimaryReferenceColumn.Reference))
+					const int32 RowCount = Context.GetRowCount();
+					TConstArrayView<RowHandle> RowsView = MakeArrayView(Rows, RowCount);
+					TArrayView<FTEDSProcessorTestsReferenceColumn> ReferenceColumnsView = MakeArrayView(ReferenceColumns, RowCount);
+
+					for (int32 Index = 0; Index < RowCount; ++Index)
 					{
-						TypedElementDataStorage::RowHandle SecondaryRow = Context.AddRow(SecondaryTable);
-						// Initialize bi-directional row references
-						PrimaryReferenceColumn.Reference = SecondaryRow;
-						Context.AddColumn(SecondaryRow, FTEDSProcessorTestsReferenceColumn{.Reference = RowsView[Index]});
+						FTEDSProcessorTestsReferenceColumn& PrimaryReferenceColumn = ReferenceColumnsView[Index];
+						// Auto-create a secondary row if this points to no row
+						if (!Context.IsRowAvailable(PrimaryReferenceColumn.Reference))
+						{
+							RowHandle SecondaryRow = Context.AddRow(SecondaryTable);
+							// Initialize bi-directional row references
+							PrimaryReferenceColumn.Reference = SecondaryRow;
+							Context.AddColumn(SecondaryRow, FTEDSProcessorTestsReferenceColumn{.Reference = RowsView[Index]});
+						}
 					}
 				}
-			}
-			)
-			.Where()
-				.All<FTEDSProcessorTests_PrimaryTag>()
-				.None<FTEDSProcessorTests_Linked>()
-			.Compile());
-
-		RegisteredQueries.Emplace(PrimaryRowQuery);
-
-		TypedElementQueryHandle UpdateTransformWidget = DataStorage->RegisterQuery(
-			Select()
-				.ReadOnly<FTEDSProcessorTestsReferenceColumn>()
-			.Where()
-				.All<FTEDSProcessorTests_PrimaryTag>()
-				.None<FTEDSProcessorTests_Linked>()
-			.Compile());
+				)
+				.Where()
+					.All<FTEDSProcessorTests_PrimaryTag>()
+					.None<FTEDSProcessorTests_Linked>()
+				.Compile());
+
+			RegisteredQueries.Emplace(PrimaryRowQuery);
+
+			QueryHandle UpdateTransformWidget = DataStorage->RegisterQuery(
+				Select()
+					.ReadOnly<FTEDSProcessorTestsReferenceColumn>()
+				.Where()
+					.All<FTEDSProcessorTests_PrimaryTag>()
+					.None<FTEDSProcessorTests_Linked>()
+				.Compile());
 		
-		RegisteredQueries.Emplace(UpdateTransformWidget);
-
-		TypedElementQueryHandle SecondaryRowQuery = DataStorage->RegisterQuery(
-			Select(TEXT("TEST: Creating a row for secondary reference column"),
-			FProcessor(DSI::EQueryTickPhase::DuringPhysics, DataStorage->GetQueryTickGroupName(DSI::EQueryTickGroups::Default)),
-			[](DS::IQueryContext& Context, const TypedElementRowHandle* Rows, FTEDSProcessorTestsReferenceColumn* ReferenceColumns)
-			{
-				const int32 RowCount = Context.GetRowCount();
-				TConstArrayView<TypedElementRowHandle> RowsView = MakeArrayView(Rows, RowCount);
-				TArrayView<FTEDSProcessorTestsReferenceColumn> ReferenceColumnsView = MakeArrayView(ReferenceColumns, RowCount);
+			RegisteredQueries.Emplace(UpdateTransformWidget);
 
-				for (int32 Index = 0; Index < RowCount; ++Index)
+			QueryHandle SecondaryRowQuery = DataStorage->RegisterQuery(
+				Select(TEXT("TEST: Creating a row for secondary reference column"),
+				FProcessor(DSI::EQueryTickPhase::DuringPhysics, DataStorage->GetQueryTickGroupName(DSI::EQueryTickGroups::Default)),
+				[](DS::IQueryContext& Context, const RowHandle* Rows, FTEDSProcessorTestsReferenceColumn* ReferenceColumns)
 				{
-					const TypedElementRowHandle& SecondaryRow = RowsView[Index];
-					Context.RunSubquery(
-						0,
-						ReferenceColumnsView[Index].Reference,
-						CreateSubqueryCallbackBinding(
-						[SecondaryRow](TypedElementDataStorage::ISubqueryContext& SubqueryContext, TypedElementDataStorage::RowHandle PrimaryRow, const FTEDSProcessorTestsReferenceColumn& ReferenceColumn)
-						{
-							if (ReferenceColumn.Reference == SecondaryRow)
+					const int32 RowCount = Context.GetRowCount();
+					TConstArrayView<RowHandle> RowsView = MakeArrayView(Rows, RowCount);
+					TArrayView<FTEDSProcessorTestsReferenceColumn> ReferenceColumnsView = MakeArrayView(ReferenceColumns, RowCount);
+
+					for (int32 Index = 0; Index < RowCount; ++Index)
+					{
+						const RowHandle& SecondaryRow = RowsView[Index];
+						Context.RunSubquery(
+							0,
+							ReferenceColumnsView[Index].Reference,
+							CreateSubqueryCallbackBinding(
+							[SecondaryRow](TypedElementDataStorage::ISubqueryContext& SubqueryContext, RowHandle PrimaryRow, const FTEDSProcessorTestsReferenceColumn& ReferenceColumn)
 							{
-								// Add these tags to prevent further processing
-								SubqueryContext.AddColumns<FTEDSProcessorTests_Linked>(SecondaryRow);
-								SubqueryContext.AddColumns<FTEDSProcessorTests_Linked>(PrimaryRow);
+								if (ReferenceColumn.Reference == SecondaryRow)
+								{
+									// Add these tags to prevent further processing
+									SubqueryContext.AddColumns<FTEDSProcessorTests_Linked>(SecondaryRow);
+									SubqueryContext.AddColumns<FTEDSProcessorTests_Linked>(PrimaryRow);
+								}
 							}
-						}
-					));
+						));
+					}
 				}
-			}
-			)
-			.Where()
-				.All<FTEDSProcessorTests_SecondaryTag>()
-				.None<FTEDSProcessorTests_Linked>()
-			.DependsOn()
-				.SubQuery(UpdateTransformWidget)
-			.Compile());
-
-		RegisteredQueries.Emplace(SecondaryRowQuery);
-	}
-
-	void UnregisterProcessors()
-	{
-		ITypedElementDataStorageInterface* DataStorage = UTypedElementRegistry::GetInstance()->GetMutableDataStorage();
-		for (const TypedElementQueryHandle& Handle : RegisteredQueries)
-		{
-			DataStorage->UnregisterQuery(Handle);
+				)
+				.Where()
+					.All<FTEDSProcessorTests_SecondaryTag>()
+					.None<FTEDSProcessorTests_Linked>()
+				.DependsOn()
+					.SubQuery(UpdateTransformWidget)
+				.Compile());
+
+			RegisteredQueries.Emplace(SecondaryRowQuery);
 		}
-		RegisteredQueries.Empty();
-	}
 
-	void RegisterCommands()
-	{
-		IConsoleCommand* Command = IConsoleManager::Get().RegisterConsoleCommand(TEXT("TEDS.Tests.ProcessorTests.AddPrimaryRows"), TEXT(""), FConsoleCommandWithArgsDelegate::CreateLambda([](const TArray<FString>& Args)
+		void UnregisterProcessors()
 		{
-			int32 RowsToCreate = 0;
-			if (Args.Num() != 1)
-			{
-				return;
-			}
-			TTypeFromString<int32>::FromString(RowsToCreate, *Args[0]);
-			if (RowsToCreate <= 0)
+			ITypedElementDataStorageInterface* DataStorage = UTypedElementRegistry::GetInstance()->GetMutableDataStorage();
+			for (const UE::Editor::DataStorage::QueryHandle& Handle : RegisteredQueries)
 			{
-				return;
+				DataStorage->UnregisterQuery(Handle);
 			}
+			RegisteredQueries.Empty();
+		}
 
-			ITypedElementDataStorageInterface* DataStorage = UTypedElementRegistry::GetInstance()->GetMutableDataStorage();
-			DataStorage->BatchAddRow(PrimaryTable, RowsToCreate, [](TypedElementRowHandle Row)
+		void RegisterCommands()
+		{
+			IConsoleCommand* Command = IConsoleManager::Get().RegisterConsoleCommand(TEXT("TEDS.Tests.ProcessorTests.AddPrimaryRows"), TEXT(""), FConsoleCommandWithArgsDelegate::CreateLambda([](const TArray<FString>& Args)
 			{
-			});
-		}), ECVF_Default);
+				int32 RowsToCreate = 0;
+				if (Args.Num() != 1)
+				{
+					return;
+				}
+				TTypeFromString<int32>::FromString(RowsToCreate, *Args[0]);
+				if (RowsToCreate <= 0)
+				{
+					return;
+				}
+
+				ITypedElementDataStorageInterface* DataStorage = UTypedElementRegistry::GetInstance()->GetMutableDataStorage();
+				DataStorage->BatchAddRow(PrimaryTable, RowsToCreate, [](UE::Editor::DataStorage::RowHandle Row)
+				{
+				});
+			}), ECVF_Default);
 		
-		RegisteredCommands.Emplace(Command);
-	}
+			RegisteredCommands.Emplace(Command);
+		}
 
-	void UnregisterCommands()
-	{
-		for (IConsoleCommand* Command : RegisteredCommands)
+		void UnregisterCommands()
 		{
-			IConsoleManager::Get().UnregisterConsoleObject(Command);
+			for (IConsoleCommand* Command : RegisteredCommands)
+			{
+				IConsoleManager::Get().UnregisterConsoleObject(Command);
+			}
+			RegisteredCommands.Empty();
 		}
-		RegisteredCommands.Empty();
-	}
 	
-	void OnProcessorTestsEnabled(IConsoleVariable* Variable)
-	{
-		if (Variable->GetBool())
-		{
-			RegisterProcessors();
-			RegisterCommands();
-		}
-		else
+		void OnProcessorTestsEnabled(IConsoleVariable* Variable)
 		{
-			UnregisterProcessors();
-			UnregisterCommands();
-		}
-	};
-}
+			if (Variable->GetBool())
+			{
+				RegisterProcessors();
+				RegisterCommands();
+			}
+			else
+			{
+				UnregisterProcessors();
+				UnregisterCommands();
+			}
+		};
+	} // namespace Debug::ProcessorTests
+} // namespace UE::Editor::DataStorage
 #endif
 
diff --git a/Engine/Source/Runtime/TypedElementFramework/Tests/Elements/Framework/TypedElementQueryBuilderTests.cpp b/Engine/Source/Runtime/TypedElementFramework/Tests/Elements/Framework/TypedElementQueryBuilderTests.cpp
--- a/Engine/Source/Runtime/TypedElementFramework/Tests/Elements/Framework/TypedElementQueryBuilderTests.cpp
+++ b/Engine/Source/Runtime/TypedElementFramework/Tests/Elements/Framework/TypedElementQueryBuilderTests.cpp
@@ -105,169 +105,169 @@ TEST_CASE_NAMED(FTypedElementQueryConditions_NoColumn, "TypedElementQueryBuilder
 	
 	CHECK(Example.MinimumColumnMatchRequired() == 0);
 	// Since there are no restrictions provided in the query, all input passes.
-	CHECK(TestMatching(Example, { FColumn<FTestColumnA>() }, true));
+	CHECK(TestMatching(Example, { TypedElementDataStorage::FColumn<FTestColumnA>() }, true));
 }
 
 TEST_CASE_NAMED(FTypedElementQueryConditions_OneColumn, "TypedElementQueryBuilder::FTypedElementQueryConditions_OneColumn", "[ApplicationContextMask][EngineFilter]")
 {
-	FQueryConditions Example{ FColumn<FTestColumnA>() };
+	FQueryConditions Example{ TypedElementDataStorage::FColumn<FTestColumnA>() };
 
 	CHECK(Example.MinimumColumnMatchRequired() == 1);
-	CHECK(TestMatching(Example, { FColumn<FTestColumnA>() }, true));
+	CHECK(TestMatching(Example, { TypedElementDataStorage::FColumn<FTestColumnA>() }, true));
 }
 
 TEST_CASE_NAMED(FTypedElementQueryConditions1, "TypedElementQueryBuilder::FTypedElementQueryConditions A && B && C", "[ApplicationContextMask][EngineFilter]")
 {
-	FColumn TestA(FTestColumnA::StaticStruct());
+	TypedElementDataStorage::FColumn TestA(FTestColumnA::StaticStruct());
 	
-	FQueryConditions Example = FColumn<FTestColumnA>() && FColumn<FTestColumnB>() && FColumn<FTestColumnC>();
+	FQueryConditions Example = TypedElementDataStorage::FColumn<FTestColumnA>() && TypedElementDataStorage::FColumn<FTestColumnB>() && TypedElementDataStorage::FColumn<FTestColumnC>();
 	
 	CHECK(Example.MinimumColumnMatchRequired() == 3);
-	CHECK(TestMatching(Example, { FColumn<FTestColumnA>(), FColumn<FTestColumnB>(), FColumn<FTestColumnC>() }, true));
-	CHECK(TestMatching(Example, { FColumn<FTestColumnA>(), FColumn<FTestColumnB>(), FColumn<FTestColumnD>() }, false));
+	CHECK(TestMatching(Example, { TypedElementDataStorage::FColumn<FTestColumnA>(), TypedElementDataStorage::FColumn<FTestColumnB>(), TypedElementDataStorage::FColumn<FTestColumnC>() }, true));
+	CHECK(TestMatching(Example, { TypedElementDataStorage::FColumn<FTestColumnA>(), TypedElementDataStorage::FColumn<FTestColumnB>(), TypedElementDataStorage::FColumn<FTestColumnD>() }, false));
 }
 
 TEST_CASE_NAMED(FTypedElementQueryConditions2, "TypedElementQueryBuilder::FTypedElementQueryConditions A || B || C", "[ApplicationContextMask][EngineFilter]")
 {
-	FQueryConditions Example = FColumn<FTestColumnA>() || FColumn<FTestColumnB>() || FColumn<FTestColumnC>();
+	FQueryConditions Example = TypedElementDataStorage::FColumn<FTestColumnA>() || TypedElementDataStorage::FColumn<FTestColumnB>() || TypedElementDataStorage::FColumn<FTestColumnC>();
 	
 	CHECK(Example.MinimumColumnMatchRequired() == 1);
-	CHECK(TestMatching(Example, { FColumn<FTestColumnB>() }, true));
-	CHECK(TestMatching(Example, { FColumn<FTestColumnB>(), FColumn<FTestColumnC>() }, true));
-	CHECK(TestMatching(Example, { FColumn<FTestColumnD>() }, false));
+	CHECK(TestMatching(Example, { TypedElementDataStorage::FColumn<FTestColumnB>() }, true));
+	CHECK(TestMatching(Example, { TypedElementDataStorage::FColumn<FTestColumnB>(), TypedElementDataStorage::FColumn<FTestColumnC>() }, true));
+	CHECK(TestMatching(Example, { TypedElementDataStorage::FColumn<FTestColumnD>() }, false));
 }
 
 TEST_CASE_NAMED(FTypedElementQueryConditions3, "TypedElementQueryBuilder::FTypedElementQueryConditions A && (B || C)", "[ApplicationContextMask][EngineFilter]")
 {
-	FQueryConditions Example = FColumn<FTestColumnA>() && (FColumn<FTestColumnB>() || FColumn<FTestColumnC>());
+	FQueryConditions Example = TypedElementDataStorage::FColumn<FTestColumnA>() && (TypedElementDataStorage::FColumn<FTestColumnB>() || TypedElementDataStorage::FColumn<FTestColumnC>());
 	
 	CHECK(Example.MinimumColumnMatchRequired() == 2);
-	CHECK(TestMatching(Example, { FColumn<FTestColumnA>(), FColumn<FTestColumnB>() }, true));
-	CHECK(TestMatching(Example, { FColumn<FTestColumnA>(), FColumn<FTestColumnC>() }, true));
-	CHECK(TestMatching(Example, { FColumn<FTestColumnA>(), FColumn<FTestColumnD>() }, false));
-	CHECK(TestMatching(Example, { FColumn<FTestColumnD>(), FColumn<FTestColumnB>() }, false));
+	CHECK(TestMatching(Example, { TypedElementDataStorage::FColumn<FTestColumnA>(), TypedElementDataStorage::FColumn<FTestColumnB>() }, true));
+	CHECK(TestMatching(Example, { TypedElementDataStorage::FColumn<FTestColumnA>(), TypedElementDataStorage::FColumn<FTestColumnC>() }, true));
+	CHECK(TestMatching(Example, { TypedElementDataStorage::FColumn<FTestColumnA>(), TypedElementDataStorage::FColumn<FTestColumnD>() }, false));
+	CHECK(TestMatching(Example, { TypedElementDataStorage::FColumn<FTestColumnD>(), TypedElementDataStorage::FColumn<FTestColumnB>() }, false));
 }
 
 TEST_CASE_NAMED(FTypedElementQueryConditions4, "TypedElementQueryBuilder::FTypedElementQueryConditions A && (B || C) && (D || E)", "[ApplicationContextMask][EngineFilter]")
 {
 	FQueryConditions Example = 
-		FColumn<FTestColumnA>() && 
-		(FColumn<FTestColumnB>() || FColumn<FTestColumnC>()) &&
-		(FColumn<FTestColumnD>() || FColumn<FTestColumnE>());
+		TypedElementDataStorage::FColumn<FTestColumnA>() && 
+		(TypedElementDataStorage::FColumn<FTestColumnB>() || TypedElementDataStorage::FColumn<FTestColumnC>()) &&
+		(TypedElementDataStorage::FColumn<FTestColumnD>() || TypedElementDataStorage::FColumn<FTestColumnE>());
 	
 	CHECK(Example.MinimumColumnMatchRequired() == 3);
 	
-	CHECK(TestMatching(Example, { FColumn<FTestColumnA>() }, false));
-	CHECK(TestMatching(Example, { FColumn<FTestColumnA>(), FColumn<FTestColumnB>() }, false));
+	CHECK(TestMatching(Example, { TypedElementDataStorage::FColumn<FTestColumnA>() }, false));
+	CHECK(TestMatching(Example, { TypedElementDataStorage::FColumn<FTestColumnA>(), TypedElementDataStorage::FColumn<FTestColumnB>() }, false));
 
-	CHECK(TestMatching(Example, { FColumn<FTestColumnA>(), FColumn<FTestColumnB>(), FColumn<FTestColumnD>() }, true));
-	CHECK(TestMatching(Example, { FColumn<FTestColumnA>(), FColumn<FTestColumnB>(), FColumn<FTestColumnE>() }, true));
-	CHECK(TestMatching(Example, { FColumn<FTestColumnA>(), FColumn<FTestColumnC>(), FColumn<FTestColumnD>() }, true));
+	CHECK(TestMatching(Example, { TypedElementDataStorage::FColumn<FTestColumnA>(), TypedElementDataStorage::FColumn<FTestColumnB>(), TypedElementDataStorage::FColumn<FTestColumnD>() }, true));
+	CHECK(TestMatching(Example, { TypedElementDataStorage::FColumn<FTestColumnA>(), TypedElementDataStorage::FColumn<FTestColumnB>(), TypedElementDataStorage::FColumn<FTestColumnE>() }, true));
+	CHECK(TestMatching(Example, { TypedElementDataStorage::FColumn<FTestColumnA>(), TypedElementDataStorage::FColumn<FTestColumnC>(), TypedElementDataStorage::FColumn<FTestColumnD>() }, true));
 
-	CHECK(TestMatching(Example, { FColumn<FTestColumnA>(), FColumn<FTestColumnC>(), FColumn<FTestColumnF>() }, false));
-	CHECK(TestMatching(Example, { FColumn<FTestColumnA>(), FColumn<FTestColumnF>(), FColumn<FTestColumnD>() }, false));
-	CHECK(TestMatching(Example, { FColumn<FTestColumnB>(), FColumn<FTestColumnC>(), FColumn<FTestColumnD>() }, false));
+	CHECK(TestMatching(Example, { TypedElementDataStorage::FColumn<FTestColumnA>(), TypedElementDataStorage::FColumn<FTestColumnC>(), TypedElementDataStorage::FColumn<FTestColumnF>() }, false));
+	CHECK(TestMatching(Example, { TypedElementDataStorage::FColumn<FTestColumnA>(), TypedElementDataStorage::FColumn<FTestColumnF>(), TypedElementDataStorage::FColumn<FTestColumnD>() }, false));
+	CHECK(TestMatching(Example, { TypedElementDataStorage::FColumn<FTestColumnB>(), TypedElementDataStorage::FColumn<FTestColumnC>(), TypedElementDataStorage::FColumn<FTestColumnD>() }, false));
 }
 
 TEST_CASE_NAMED(FTypedElementQueryConditions5, "TypedElementQueryBuilder::FTypedElementQueryConditions (A || B) && (C || D) && (E || F)", "[ApplicationContextMask][EngineFilter]")
 {
 	FQueryConditions Example =
-		(FColumn<FTestColumnA>() || FColumn<FTestColumnB>()) &&
-		(FColumn<FTestColumnC>() || FColumn<FTestColumnD>()) &&
-		(FColumn<FTestColumnE>() || FColumn<FTestColumnF>());
+		(TypedElementDataStorage::FColumn<FTestColumnA>() || TypedElementDataStorage::FColumn<FTestColumnB>()) &&
+		(TypedElementDataStorage::FColumn<FTestColumnC>() || TypedElementDataStorage::FColumn<FTestColumnD>()) &&
+		(TypedElementDataStorage::FColumn<FTestColumnE>() || TypedElementDataStorage::FColumn<FTestColumnF>());
 	
 	CHECK(Example.MinimumColumnMatchRequired() == 3); 
 	
-	CHECK(TestMatching(Example, { FColumn<FTestColumnA>() }, false));
-	CHECK(TestMatching(Example, { FColumn<FTestColumnA>(), FColumn<FTestColumnC>() }, false));
+	CHECK(TestMatching(Example, { TypedElementDataStorage::FColumn<FTestColumnA>() }, false));
+	CHECK(TestMatching(Example, { TypedElementDataStorage::FColumn<FTestColumnA>(), TypedElementDataStorage::FColumn<FTestColumnC>() }, false));
 	
-	CHECK(TestMatching(Example, { FColumn<FTestColumnA>(), FColumn<FTestColumnC>(), FColumn<FTestColumnE>() }, true));
-	CHECK(TestMatching(Example, { FColumn<FTestColumnB>(), FColumn<FTestColumnC>(), FColumn<FTestColumnE>() }, true));
+	CHECK(TestMatching(Example, { TypedElementDataStorage::FColumn<FTestColumnA>(), TypedElementDataStorage::FColumn<FTestColumnC>(), TypedElementDataStorage::FColumn<FTestColumnE>() }, true));
+	CHECK(TestMatching(Example, { TypedElementDataStorage::FColumn<FTestColumnB>(), TypedElementDataStorage::FColumn<FTestColumnC>(), TypedElementDataStorage::FColumn<FTestColumnE>() }, true));
 	
-	CHECK(TestMatching(Example, { FColumn<FTestColumnA>(), FColumn<FTestColumnC>(), FColumn<FTestColumnG>() }, false));
-	CHECK(TestMatching(Example, { FColumn<FTestColumnA>(), FColumn<FTestColumnG>(), FColumn<FTestColumnD>() }, false));
-	CHECK(TestMatching(Example, { FColumn<FTestColumnG>(), FColumn<FTestColumnC>(), FColumn<FTestColumnD>() }, false));
+	CHECK(TestMatching(Example, { TypedElementDataStorage::FColumn<FTestColumnA>(), TypedElementDataStorage::FColumn<FTestColumnC>(), TypedElementDataStorage::FColumn<FTestColumnG>() }, false));
+	CHECK(TestMatching(Example, { TypedElementDataStorage::FColumn<FTestColumnA>(), TypedElementDataStorage::FColumn<FTestColumnG>(), TypedElementDataStorage::FColumn<FTestColumnD>() }, false));
+	CHECK(TestMatching(Example, { TypedElementDataStorage::FColumn<FTestColumnG>(), TypedElementDataStorage::FColumn<FTestColumnC>(), TypedElementDataStorage::FColumn<FTestColumnD>() }, false));
 }
 
 TEST_CASE_NAMED(FTypedElementQueryConditions6, "TypedElementQueryBuilder::FTypedElementQueryConditions ((A || B) && (C || D)) || (E && F)", "[ApplicationContextMask][EngineFilter]")
 {
 	FQueryConditions Example =
 		(
-			(FColumn<FTestColumnA>() || FColumn<FTestColumnB>()) &&
-			(FColumn<FTestColumnC>() || FColumn<FTestColumnD>())
+			(TypedElementDataStorage::FColumn<FTestColumnA>() || TypedElementDataStorage::FColumn<FTestColumnB>()) &&
+			(TypedElementDataStorage::FColumn<FTestColumnC>() || TypedElementDataStorage::FColumn<FTestColumnD>())
 		) ||
-		(FColumn<FTestColumnE>() && FColumn<FTestColumnF>());
+		(TypedElementDataStorage::FColumn<FTestColumnE>() && TypedElementDataStorage::FColumn<FTestColumnF>());
 	
 	CHECK(Example.MinimumColumnMatchRequired() == 2);
 
-	CHECK(TestMatching(Example, { FColumn<FTestColumnA>() }, false));
-	CHECK(TestMatching(Example, { FColumn<FTestColumnA>(), FColumn<FTestColumnC>() }, true));
+	CHECK(TestMatching(Example, { TypedElementDataStorage::FColumn<FTestColumnA>() }, false));
+	CHECK(TestMatching(Example, { TypedElementDataStorage::FColumn<FTestColumnA>(), TypedElementDataStorage::FColumn<FTestColumnC>() }, true));
 
-	CHECK(TestMatching(Example, { FColumn<FTestColumnE>(), FColumn<FTestColumnF>() }, true));
-	CHECK(TestMatching(Example, { FColumn<FTestColumnG>() }, false));
+	CHECK(TestMatching(Example, { TypedElementDataStorage::FColumn<FTestColumnE>(), TypedElementDataStorage::FColumn<FTestColumnF>() }, true));
+	CHECK(TestMatching(Example, { TypedElementDataStorage::FColumn<FTestColumnG>() }, false));
 }
 
 TEST_CASE_NAMED(FTypedElementQueryConditions7, "TypedElementQueryBuilder::FTypedElementQueryConditions (A && B) || (C && D) || (E && F)", "[ApplicationContextMask][EngineFilter]")
 {
 	FQueryConditions Example =
-		(FColumn<FTestColumnA>() && FColumn<FTestColumnB>()) ||
-		(FColumn<FTestColumnC>() && FColumn<FTestColumnD>()) ||
-		(FColumn<FTestColumnE>() && FColumn<FTestColumnF>());
+		(TypedElementDataStorage::FColumn<FTestColumnA>() && TypedElementDataStorage::FColumn<FTestColumnB>()) ||
+		(TypedElementDataStorage::FColumn<FTestColumnC>() && TypedElementDataStorage::FColumn<FTestColumnD>()) ||
+		(TypedElementDataStorage::FColumn<FTestColumnE>() && TypedElementDataStorage::FColumn<FTestColumnF>());
 
 	CHECK(Example.MinimumColumnMatchRequired() == 2);
 
-	CHECK(TestMatching(Example, { FColumn<FTestColumnA>() }, false));
-	CHECK(TestMatching(Example, { FColumn<FTestColumnA>(), FColumn<FTestColumnB>() }, true));
-	CHECK(TestMatching(Example, { FColumn<FTestColumnC>(), FColumn<FTestColumnD>() }, true));
-	CHECK(TestMatching(Example, { FColumn<FTestColumnE>(), FColumn<FTestColumnF>() }, true));
+	CHECK(TestMatching(Example, { TypedElementDataStorage::FColumn<FTestColumnA>() }, false));
+	CHECK(TestMatching(Example, { TypedElementDataStorage::FColumn<FTestColumnA>(), TypedElementDataStorage::FColumn<FTestColumnB>() }, true));
+	CHECK(TestMatching(Example, { TypedElementDataStorage::FColumn<FTestColumnC>(), TypedElementDataStorage::FColumn<FTestColumnD>() }, true));
+	CHECK(TestMatching(Example, { TypedElementDataStorage::FColumn<FTestColumnE>(), TypedElementDataStorage::FColumn<FTestColumnF>() }, true));
 
-	CHECK(TestMatching(Example, { FColumn<FTestColumnC>(), FColumn<FTestColumnD>(), FColumn<FTestColumnE>(), FColumn<FTestColumnF>() }, true));
+	CHECK(TestMatching(Example, { TypedElementDataStorage::FColumn<FTestColumnC>(), TypedElementDataStorage::FColumn<FTestColumnD>(), TypedElementDataStorage::FColumn<FTestColumnE>(), TypedElementDataStorage::FColumn<FTestColumnF>() }, true));
 
 
-	CHECK(TestMatching(Example, { FColumn<FTestColumnE>() }, false));
-	CHECK(TestMatching(Example, { FColumn<FTestColumnG>() }, false));
+	CHECK(TestMatching(Example, { TypedElementDataStorage::FColumn<FTestColumnE>() }, false));
+	CHECK(TestMatching(Example, { TypedElementDataStorage::FColumn<FTestColumnG>() }, false));
 }
 
 TEST_CASE_NAMED(FTypedElementQueryConditions_MultiMatch, "TypedElementQueryBuilder::FTypedElementQueryConditions_MultiMatch", "[ApplicationContextMask][EngineFilter]")
 {
 	FQueryConditions Example =
-		(FColumn<FTestColumnA>() || FColumn<FTestColumnB>()) &&
-		(FColumn<FTestColumnC>() || FColumn<FTestColumnD>()) &&
-		(FColumn<FTestColumnE>() || FColumn<FTestColumnF>());
+		(TypedElementDataStorage::FColumn<FTestColumnA>() || TypedElementDataStorage::FColumn<FTestColumnB>()) &&
+		(TypedElementDataStorage::FColumn<FTestColumnC>() || TypedElementDataStorage::FColumn<FTestColumnD>()) &&
+		(TypedElementDataStorage::FColumn<FTestColumnE>() || TypedElementDataStorage::FColumn<FTestColumnF>());
 
 	CHECK(TestMatching(Example, 
 		{ 
-			FColumn<FTestColumnA>(), 
-			FColumn<FTestColumnB>(), 
-			FColumn<FTestColumnC>(), 
-			FColumn<FTestColumnD>(),
-			FColumn<FTestColumnE>(),
-			FColumn<FTestColumnF>()
+			TypedElementDataStorage::FColumn<FTestColumnA>(), 
+			TypedElementDataStorage::FColumn<FTestColumnB>(), 
+			TypedElementDataStorage::FColumn<FTestColumnC>(), 
+			TypedElementDataStorage::FColumn<FTestColumnD>(),
+			TypedElementDataStorage::FColumn<FTestColumnE>(),
+			TypedElementDataStorage::FColumn<FTestColumnF>()
 		}, true));
 
 	CHECK(TestMatching(Example,
 		{
-			FColumn<FTestColumnA>(),
-			FColumn<FTestColumnC>(),
-			FColumn<FTestColumnE>(),
-			FColumn<FTestColumnG>()
+			TypedElementDataStorage::FColumn<FTestColumnA>(),
+			TypedElementDataStorage::FColumn<FTestColumnC>(),
+			TypedElementDataStorage::FColumn<FTestColumnE>(),
+			TypedElementDataStorage::FColumn<FTestColumnG>()
 		}, true));
 }
 
 TEST_CASE_NAMED(FTypedElementQueryConditions_Sorted, "TypedElementQueryBuilder::FTypedElementQueryConditions_Sorted", "[ApplicationContextMask][EngineFilter]")
 {
 	FQueryConditions Example =
-		(FColumn<FTestColumnA>() && FColumn<FTestColumnB>()) ||
-		(FColumn<FTestColumnC>() && FColumn<FTestColumnD>()) ||
-		(FColumn<FTestColumnE>() && FColumn<FTestColumnF>());
+		(TypedElementDataStorage::FColumn<FTestColumnA>() && TypedElementDataStorage::FColumn<FTestColumnB>()) ||
+		(TypedElementDataStorage::FColumn<FTestColumnC>() && TypedElementDataStorage::FColumn<FTestColumnD>()) ||
+		(TypedElementDataStorage::FColumn<FTestColumnE>() && TypedElementDataStorage::FColumn<FTestColumnF>());
 
 	CHECK(Example.MinimumColumnMatchRequired() == 2);
 
-	CHECK(TestMatching(Example, { FColumn<FTestColumnA>(), FColumn<FTestColumnB>() }, true, true));
-	CHECK(TestMatching(Example, { FColumn<FTestColumnC>(), FColumn<FTestColumnD>() }, true, true));
-	CHECK(TestMatching(Example, { FColumn<FTestColumnE>(), FColumn<FTestColumnF>() }, true, true));
+	CHECK(TestMatching(Example, { TypedElementDataStorage::FColumn<FTestColumnA>(), TypedElementDataStorage::FColumn<FTestColumnB>() }, true, true));
+	CHECK(TestMatching(Example, { TypedElementDataStorage::FColumn<FTestColumnC>(), TypedElementDataStorage::FColumn<FTestColumnD>() }, true, true));
+	CHECK(TestMatching(Example, { TypedElementDataStorage::FColumn<FTestColumnE>(), TypedElementDataStorage::FColumn<FTestColumnF>() }, true, true));
 
-	CHECK(TestMatching(Example, { FColumn<FTestColumnC>(), FColumn<FTestColumnD>(), FColumn<FTestColumnE>(), FColumn<FTestColumnF>() }, true, true));
+	CHECK(TestMatching(Example, { TypedElementDataStorage::FColumn<FTestColumnC>(), TypedElementDataStorage::FColumn<FTestColumnD>(), TypedElementDataStorage::FColumn<FTestColumnE>(), TypedElementDataStorage::FColumn<FTestColumnF>() }, true, true));
 }
 
 #endif // WITH_TESTS
diff --git a/Engine/Source/Runtime/TypedElementFramework/Tests/Elements/Framework/TypedElementTestColumns.h b/Engine/Source/Runtime/TypedElementFramework/Tests/Elements/Framework/TypedElementTestColumns.h
--- a/Engine/Source/Runtime/TypedElementFramework/Tests/Elements/Framework/TypedElementTestColumns.h
+++ b/Engine/Source/Runtime/TypedElementFramework/Tests/Elements/Framework/TypedElementTestColumns.h
@@ -98,7 +98,7 @@ struct FTEDSProcessorTestsReferenceColumn final : public FEditorDataStorageColum
 	GENERATED_BODY()
 
 	// UPROPERTY()
-	TypedElementDataStorage::RowHandle Reference = TypedElementDataStorage::InvalidRowHandle;
+	UE::Editor::DataStorage::RowHandle Reference = UE::Editor::DataStorage::InvalidRowHandle;
 
 	bool IsReferenced = false;
 };
