TEDS: Fix some uninitialized column members

diff --git a/Engine/Source/Runtime/TypedElementFramework/Public/Elements/Common/TypedElementHandles.h b/Engine/Source/Runtime/TypedElementFramework/Public/Elements/Common/TypedElementHandles.h
--- a/Engine/Source/Runtime/TypedElementFramework/Public/Elements/Common/TypedElementHandles.h
+++ b/Engine/Source/Runtime/TypedElementFramework/Public/Elements/Common/TypedElementHandles.h
@@ -30,7 +30,7 @@ struct FTedsRowHandle
 	GENERATED_BODY()
 
 	UPROPERTY()
-	uint64 RowHandle;
+	uint64 RowHandle = UE::Editor::DataStorage::InvalidRowHandle;
 
 	operator UE::Editor::DataStorage::RowHandle () const
 	{
diff --git a/Engine/Source/Runtime/TypedElementFramework/Tests/Elements/Framework/TypedElementDataStoragePerformanceTestCommands.h b/Engine/Source/Runtime/TypedElementFramework/Tests/Elements/Framework/TypedElementDataStoragePerformanceTestCommands.h
--- a/Engine/Source/Runtime/TypedElementFramework/Tests/Elements/Framework/TypedElementDataStoragePerformanceTestCommands.h
+++ b/Engine/Source/Runtime/TypedElementFramework/Tests/Elements/Framework/TypedElementDataStoragePerformanceTestCommands.h
@@ -17,7 +17,7 @@ struct FTest_PingPongPrePhys final : public FEditorDataStorageColumn
 {
 	GENERATED_BODY()
 	
-	UPROPERTY()
+	UPROPERTY(meta = (IgnoreForMemberInitializationTest))
 	uint64 Value;
 };
 
@@ -26,7 +26,7 @@ struct FTest_PingPongDurPhys final : public FEditorDataStorageColumn
 {
 	GENERATED_BODY()
 
-	UPROPERTY()
+	UPROPERTY(meta = (IgnoreForMemberInitializationTest))
 	uint64 Value;
 };
 
@@ -35,7 +35,7 @@ struct FTest_PingPongPostPhys final : public FEditorDataStorageColumn
 {
 	GENERATED_BODY()
 	
-	UPROPERTY()
+	UPROPERTY(meta = (IgnoreForMemberInitializationTest))
 	uint64 Value;
 };
 
