[Teds] Update namespace from TypedElementDataStorage to UE::Editor::DataStorage for metadata and query store.

diff --git a/Engine/Plugins/Experimental/EditorDataStorageFeatures/Source/TedsOutliner/Public/TedsOutlinerImpl.h b/Engine/Plugins/Experimental/EditorDataStorageFeatures/Source/TedsOutliner/Public/TedsOutlinerImpl.h
--- a/Engine/Plugins/Experimental/EditorDataStorageFeatures/Source/TedsOutliner/Public/TedsOutlinerImpl.h
+++ b/Engine/Plugins/Experimental/EditorDataStorageFeatures/Source/TedsOutliner/Public/TedsOutlinerImpl.h
@@ -48,11 +48,11 @@ struct FTedsOutlinerHierarchyData
 					return ParentColumn->Parent;
 				}
 
-				return TypedElementDataStorage::InvalidRowHandle;
+				return UE::Editor::DataStorage::InvalidRowHandle;
 			});
 
 		const FSetParentRowHandle RowHandleSetter = FSetParentRowHandle::CreateLambda([](void* InColumnData,
-			TypedElementDataStorage::RowHandle InRowHandle)
+			UE::Editor::DataStorage::RowHandle InRowHandle)
 			{
 				if(FTableRowParentColumn* ParentColumn = static_cast<FTableRowParentColumn *>(InColumnData))
 				{
@@ -78,12 +78,12 @@ struct FTedsOutlinerParams
 	SSceneOutliner* SceneOutliner;
 
 	// The query description that will be used to populate rows in the TEDS-Outliner
-	TAttribute<TypedElementDataStorage::FQueryDescription> QueryDescription;
+	TAttribute<UE::Editor::DataStorage::FQueryDescription> QueryDescription;
 	
 	// TEDS queries that will be used to create filters in this Outliner
 	// TEDS-Outliner TODO: Can we consolidate this with the SceneOutliner API to create filters? Currently has to be separate because FTEDSOutlinerFilter
 	// needs a reference to the mode which is not possible since filters with the Outliner API are added before the mode is init
-	TMap<FName, const TypedElementDataStorage::FQueryDescription> FilterQueries;
+	TMap<FName, const UE::Editor::DataStorage::FQueryDescription> FilterQueries;
 
 	// If true, this Outliner will automatically add all TEDS tags and columns as filters
 	bool bUseDefaultTedsFilters;
@@ -132,7 +132,7 @@ public:
 	FIsItemCompatible& IsItemCompatible();
 
 	// Update the selection in TEDS to the input rows, only if SelectionSetName is set
-	void SetSelection(const TArray<TypedElementDataStorage::RowHandle>& InSelectedRows);
+	void SetSelection(const TArray<UE::Editor::DataStorage::RowHandle>& InSelectedRows);
 
 	// Helper function to create a label widget for a given row
 	TSharedRef<SWidget> CreateLabelWidgetForItem(UE::Editor::DataStorage::RowHandle InRowHandle, const STableRow<FSceneOutlinerTreeItemPtr>& InRow) const;
@@ -141,15 +141,15 @@ public:
 	const TOptional<FTedsOutlinerHierarchyData>& GetHierarchyData();
 	
 	// Add an external query to the Outliner
-	void AddExternalQuery(FName QueryName, const TypedElementDataStorage::FQueryDescription& InQueryDescription);
+	void AddExternalQuery(FName QueryName, const UE::Editor::DataStorage::FQueryDescription& InQueryDescription);
 	void RemoveExternalQuery(FName QueryName);
 
 	// Append all external queries into the given query
-	void AppendExternalQueries(TypedElementDataStorage::FQueryDescription& OutQuery);
+	void AppendExternalQueries(UE::Editor::DataStorage::FQueryDescription& OutQuery);
 
 	// TEDS-Outliner TODO: This should live in TEDS long term
 	// Funtion to combine 2 queries (adds to second query to the first)
-	static void AppendQuery(TypedElementDataStorage::FQueryDescription& Query1, const TypedElementDataStorage::FQueryDescription& Query2);
+	static void AppendQuery(UE::Editor::DataStorage::FQueryDescription& Query1, const UE::Editor::DataStorage::FQueryDescription& Query2);
 
 	// Check if the given item's parent has changed (i.e ParentRowHandle does not match what the Outliner reports as the parent)
 	bool HasItemParentChanged(UE::Editor::DataStorage::RowHandle ItemRowHandle, UE::Editor::DataStorage::RowHandle ParentRowHandle) const;
@@ -195,10 +195,10 @@ protected:
 	TArray<FName> CellWidgetPurposes;
 	
 	// Initial query provided by user
-	TAttribute<TypedElementDataStorage::FQueryDescription> InitialQueryDescription;
+	TAttribute<UE::Editor::DataStorage::FQueryDescription> InitialQueryDescription;
 
 	// External queries that are currently active (e.g Filters)
-	TMap<FName, TypedElementDataStorage::FQueryDescription> ExternalQueries;
+	TMap<FName, UE::Editor::DataStorage::FQueryDescription> ExternalQueries;
 
 	// Optional Hierarchy Data
 	TOptional<FTedsOutlinerHierarchyData> HierarchyData;
diff --git a/Engine/Plugins/Experimental/EditorDataStorageFeatures/Source/TedsOutliner/Public/TedsOutlinerItem.h b/Engine/Plugins/Experimental/EditorDataStorageFeatures/Source/TedsOutliner/Public/TedsOutlinerItem.h
--- a/Engine/Plugins/Experimental/EditorDataStorageFeatures/Source/TedsOutliner/Public/TedsOutlinerItem.h
+++ b/Engine/Plugins/Experimental/EditorDataStorageFeatures/Source/TedsOutliner/Public/TedsOutlinerItem.h
@@ -17,7 +17,7 @@ struct FTedsOutlinerTreeItem : ISceneOutlinerTreeItem
 {
 public:
 	
-	DECLARE_DELEGATE_RetVal_OneParam(bool, FFilterPredicate, const TypedElementDataStorage::RowHandle);
+	DECLARE_DELEGATE_RetVal_OneParam(bool, FFilterPredicate, const UE::Editor::DataStorage::RowHandle);
 
 	TEDSOUTLINER_API bool Filter(FFilterPredicate Pred) const
 	{
diff --git a/Engine/Plugins/Experimental/EditorDataStorageFeatures/Source/TedsOutliner/Public/TedsOutlinerMode.h b/Engine/Plugins/Experimental/EditorDataStorageFeatures/Source/TedsOutliner/Public/TedsOutlinerMode.h
--- a/Engine/Plugins/Experimental/EditorDataStorageFeatures/Source/TedsOutliner/Public/TedsOutlinerMode.h
+++ b/Engine/Plugins/Experimental/EditorDataStorageFeatures/Source/TedsOutliner/Public/TedsOutlinerMode.h
@@ -16,14 +16,14 @@ public:
 	DRAG_DROP_OPERATOR_TYPE(FTedsRowDragDropOp, FDecoratedDragDropOp)
 
 	/** Rows we are dragging */
-	TArray<TypedElementDataStorage::RowHandle> DraggedRows;
+	TArray<UE::Editor::DataStorage::RowHandle> DraggedRows;
 
-	void Init(const TArray<TypedElementDataStorage::RowHandle>& InRowHandles)
+	void Init(const TArray<UE::Editor::DataStorage::RowHandle>& InRowHandles)
 	{
 		DraggedRows = InRowHandles;
 	}
 
-	static TSharedRef<FTedsRowDragDropOp> New(const TArray<TypedElementDataStorage::RowHandle>& InRowHandles)
+	static TSharedRef<FTedsRowDragDropOp> New(const TArray<UE::Editor::DataStorage::RowHandle>& InRowHandles)
 	{
 		TSharedRef<FTedsRowDragDropOp> Operation = MakeShareable(new FTedsRowDragDropOp);
 		
diff --git a/Engine/Plugins/Experimental/EditorDataStorageFeatures/Source/TedsOutliner/Public/Widgets/AlertWidget.h b/Engine/Plugins/Experimental/EditorDataStorageFeatures/Source/TedsOutliner/Public/Widgets/AlertWidget.h
--- a/Engine/Plugins/Experimental/EditorDataStorageFeatures/Source/TedsOutliner/Public/Widgets/AlertWidget.h
+++ b/Engine/Plugins/Experimental/EditorDataStorageFeatures/Source/TedsOutliner/Public/Widgets/AlertWidget.h
@@ -44,10 +44,10 @@ public:
 	TEDSOUTLINER_API ~FAlertWidgetConstructor() override = default;
 
 protected:
-	TEDSOUTLINER_API TSharedPtr<SWidget> CreateWidget(const TypedElementDataStorage::FMetaDataView& Arguments) override;
+	TEDSOUTLINER_API TSharedPtr<SWidget> CreateWidget(const UE::Editor::DataStorage::FMetaDataView& Arguments) override;
 	TEDSOUTLINER_API TConstArrayView<const UScriptStruct*> GetAdditionalColumnsList() const override;
 	TEDSOUTLINER_API bool FinalizeWidget(ITypedElementDataStorageInterface* DataStorage, ITypedElementDataStorageUiInterface* DataStorageUi,
-		TypedElementDataStorage::RowHandle Row, const TSharedPtr<SWidget>& Widget) override;
+		UE::Editor::DataStorage::RowHandle Row, const TSharedPtr<SWidget>& Widget) override;
 };
 
 USTRUCT(meta = (DisplayName = "General purpose alert"))
@@ -66,10 +66,10 @@ public:
 	TEDSOUTLINER_API ~FAlertHeaderWidgetConstructor() override = default;
 
 protected:
-	TEDSOUTLINER_API TSharedPtr<SWidget> CreateWidget(const TypedElementDataStorage::FMetaDataView& Arguments) override;
+	TEDSOUTLINER_API TSharedPtr<SWidget> CreateWidget(const UE::Editor::DataStorage::FMetaDataView& Arguments) override;
 	TEDSOUTLINER_API TConstArrayView<const UScriptStruct*> GetAdditionalColumnsList() const override;
 	TEDSOUTLINER_API bool FinalizeWidget(ITypedElementDataStorageInterface* DataStorage, ITypedElementDataStorageUiInterface* DataStorageUi,
-		TypedElementDataStorage::RowHandle Row, const TSharedPtr<SWidget>& Widget) override;
+		UE::Editor::DataStorage::RowHandle Row, const TSharedPtr<SWidget>& Widget) override;
 };
 
 USTRUCT(meta = (DisplayName = "General purpose alert header"))
diff --git a/Engine/Plugins/Experimental/EditorDataStorageFeatures/Source/TedsOutliner/Public/Widgets/TransformHeadsUpWidget.h b/Engine/Plugins/Experimental/EditorDataStorageFeatures/Source/TedsOutliner/Public/Widgets/TransformHeadsUpWidget.h
--- a/Engine/Plugins/Experimental/EditorDataStorageFeatures/Source/TedsOutliner/Public/Widgets/TransformHeadsUpWidget.h
+++ b/Engine/Plugins/Experimental/EditorDataStorageFeatures/Source/TedsOutliner/Public/Widgets/TransformHeadsUpWidget.h
@@ -43,11 +43,11 @@ public:
 	TConstArrayView<const UScriptStruct*> GetAdditionalColumnsList() const override;
 
 protected:
-	TSharedPtr<SWidget> CreateWidget(const TypedElementDataStorage::FMetaDataView& Arguments) override;
+	TSharedPtr<SWidget> CreateWidget(const UE::Editor::DataStorage::FMetaDataView& Arguments) override;
 	bool FinalizeWidget(
 		ITypedElementDataStorageInterface* DataStorage,
 		ITypedElementDataStorageUiInterface* DataStorageUi,
-		TypedElementDataStorage::RowHandle Row,
+		UE::Editor::DataStorage::RowHandle Row,
 		const TSharedPtr<SWidget>& Widget) override;
 };
 
diff --git a/Engine/Plugins/Experimental/EditorDataStorageFeatures/Source/TedsPropertyEditor/Private/Widgets/SPropertyMenuTedsRowPicker.cpp b/Engine/Plugins/Experimental/EditorDataStorageFeatures/Source/TedsPropertyEditor/Private/Widgets/SPropertyMenuTedsRowPicker.cpp
--- a/Engine/Plugins/Experimental/EditorDataStorageFeatures/Source/TedsPropertyEditor/Private/Widgets/SPropertyMenuTedsRowPicker.cpp
+++ b/Engine/Plugins/Experimental/EditorDataStorageFeatures/Source/TedsPropertyEditor/Private/Widgets/SPropertyMenuTedsRowPicker.cpp
@@ -77,7 +77,7 @@ void SPropertyMenuTedsRowPicker::Construct(const FArguments& InArgs)
 
 				Outliner->AddFilter(
 					MakeShared<TSceneOutlinerPredicateFilter<FTedsOutlinerTreeItem>>(
-						FTedsOutlinerTreeItem::FFilterPredicate::CreateLambda([this](const TypedElementDataStorage::RowHandle RowHandle) -> bool
+						FTedsOutlinerTreeItem::FFilterPredicate::CreateLambda([this](const UE::Editor::DataStorage::RowHandle RowHandle) -> bool
 					{
 						if (ElementFilter.IsBound())
 						{
diff --git a/Engine/Plugins/Experimental/EditorDataStorageFeatures/Source/TedsPropertyEditor/Public/TedsRowPickingMode.h b/Engine/Plugins/Experimental/EditorDataStorageFeatures/Source/TedsPropertyEditor/Public/TedsRowPickingMode.h
--- a/Engine/Plugins/Experimental/EditorDataStorageFeatures/Source/TedsPropertyEditor/Public/TedsRowPickingMode.h
+++ b/Engine/Plugins/Experimental/EditorDataStorageFeatures/Source/TedsPropertyEditor/Public/TedsRowPickingMode.h
@@ -6,8 +6,8 @@
 #include "SceneOutlinerFwd.h"
 #include "Elements/Common/TypedElementHandles.h"
 
-DECLARE_DELEGATE_OneParam(FOnTedsRowSelected, TypedElementDataStorage::RowHandle RowHandle);
-DECLARE_DELEGATE_RetVal_OneParam(bool, FOnShouldFilterTedsRow, const TypedElementDataStorage::RowHandle RowHandle);
+DECLARE_DELEGATE_OneParam(FOnTedsRowSelected, UE::Editor::DataStorage::RowHandle RowHandle);
+DECLARE_DELEGATE_RetVal_OneParam(bool, FOnShouldFilterTedsRow, const UE::Editor::DataStorage::RowHandle RowHandle);
 
 /*
 * Picking mode for TEDs Scene Outliner Widgets. Based off of FActorPickingMode
diff --git a/Engine/Plugins/Experimental/EditorDataStorageFeatures/Source/TedsPropertyEditor/Public/Widgets/SPropertyMenuTedsRowPicker.h b/Engine/Plugins/Experimental/EditorDataStorageFeatures/Source/TedsPropertyEditor/Public/Widgets/SPropertyMenuTedsRowPicker.h
--- a/Engine/Plugins/Experimental/EditorDataStorageFeatures/Source/TedsPropertyEditor/Public/Widgets/SPropertyMenuTedsRowPicker.h
+++ b/Engine/Plugins/Experimental/EditorDataStorageFeatures/Source/TedsPropertyEditor/Public/Widgets/SPropertyMenuTedsRowPicker.h
@@ -13,7 +13,7 @@ public:
 		, _ElementFilter()
 	{}
 		SLATE_ARGUMENT(bool, AllowClear)
-		SLATE_ARGUMENT(TypedElementDataStorage::FQueryDescription, QueryFilter)
+		SLATE_ARGUMENT(UE::Editor::DataStorage::FQueryDescription, QueryFilter)
 		SLATE_ARGUMENT(FOnShouldFilterTedsRow, ElementFilter)
 		SLATE_EVENT(FOnTedsRowSelected, OnSet)
 		SLATE_EVENT(FSimpleDelegate, OnClose)
@@ -25,14 +25,14 @@ private:
 
 	void OnClear();
 
-	void OnElementSelected(TypedElementDataStorage::RowHandle RowHandle);
+	void OnElementSelected(UE::Editor::DataStorage::RowHandle RowHandle);
 
-	void SetValue(TypedElementDataStorage::RowHandle RowHandle);
+	void SetValue(UE::Editor::DataStorage::RowHandle RowHandle);
 
 private:
 	bool bAllowClear;
 
-	TypedElementDataStorage::FQueryDescription QueryFilter;
+	UE::Editor::DataStorage::FQueryDescription QueryFilter;
 
 	FOnShouldFilterTedsRow ElementFilter;
 
diff --git a/Engine/Plugins/Experimental/EditorDataStorageFeatures/Source/TedsRevisionControl/Private/Processors/RevisionControlProcessors.cpp b/Engine/Plugins/Experimental/EditorDataStorageFeatures/Source/TedsRevisionControl/Private/Processors/RevisionControlProcessors.cpp
--- a/Engine/Plugins/Experimental/EditorDataStorageFeatures/Source/TedsRevisionControl/Private/Processors/RevisionControlProcessors.cpp
+++ b/Engine/Plugins/Experimental/EditorDataStorageFeatures/Source/TedsRevisionControl/Private/Processors/RevisionControlProcessors.cpp
@@ -79,7 +79,8 @@ namespace UE::Editor::RevisionControl::Private
 		ECVF_Default);
 	#endif
 
-	static FColor DetermineOverlayColor(const TypedElementDataStorage::IQueryContext& ObjectContext, const TypedElementDataStorage::ICommonQueryContext& SCCContext, const FTypedElementUObjectColumn& Actor)
+	using namespace UE::Editor::DataStorage;
+	static FColor DetermineOverlayColor(const IQueryContext& ObjectContext, const ICommonQueryContext& SCCContext, const FTypedElementUObjectColumn& Actor)
 	{
 		check(IsInGameThread());
 
diff --git a/Engine/Plugins/Experimental/EditorDataStorageFeatures/Source/TedsSettings/Private/TedsSettingsManager.cpp b/Engine/Plugins/Experimental/EditorDataStorageFeatures/Source/TedsSettings/Private/TedsSettingsManager.cpp
--- a/Engine/Plugins/Experimental/EditorDataStorageFeatures/Source/TedsSettings/Private/TedsSettingsManager.cpp
+++ b/Engine/Plugins/Experimental/EditorDataStorageFeatures/Source/TedsSettings/Private/TedsSettingsManager.cpp
@@ -21,9 +21,9 @@
 
 FTedsSettingsManager::FTedsSettingsManager()
 	: bIsInitialized{ false }
-	, SelectAllSettingsQuery{ TypedElementDataStorage::InvalidQueryHandle }
-	, SettingsContainerTable{ TypedElementDataStorage::InvalidTableHandle }
-	, SettingsCategoryTable{ TypedElementDataStorage::InvalidTableHandle }
+	, SelectAllSettingsQuery{ UE::Editor::DataStorage::InvalidQueryHandle }
+	, SettingsContainerTable{ UE::Editor::DataStorage::InvalidTableHandle }
+	, SettingsCategoryTable{ UE::Editor::DataStorage::InvalidTableHandle }
 {
 }
 
@@ -83,14 +83,14 @@ void FTedsSettingsManager::Shutdown()
 
 void FTedsSettingsManager::RegisterTables(ITypedElementDataStorageInterface& DataStorage)
 {
-	if (SettingsContainerTable == TypedElementDataStorage::InvalidTableHandle)
+	if (SettingsContainerTable == UE::Editor::DataStorage::InvalidTableHandle)
 	{
 		SettingsContainerTable = DataStorage.RegisterTable(
 			TTypedElementColumnTypeList<FNameColumn, FDisplayNameColumn, FDescriptionColumn, FSettingsContainerTag>(),
 			FName(TEXT("Editor_SettingsContainerTable")));
 	}
 
-	if (SettingsCategoryTable == TypedElementDataStorage::InvalidTableHandle)
+	if (SettingsCategoryTable == UE::Editor::DataStorage::InvalidTableHandle)
 	{
 		SettingsCategoryTable = DataStorage.RegisterTable(
 			TTypedElementColumnTypeList<FSettingsContainerReferenceColumn, FNameColumn, FDisplayNameColumn, FDescriptionColumn, FSettingsCategoryTag>(),
@@ -102,7 +102,7 @@ void FTedsSettingsManager::RegisterQueries(ITypedElementDataStorageInterface& Da
 {
 	using namespace TypedElementQueryBuilder;
 
-	if (SelectAllSettingsQuery == TypedElementDataStorage::InvalidQueryHandle)
+	if (SelectAllSettingsQuery == UE::Editor::DataStorage::InvalidQueryHandle)
 	{
 		SelectAllSettingsQuery = DataStorage.RegisterQuery(
 			Select()
@@ -116,7 +116,7 @@ void FTedsSettingsManager::RegisterQueries(ITypedElementDataStorageInterface& Da
 void FTedsSettingsManager::UnregisterQueries(ITypedElementDataStorageInterface& DataStorage)
 {
 	DataStorage.UnregisterQuery(SelectAllSettingsQuery);
-	SelectAllSettingsQuery = TypedElementDataStorage::InvalidQueryHandle;
+	SelectAllSettingsQuery = UE::Editor::DataStorage::InvalidQueryHandle;
 }
 
 void FTedsSettingsManager::RegisterSettings()
diff --git a/Engine/Plugins/Experimental/EditorDataStorageFeatures/Source/TedsSettings/Private/TedsSettingsManager.h b/Engine/Plugins/Experimental/EditorDataStorageFeatures/Source/TedsSettings/Private/TedsSettingsManager.h
--- a/Engine/Plugins/Experimental/EditorDataStorageFeatures/Source/TedsSettings/Private/TedsSettingsManager.h
+++ b/Engine/Plugins/Experimental/EditorDataStorageFeatures/Source/TedsSettings/Private/TedsSettingsManager.h
@@ -36,7 +36,7 @@ private:
 	void UpdateSettingsCategory(TSharedPtr<ISettingsCategory> SettingsCategory, UE::Editor::DataStorage::RowHandle ContainerRow, const bool bQueryExistingRows = true);
 
 	bool bIsInitialized;
-	TypedElementDataStorage::QueryHandle SelectAllSettingsQuery;
+	UE::Editor::DataStorage::QueryHandle SelectAllSettingsQuery;
 	UE::Editor::DataStorage::TableHandle SettingsContainerTable;
 	UE::Editor::DataStorage::TableHandle SettingsCategoryTable;
 
diff --git a/Engine/Plugins/Experimental/EditorDataStorageFeatures/Source/TedsSettings/Tests/TedsSettingsTests.cpp b/Engine/Plugins/Experimental/EditorDataStorageFeatures/Source/TedsSettings/Tests/TedsSettingsTests.cpp
--- a/Engine/Plugins/Experimental/EditorDataStorageFeatures/Source/TedsSettings/Tests/TedsSettingsTests.cpp
+++ b/Engine/Plugins/Experimental/EditorDataStorageFeatures/Source/TedsSettings/Tests/TedsSettingsTests.cpp
@@ -21,7 +21,7 @@ namespace UE::Editor::Settings::Tests
 {
 	using namespace UE::Editor;
 
-	BEGIN_DEFINE_SPEC(FTedsSettingsTestFixture, "TedsSettings", EAutomationTestFlags::EditorContext | EAutomationTestFlags::EngineFilter)
+	BEGIN_DEFINE_SPEC(FTedsSettingsTestFixture, "Editor.DataStorage.Settings", EAutomationTestFlags::EditorContext | EAutomationTestFlags::EngineFilter)
 	ISettingsModule* SettingsModule = nullptr;
 	UTypedElementRegistry* TypedElementRegistry = nullptr;
 	ITypedElementDataStorageInterface* DataStorage = nullptr;
diff --git a/Engine/Plugins/Experimental/EditorDataStorageFeatures/Source/TedsTableViewer/Private/TedsTableViewerColumn.cpp b/Engine/Plugins/Experimental/EditorDataStorageFeatures/Source/TedsTableViewer/Private/TedsTableViewerColumn.cpp
--- a/Engine/Plugins/Experimental/EditorDataStorageFeatures/Source/TedsTableViewer/Private/TedsTableViewerColumn.cpp
+++ b/Engine/Plugins/Experimental/EditorDataStorageFeatures/Source/TedsTableViewer/Private/TedsTableViewerColumn.cpp
@@ -17,7 +17,7 @@ namespace UE::Editor::DataStorage
 {
 	FTedsTableViewerColumn::FTedsTableViewerColumn(const FName& InColumnName, const TSharedPtr<FTypedElementWidgetConstructor>& InCellWidgetConstructor,
 		const TArray<TWeakObjectPtr<const UScriptStruct>>& InMatchedColumns, const TSharedPtr<FTypedElementWidgetConstructor>& InHeaderWidgetConstructor,
-		const TypedElementDataStorage::FMetaDataView& InWidgetMetaData)
+		const FMetaDataView& InWidgetMetaData)
 		: ColumnName(InColumnName)
 		, CellWidgetConstructor(InCellWidgetConstructor)
 		, HeaderWidgetConstructor(InHeaderWidgetConstructor)
diff --git a/Engine/Plugins/Experimental/EditorDataStorageFeatures/Source/TedsTableViewer/Private/TedsTableViewerModel.cpp b/Engine/Plugins/Experimental/EditorDataStorageFeatures/Source/TedsTableViewer/Private/TedsTableViewerModel.cpp
--- a/Engine/Plugins/Experimental/EditorDataStorageFeatures/Source/TedsTableViewer/Private/TedsTableViewerModel.cpp
+++ b/Engine/Plugins/Experimental/EditorDataStorageFeatures/Source/TedsTableViewer/Private/TedsTableViewerModel.cpp
@@ -44,7 +44,7 @@ namespace UE::Editor::DataStorage
 	{
 		Items.Empty();
 		
-		for(const TypedElementDataStorage::RowHandle RowHandle : RowQueryStack->GetOrderedRowList())
+		for(const RowHandle RowHandle : RowQueryStack->GetOrderedRowList())
 		{
 			if(IsRowDisplayable(RowHandle))
 			{
@@ -58,7 +58,7 @@ namespace UE::Editor::DataStorage
 		OnModelChanged.Broadcast();
 	}
 
-	bool FTedsTableViewerModel::IsRowDisplayable(TypedElementDataStorage::RowHandle InRowHandle) const
+	bool FTedsTableViewerModel::IsRowDisplayable(RowHandle InRowHandle) const
 	{
 		return !Storage->HasColumns<FHideRowFromUITag>(InRowHandle);
 	}
@@ -162,7 +162,7 @@ namespace UE::Editor::DataStorage
 				TSharedPtr<FTypedElementWidgetConstructor> CellConstructor(Constructor.Release());
 
 				TSharedPtr<FTypedElementWidgetConstructor> HeaderConstructor =
-					TableViewerUtils::CreateHeaderWidgetConstructor(*StorageUi, TypedElementDataStorage::FMetaDataView(), MatchedColumns, CellWidgetPurposes);
+					TableViewerUtils::CreateHeaderWidgetConstructor(*StorageUi, FMetaDataView(), MatchedColumns, CellWidgetPurposes);
 			
 				FName NameId = TableViewerUtils::FindLongestMatchingName(MatchedColumns, IndexOffset);
 
@@ -184,7 +184,7 @@ namespace UE::Editor::DataStorage
 		for(const FName& WidgetPurpose : CellWidgetPurposes)
 		{
 			StorageUi->CreateWidgetConstructors(WidgetPurpose, MatchApproach::LongestMatch, ColumnsCopy, 
-			TypedElementDataStorage::FMetaDataView(), ColumnConstructor);
+			FMetaDataView(), ColumnConstructor);
 		}
 
 		// For any remaining columns, we'll try to find and use any default widgets
@@ -196,7 +196,7 @@ namespace UE::Editor::DataStorage
 				TSharedPtr<FTypedElementWidgetConstructor> CellConstructor(Constructor.Release());
 
 				TSharedPtr<FTypedElementWidgetConstructor> HeaderConstructor =
-					TableViewerUtils::CreateHeaderWidgetConstructor(*StorageUi, TypedElementDataStorage::FMetaDataView(), {ColumnType}, CellWidgetPurposes);
+					TableViewerUtils::CreateHeaderWidgetConstructor(*StorageUi, FMetaDataView(), {ColumnType}, CellWidgetPurposes);
 
 				FName NameId = FName(ColumnType->GetDisplayNameText().ToString());
 
@@ -217,7 +217,7 @@ namespace UE::Editor::DataStorage
 			{
 				const FName DefaultWidgetPurpose(WidgetPurpose.ToString() + TEXT(".Default"));
 
-				StorageUi->CreateWidgetConstructors(DefaultWidgetPurpose, TypedElementDataStorage::FMetaDataView(), AssignWidgetToColumn);
+				StorageUi->CreateWidgetConstructors(DefaultWidgetPurpose, FMetaDataView(), AssignWidgetToColumn);
 				
 				if (BeforeIndexOffset != IndexOffset)
 				{
@@ -246,7 +246,7 @@ namespace UE::Editor::DataStorage
 		}
 	}
 
-	bool FTedsTableViewerModel::IsRowVisible(TypedElementDataStorage::RowHandle InRowHandle) const
+	bool FTedsTableViewerModel::IsRowVisible(RowHandle InRowHandle) const
 	{
 		if(!IsItemVisible.IsBound())
 		{
diff --git a/Engine/Plugins/Experimental/EditorDataStorageFeatures/Source/TedsTableViewer/Private/TedsTableViewerModule.cpp b/Engine/Plugins/Experimental/EditorDataStorageFeatures/Source/TedsTableViewer/Private/TedsTableViewerModule.cpp
--- a/Engine/Plugins/Experimental/EditorDataStorageFeatures/Source/TedsTableViewer/Private/TedsTableViewerModule.cpp
+++ b/Engine/Plugins/Experimental/EditorDataStorageFeatures/Source/TedsTableViewer/Private/TedsTableViewerModule.cpp
@@ -75,7 +75,7 @@ public:
 		using namespace TypedElementQueryBuilder;
 
 		// We'll just create a test table viewer that views all actor rows without actually updating it dynamically for now
-		static TypedElementDataStorage::QueryHandle QueryHandle =
+		static QueryHandle QueryHandle =
 			DataStorage->RegisterQuery(
 						Select()
 						.Where()
@@ -85,7 +85,7 @@ public:
 
 		Rows.Empty();
 		
-		TypedElementDataStorage::FQueryResult QueryResult = DataStorage->RunQuery(QueryHandle,
+		UE::Editor::DataStorage::FQueryResult QueryResult = DataStorage->RunQuery(QueryHandle,
 			CreateDirectQueryCallbackBinding([this](const ITypedElementDataStorageInterface::IDirectQueryContext& Context, const UE::Editor::DataStorage::RowHandle* RowHandles)
 		{
 			Rows.Append(RowHandles, Context.GetRowCount());
diff --git a/Engine/Plugins/Experimental/EditorDataStorageFeatures/Source/TedsTableViewer/Private/TedsTableViewerUtils.cpp b/Engine/Plugins/Experimental/EditorDataStorageFeatures/Source/TedsTableViewer/Private/TedsTableViewerUtils.cpp
--- a/Engine/Plugins/Experimental/EditorDataStorageFeatures/Source/TedsTableViewer/Private/TedsTableViewerUtils.cpp
+++ b/Engine/Plugins/Experimental/EditorDataStorageFeatures/Source/TedsTableViewer/Private/TedsTableViewerUtils.cpp
@@ -90,7 +90,7 @@ namespace UE::Editor::DataStorage::TableViewerUtils
 		return VerifiedColumnTypes;
 	}
 
-	TSharedPtr<FTypedElementWidgetConstructor> CreateHeaderWidgetConstructor(ITypedElementDataStorageUiInterface& StorageUi, const TypedElementDataStorage::FMetaDataView& InMetaData, TConstArrayView<TWeakObjectPtr<const UScriptStruct>> ColumnTypes, const TConstArrayView<FName> CellWidgetPurposes)
+	TSharedPtr<FTypedElementWidgetConstructor> CreateHeaderWidgetConstructor(ITypedElementDataStorageUiInterface& StorageUi, const UE::Editor::DataStorage::FMetaDataView& InMetaData, TConstArrayView<TWeakObjectPtr<const UScriptStruct>> ColumnTypes, const TConstArrayView<FName> CellWidgetPurposes)
 	{
 		using MatchApproach = ITypedElementDataStorageUiInterface::EMatchApproach;
 
diff --git a/Engine/Plugins/Experimental/EditorDataStorageFeatures/Source/TedsTableViewer/Private/Widgets/RowHandleWidget.cpp b/Engine/Plugins/Experimental/EditorDataStorageFeatures/Source/TedsTableViewer/Private/Widgets/RowHandleWidget.cpp
--- a/Engine/Plugins/Experimental/EditorDataStorageFeatures/Source/TedsTableViewer/Private/Widgets/RowHandleWidget.cpp
+++ b/Engine/Plugins/Experimental/EditorDataStorageFeatures/Source/TedsTableViewer/Private/Widgets/RowHandleWidget.cpp
@@ -34,7 +34,7 @@ FRowHandleWidgetConstructor::FRowHandleWidgetConstructor()
 	
 }
 
-TSharedPtr<SWidget> FRowHandleWidgetConstructor::CreateWidget(const TypedElementDataStorage::FMetaDataView& Arguments)
+TSharedPtr<SWidget> FRowHandleWidgetConstructor::CreateWidget(const UE::Editor::DataStorage::FMetaDataView& Arguments)
 {
 	return SNew(SBox)
 		.HAlign(HAlign_Left)
@@ -51,7 +51,7 @@ bool FRowHandleWidgetConstructor::FinalizeWidget(ITypedElementDataStorageInterfa
 	
 	SBox* BoxWidget = static_cast<SBox*>(Widget.Get());
 
-	TypedElementDataStorage::RowHandle TargetRowHandle = TypedElementDataStorage::InvalidRowHandle;
+	UE::Editor::DataStorage::RowHandle TargetRowHandle = UE::Editor::DataStorage::InvalidRowHandle;
 
 	if(const FTypedElementRowReferenceColumn* RowReferenceColumn = DataStorage->GetColumn<FTypedElementRowReferenceColumn>(Row))
 	{
diff --git a/Engine/Plugins/Experimental/EditorDataStorageFeatures/Source/TedsTableViewer/Private/Widgets/STedsTableViewer.cpp b/Engine/Plugins/Experimental/EditorDataStorageFeatures/Source/TedsTableViewer/Private/Widgets/STedsTableViewer.cpp
--- a/Engine/Plugins/Experimental/EditorDataStorageFeatures/Source/TedsTableViewer/Private/Widgets/STedsTableViewer.cpp
+++ b/Engine/Plugins/Experimental/EditorDataStorageFeatures/Source/TedsTableViewer/Private/Widgets/STedsTableViewer.cpp
@@ -30,7 +30,7 @@ namespace UE::Editor::DataStorage
 		CreateTedsWidget();
 
 		// Attribute binder to bind widget columns to attributes on the ListView
-		EditorDataStorage::FAttributeBinder Binder(TedsWidget->GetRowHandle());
+		FAttributeBinder Binder(TedsWidget->GetRowHandle());
 
 		ListView = SNew(SListView<TableViewerItemPtr>)
 			.HeaderRow(HeaderRowWidget)
@@ -64,7 +64,7 @@ namespace UE::Editor::DataStorage
 
 		if(ITypedElementDataStorageInterface* DataStorage = Model->GetDataStorageInterface())
 		{
-			const TypedElementDataStorage::RowHandle WidgetRowHandle = TedsWidget->GetRowHandle();
+			const RowHandle WidgetRowHandle = TedsWidget->GetRowHandle();
 		
 			if(DataStorage->IsRowAvailable(WidgetRowHandle))
 			{
@@ -150,7 +150,7 @@ namespace UE::Editor::DataStorage
 		RefreshColumnWidgets();
 	}
 
-	void STedsTableViewer::ForEachSelectedRow(TFunctionRef<void(TypedElementDataStorage::RowHandle)> InCallback) const
+	void STedsTableViewer::ForEachSelectedRow(TFunctionRef<void(RowHandle)> InCallback) const
 	{
 		TArray<TableViewerItemPtr> SelectedRows;
 		ListView->GetSelectedItems(SelectedRows);
@@ -161,7 +161,7 @@ namespace UE::Editor::DataStorage
 		}
 	}
 
-	TypedElementDataStorage::RowHandle STedsTableViewer::GetWidgetRowHandle() const
+	RowHandle STedsTableViewer::GetWidgetRowHandle() const
 	{
 		return TedsWidget->GetRowHandle();
 	}
diff --git a/Engine/Plugins/Experimental/EditorDataStorageFeatures/Source/TedsTableViewer/Private/Widgets/UObjectLabelWidget.cpp b/Engine/Plugins/Experimental/EditorDataStorageFeatures/Source/TedsTableViewer/Private/Widgets/UObjectLabelWidget.cpp
--- a/Engine/Plugins/Experimental/EditorDataStorageFeatures/Source/TedsTableViewer/Private/Widgets/UObjectLabelWidget.cpp
+++ b/Engine/Plugins/Experimental/EditorDataStorageFeatures/Source/TedsTableViewer/Private/Widgets/UObjectLabelWidget.cpp
@@ -29,11 +29,12 @@ FUObjectLabelWidgetConstructor::FUObjectLabelWidgetConstructor()
 
 TSharedPtr<SWidget> FUObjectLabelWidgetConstructor::CreateWidget(ITypedElementDataStorageInterface* DataStorage,
 	ITypedElementDataStorageUiInterface* DataStorageUi, RowHandle TargetRow, RowHandle WidgetRow,
-	const TypedElementDataStorage::FMetaDataView& Arguments)
+	const UE::Editor::DataStorage::FMetaDataView& Arguments)
 {
+	using namespace UE::Editor::DataStorage;
 	if(DataStorage->IsRowAvailable(TargetRow))
 	{
-		UE::EditorDataStorage::FAttributeBinder Binder(TargetRow, DataStorage);
+		FAttributeBinder Binder(TargetRow, DataStorage);
 
 		// Once TEDS UI has widget combining functionality, the binder can be used to create the type info widget and label widget and combine them
 		return SNew(SHorizontalBox)
@@ -41,7 +42,7 @@ TSharedPtr<SWidget> FUObjectLabelWidgetConstructor::CreateWidget(ITypedElementDa
 				.AutoWidth()
 			[
 				SNew(SImage)
-					.Image(UE::Editor::DataStorage::TableViewerUtils::GetIconForRow(DataStorage, TargetRow))
+					.Image(TableViewerUtils::GetIconForRow(DataStorage, TargetRow))
 					.ColorAndOpacity(FSlateColor::UseForeground())
 			]
 			+SHorizontalBox::Slot()
diff --git a/Engine/Plugins/Experimental/EditorDataStorageFeatures/Source/TedsTableViewer/Public/TedsTableViewerColumn.h b/Engine/Plugins/Experimental/EditorDataStorageFeatures/Source/TedsTableViewer/Public/TedsTableViewerColumn.h
--- a/Engine/Plugins/Experimental/EditorDataStorageFeatures/Source/TedsTableViewer/Public/TedsTableViewerColumn.h
+++ b/Engine/Plugins/Experimental/EditorDataStorageFeatures/Source/TedsTableViewer/Public/TedsTableViewerColumn.h
@@ -31,7 +31,7 @@ namespace UE::Editor::DataStorage
 			const TSharedPtr<FTypedElementWidgetConstructor>& InCellWidgetConstructor, // The widget constructor to use for this column
 			const TArray<TWeakObjectPtr<const UScriptStruct>>& InMatchedColumns = {}, // An optional list of matched Teds columns 
 			const TSharedPtr<FTypedElementWidgetConstructor>& InHeaderWidgetConstructor = nullptr, // Optional constructor to use for the header widget
-			const TypedElementDataStorage::FMetaDataView& InWidgetMetaData = TypedElementDataStorage::FMetaDataView()); // Optional metadata to use when constructing widgets
+			const FMetaDataView& InWidgetMetaData = FMetaDataView()); // Optional metadata to use when constructing widgets
 
 		TEDSTABLEVIEWER_API ~FTedsTableViewerColumn();
 		
@@ -67,7 +67,7 @@ namespace UE::Editor::DataStorage
 		TArray<TWeakObjectPtr<const UScriptStruct>> MatchedColumns;
 		
 		// The Metadata used to create widgets
-		TypedElementDataStorage::FMetaDataView WidgetMetaData;
+		FMetaDataView WidgetMetaData;
 
 		// TEDS Constructs
 		ITypedElementDataStorageInterface* Storage;
diff --git a/Engine/Plugins/Experimental/EditorDataStorageFeatures/Source/TedsTableViewer/Public/TedsTableViewerUtils.h b/Engine/Plugins/Experimental/EditorDataStorageFeatures/Source/TedsTableViewer/Public/TedsTableViewerUtils.h
--- a/Engine/Plugins/Experimental/EditorDataStorageFeatures/Source/TedsTableViewer/Public/TedsTableViewerUtils.h
+++ b/Engine/Plugins/Experimental/EditorDataStorageFeatures/Source/TedsTableViewer/Public/TedsTableViewerUtils.h
@@ -9,7 +9,7 @@
 
 #include "TedsTableViewerUtils.generated.h"
 
-namespace TypedElementDataStorage
+namespace UE::Editor::DataStorage
 {
 	class FMetaDataView;
 }
@@ -29,7 +29,7 @@ namespace UE::Editor::DataStorage::TableViewerUtils
 
 	// Create a header widget constructor for the given columns
 	TEDSTABLEVIEWER_API TSharedPtr<FTypedElementWidgetConstructor> CreateHeaderWidgetConstructor(ITypedElementDataStorageUiInterface& StorageUi, 
-	const TypedElementDataStorage::FMetaDataView& InMetaData, TConstArrayView<TWeakObjectPtr<const UScriptStruct>> ColumnTypes,
+	const UE::Editor::DataStorage::FMetaDataView& InMetaData, TConstArrayView<TWeakObjectPtr<const UScriptStruct>> ColumnTypes,
 	const TConstArrayView<FName> CellWidgetPurposes);
 
 	// Create a copy of the provided column types array after discarding invalid entries
diff --git a/Engine/Plugins/Experimental/EditorDataStorageFeatures/Source/TedsTableViewer/Public/Widgets/RowHandleWidget.h b/Engine/Plugins/Experimental/EditorDataStorageFeatures/Source/TedsTableViewer/Public/Widgets/RowHandleWidget.h
--- a/Engine/Plugins/Experimental/EditorDataStorageFeatures/Source/TedsTableViewer/Public/Widgets/RowHandleWidget.h
+++ b/Engine/Plugins/Experimental/EditorDataStorageFeatures/Source/TedsTableViewer/Public/Widgets/RowHandleWidget.h
@@ -36,7 +36,7 @@ public:
 	virtual ~FRowHandleWidgetConstructor() override = default;
 
 protected:
-	virtual TSharedPtr<SWidget> CreateWidget(const TypedElementDataStorage::FMetaDataView& Arguments) override;
+	virtual TSharedPtr<SWidget> CreateWidget(const UE::Editor::DataStorage::FMetaDataView& Arguments) override;
 	virtual bool FinalizeWidget(
 		ITypedElementDataStorageInterface* DataStorage,
 		ITypedElementDataStorageUiInterface* DataStorageUi,
diff --git a/Engine/Plugins/Experimental/InstanceDataObjectFixupTool/Source/Private/InstanceDataObjectFixupPanel.cpp b/Engine/Plugins/Experimental/InstanceDataObjectFixupTool/Source/Private/InstanceDataObjectFixupPanel.cpp
--- a/Engine/Plugins/Experimental/InstanceDataObjectFixupTool/Source/Private/InstanceDataObjectFixupPanel.cpp
+++ b/Engine/Plugins/Experimental/InstanceDataObjectFixupTool/Source/Private/InstanceDataObjectFixupPanel.cpp
@@ -248,7 +248,7 @@ static bool ObjectHasLoosePropertiesThatNeedFixup(UObject* Object)
 
 FInstanceDataObjectFixupPanel::~FInstanceDataObjectFixupPanel()
 {
-	using namespace TypedElementDataStorage;
+	using namespace UE::Editor::DataStorage;
 
 	UTypedElementRegistry* Registry = UTypedElementRegistry::GetInstance();
 	ITypedElementDataStorageInterface* DataStorage = Registry ? Registry->GetMutableDataStorage() : nullptr;
diff --git a/Engine/Source/Runtime/TypedElementFramework/Private/Elements/Framework/TypedElementAttributeBinding.cpp b/Engine/Source/Runtime/TypedElementFramework/Private/Elements/Framework/TypedElementAttributeBinding.cpp
--- a/Engine/Source/Runtime/TypedElementFramework/Private/Elements/Framework/TypedElementAttributeBinding.cpp
+++ b/Engine/Source/Runtime/TypedElementFramework/Private/Elements/Framework/TypedElementAttributeBinding.cpp
@@ -4,9 +4,9 @@
 
 #include "Elements/Framework/TypedElementRegistry.h"
 
-namespace UE::EditorDataStorage
+namespace UE::Editor::DataStorage
 {
-	FAttributeBinder::FAttributeBinder(TypedElementDataStorage::RowHandle InTargetRow)
+	FAttributeBinder::FAttributeBinder(RowHandle InTargetRow)
 		: TargetRow(InTargetRow)
 	{
 		if (ensureMsgf(UTypedElementRegistry::GetInstance()->AreDataStorageInterfacesSet(), 
@@ -17,10 +17,10 @@ namespace UE::EditorDataStorage
 	}
 
 	FAttributeBinder::FAttributeBinder(
-		TypedElementDataStorage::RowHandle InTargetRow, ITypedElementDataStorageInterface* InDataStorage)
+		RowHandle InTargetRow, ITypedElementDataStorageInterface* InDataStorage)
 		: TargetRow(InTargetRow)
 		, DataStorage(InDataStorage)
 	{
 		ensureMsgf(DataStorage, TEXT("The Editor Data Storage plugin needs to be enabled to use attribute bindings."));
 	}
-}
+} // namespace UE::Editor::DataStorage
diff --git a/Engine/Source/Runtime/TypedElementFramework/Private/Elements/Framework/TypedElementDataStorageWidget.cpp b/Engine/Source/Runtime/TypedElementFramework/Private/Elements/Framework/TypedElementDataStorageWidget.cpp
--- a/Engine/Source/Runtime/TypedElementFramework/Private/Elements/Framework/TypedElementDataStorageWidget.cpp
+++ b/Engine/Source/Runtime/TypedElementFramework/Private/Elements/Framework/TypedElementDataStorageWidget.cpp
@@ -5,7 +5,7 @@
 #include "Elements/Framework/TypedElementRegistry.h"
 
 STedsWidget::STedsWidget()
-	: UiRowHandle(TypedElementDataStorage::InvalidRowHandle)
+	: UiRowHandle(UE::Editor::DataStorage::InvalidRowHandle)
 {
 
 }
@@ -15,7 +15,7 @@ void STedsWidget::Construct(const FArguments& InArgs)
 	UiRowHandle = InArgs._UiRowHandle;
 
 	// If the Ui row wasn't already registered externally, register it with Teds
-	if(UiRowHandle == TypedElementDataStorage::InvalidRowHandle)
+	if(UiRowHandle == UE::Editor::DataStorage::InvalidRowHandle)
 	{
 		RegisterTedsWidget(InArgs._Content.Widget);
 	}
@@ -28,6 +28,7 @@ void STedsWidget::Construct(const FArguments& InArgs)
 
 void STedsWidget::RegisterTedsWidget(const TSharedPtr<SWidget>& InContentWidget)
 {
+	using namespace UE::Editor::DataStorage;
 	ITypedElementDataStorageInterface* Storage = GetStorageIfAvailable();
 
 	// If TEDS is not enabled, STedsWidget will just behave like a regular widget
@@ -36,8 +37,8 @@ void STedsWidget::RegisterTedsWidget(const TSharedPtr<SWidget>& InContentWidget)
 		return;
 	}
 	
-	const TypedElementDataStorage::TableHandle WidgetTable = Storage->FindTable(TEXT("Editor_WidgetTable"));
-	if(WidgetTable == TypedElementDataStorage::InvalidTableHandle)
+	const TableHandle WidgetTable = Storage->FindTable(TEXT("Editor_WidgetTable"));
+	if(WidgetTable == InvalidTableHandle)
 	{
 		return;
 	}
@@ -67,7 +68,7 @@ void STedsWidget::SetContent(const TSharedRef< SWidget >& InContent)
 	];
 }
 
-TypedElementDataStorage::RowHandle STedsWidget::GetRowHandle() const
+UE::Editor::DataStorage::RowHandle STedsWidget::GetRowHandle() const
 {
 	return UiRowHandle;
 }
diff --git a/Engine/Source/Runtime/TypedElementFramework/Private/Elements/Framework/TypedElementMetaData.cpp b/Engine/Source/Runtime/TypedElementFramework/Private/Elements/Framework/TypedElementMetaData.cpp
--- a/Engine/Source/Runtime/TypedElementFramework/Private/Elements/Framework/TypedElementMetaData.cpp
+++ b/Engine/Source/Runtime/TypedElementFramework/Private/Elements/Framework/TypedElementMetaData.cpp
@@ -6,7 +6,7 @@
 #include "UObject/Class.h"
 #include "UObject/UnrealType.h"
 
-namespace TypedElementDataStorage
+namespace UE::Editor::DataStorage
 {
 	template<class... Ts>
 	struct TOverloaded : Ts... 
@@ -135,7 +135,7 @@ namespace TypedElementDataStorage
 	// FQueryMetaDataView
 	//
 
-	FQueryMetaDataView::FQueryMetaDataView(const TypedElementDataStorage::FQueryDescription& InQuery)
+	FQueryMetaDataView::FQueryMetaDataView(const FQueryDescription& InQuery)
 		:Query(InQuery)
 	{
 	}
@@ -226,4 +226,4 @@ namespace TypedElementDataStorage
 	{
 		return MetaData.Find(AttributeName);
 	}
-} // namespace TypedElementDataStorage
+} // namespace UE::Editor::DataStorage
diff --git a/Engine/Source/Runtime/TypedElementFramework/Private/Elements/Framework/TypedElementQueryBuilder.cpp b/Engine/Source/Runtime/TypedElementFramework/Private/Elements/Framework/TypedElementQueryBuilder.cpp
--- a/Engine/Source/Runtime/TypedElementFramework/Private/Elements/Framework/TypedElementQueryBuilder.cpp
+++ b/Engine/Source/Runtime/TypedElementFramework/Private/Elements/Framework/TypedElementQueryBuilder.cpp
@@ -9,6 +9,8 @@
 
 namespace TypedElementQueryBuilder
 {
+	using namespace UE::Editor::DataStorage;
+
 	const UScriptStruct* Type(FTopLevelAssetPath Name)
 	{
 		const UScriptStruct* StructInfo = TypeOptional(Name);
@@ -143,10 +145,10 @@ namespace TypedElementQueryBuilder
 		return *this;
 	}
 
-	FSimpleQuery& FSimpleQuery::All(const UE::Editor::DataStorage::FValueTag& Tag, const FName& Value)
+	FSimpleQuery& FSimpleQuery::All(const FValueTag& Tag, const FName& Value)
 	{
 		Query->ValueTags.Emplace(
-			TypedElementDataStorage::FQueryDescription::FValueTagData
+			FQueryDescription::FValueTagData
 			{
 				.Tag = Tag,
 				.MatchValue = Value
@@ -172,14 +174,14 @@ namespace TypedElementQueryBuilder
 		return All(FValueTag(Enum.GetFName()), ValueName);
 	}
 	
-	FSimpleQuery& FSimpleQuery::All(const UE::Editor::DataStorage::FDynamicColumnDescription& Description)
+	FSimpleQuery& FSimpleQuery::All(const FDynamicColumnDescription& Description)
 	{
 		Query->DynamicConditionDescriptions.Add(Description);
 		Query->DynamicConditionOperations.Add(ITypedElementDataStorageInterface::FQueryDescription::EOperatorType::SimpleAll);
 		return *this;
 	}
 
-	FSimpleQuery& FSimpleQuery::All(const UE::Editor::DataStorage::FValueTag& Tag)
+	FSimpleQuery& FSimpleQuery::All(const FValueTag& Tag)
 	{
 		return All(Tag, NAME_None);
 	}
@@ -207,10 +209,10 @@ namespace TypedElementQueryBuilder
 		return *this;
 	}
 
-	FSimpleQuery& FSimpleQuery::Any(const UE::Editor::DataStorage::FDynamicColumnDescription& Description)
+	FSimpleQuery& FSimpleQuery::Any(const FDynamicColumnDescription& Description)
 	{
 		Query->DynamicConditionDescriptions.Add(Description);
-		Query->DynamicConditionOperations.Add(TypedElementDataStorage::FQueryDescription::EOperatorType::SimpleAny);
+		Query->DynamicConditionOperations.Add(FQueryDescription::EOperatorType::SimpleAny);
 		return *this;
 	}
 
@@ -237,10 +239,10 @@ namespace TypedElementQueryBuilder
 		return *this;
 	}
 
-	FSimpleQuery& FSimpleQuery::None(const UE::Editor::DataStorage::FDynamicColumnDescription& Description)
+	FSimpleQuery& FSimpleQuery::None(const FDynamicColumnDescription& Description)
 	{
 		Query->DynamicConditionDescriptions.Add(Description);
-		Query->DynamicConditionOperations.Add(TypedElementDataStorage::FQueryDescription::EOperatorType::SimpleNone);
+		Query->DynamicConditionOperations.Add(FQueryDescription::EOperatorType::SimpleNone);
 		return *this;
 	}
 
@@ -255,7 +257,7 @@ namespace TypedElementQueryBuilder
 		Query->Callback.AfterGroups.Shrink();
 		Query->SelectionTypes.Shrink();
 		Query->SelectionAccessTypes.Shrink();
-		for (TypedElementDataStorage::FColumnMetaData& Metadata : Query->SelectionMetaData)
+		for (FColumnMetaData& Metadata : Query->SelectionMetaData)
 		{
 			Metadata.Shrink();
 		}
@@ -403,7 +405,7 @@ namespace TypedElementQueryBuilder
 		checkf(Target, TEXT("The Select section in the Typed Elements query builder doesn't support nullptrs as Read-Only input."));
 		Query.SelectionTypes.Emplace(Target);
 		Query.SelectionAccessTypes.Emplace(ITypedElementDataStorageInterface::EQueryAccessType::ReadOnly);
-		Query.SelectionMetaData.Emplace(Target, TypedElementDataStorage::FColumnMetaData::EFlags::None);
+		Query.SelectionMetaData.Emplace(Target, FColumnMetaData::EFlags::None);
 		
 		return *this;
 	}
@@ -422,11 +424,11 @@ namespace TypedElementQueryBuilder
 		return *this;
 	}
 
-	Select& Select::ReadOnly(const UE::Editor::DataStorage::FDynamicColumnDescription& Description)
+	Select& Select::ReadOnly(const FDynamicColumnDescription& Description)
 	{
 		Query.DynamicSelectionTypes.Emplace(Description);
 		Query.DynamicSelectionAccessTypes.Emplace(ITypedElementDataStorageInterface::EQueryAccessType::ReadOnly);
-		Query.DynamicSelectionMetaData.Emplace(TypedElementDataStorage::FColumnMetaData::EFlags::None);
+		Query.DynamicSelectionMetaData.Emplace(FColumnMetaData::EFlags::None);
 		return *this;
 	}
 
@@ -437,7 +439,7 @@ namespace TypedElementQueryBuilder
 		Query.SelectionAccessTypes.Emplace(Optional == EOptional::Yes
 			? ITypedElementDataStorageInterface::EQueryAccessType::OptionalReadOnly
 			: ITypedElementDataStorageInterface::EQueryAccessType::ReadOnly);
-		Query.SelectionMetaData.Emplace(Target, TypedElementDataStorage::FColumnMetaData::EFlags::None);
+		Query.SelectionMetaData.Emplace(Target, FColumnMetaData::EFlags::None);
 
 		return *this;
 	}
@@ -461,7 +463,7 @@ namespace TypedElementQueryBuilder
 		checkf(Target, TEXT("The Select section in the Typed Elements query builder doesn't support nullptrs as Read/Write input."));
 		Query.SelectionTypes.Emplace(Target);
 		Query.SelectionAccessTypes.Emplace(ITypedElementDataStorageInterface::EQueryAccessType::ReadWrite);
-		Query.SelectionMetaData.Emplace(Target, TypedElementDataStorage::FColumnMetaData::EFlags::IsMutable);
+		Query.SelectionMetaData.Emplace(Target, FColumnMetaData::EFlags::IsMutable);
 		return *this;
 	}
 
@@ -479,11 +481,11 @@ namespace TypedElementQueryBuilder
 		return *this;
 	}
 
-	Select& Select::ReadWrite(const UE::Editor::DataStorage::FDynamicColumnDescription& Descirption)
+	Select& Select::ReadWrite(const FDynamicColumnDescription& Descirption)
 	{
 		Query.DynamicSelectionTypes.Emplace(Descirption);
 		Query.DynamicSelectionAccessTypes.Emplace(ITypedElementDataStorageInterface::EQueryAccessType::ReadWrite);
-		Query.DynamicSelectionMetaData.Emplace(TypedElementDataStorage::FColumnMetaData::EFlags::IsMutable);
+		Query.DynamicSelectionMetaData.Emplace(FColumnMetaData::EFlags::IsMutable);
 		return *this;
 	}
 
diff --git a/Engine/Source/Runtime/TypedElementFramework/Private/Elements/Interfaces/TypedElementDataStorageUiInterface.cpp b/Engine/Source/Runtime/TypedElementFramework/Private/Elements/Interfaces/TypedElementDataStorageUiInterface.cpp
--- a/Engine/Source/Runtime/TypedElementFramework/Private/Elements/Interfaces/TypedElementDataStorageUiInterface.cpp
+++ b/Engine/Source/Runtime/TypedElementFramework/Private/Elements/Interfaces/TypedElementDataStorageUiInterface.cpp
@@ -13,7 +13,7 @@ FTypedElementWidgetConstructor::FTypedElementWidgetConstructor(const UScriptStru
 {
 }
 
-bool FTypedElementWidgetConstructor::Initialize(const TypedElementDataStorage::FMetaDataView& InArguments,
+bool FTypedElementWidgetConstructor::Initialize(const UE::Editor::DataStorage::FMetaDataView& InArguments,
 	TArray<TWeakObjectPtr<const UScriptStruct>> InMatchedColumnTypes, const TypedElementDataStorage::FQueryConditions& InQueryConditions)
 {
 	MatchedColumnTypes = MoveTemp(InMatchedColumnTypes);
@@ -42,7 +42,7 @@ TConstArrayView<const UScriptStruct*> FTypedElementWidgetConstructor::GetAdditio
 }
 
 FString FTypedElementWidgetConstructor::CreateWidgetDisplayName(
-	ITypedElementDataStorageInterface* DataStorage, TypedElementDataStorage::RowHandle Row) const
+	ITypedElementDataStorageInterface* DataStorage, RowHandle Row) const
 {
 	switch (MatchedColumnTypes.Num())
 	{
@@ -97,7 +97,7 @@ TSharedPtr<SWidget> FTypedElementWidgetConstructor::ConstructFinalWidget(
 	RowHandle Row,
 	ITypedElementDataStorageInterface* DataStorage,
 	ITypedElementDataStorageUiInterface* DataStorageUi,
-	const TypedElementDataStorage::FMetaDataView& Arguments)
+	const UE::Editor::DataStorage::FMetaDataView& Arguments)
 {
 	// Add the additional columns to the UI row
 	TSharedPtr<SWidget> Widget = SNullWidget::NullWidget;
@@ -139,9 +139,9 @@ TSharedPtr<SWidget> FTypedElementWidgetConstructor::Construct(
 	RowHandle Row,
 	ITypedElementDataStorageInterface* DataStorage,
 	ITypedElementDataStorageUiInterface* DataStorageUi,
-	const TypedElementDataStorage::FMetaDataView& Arguments)
+	const UE::Editor::DataStorage::FMetaDataView& Arguments)
 {
-	TypedElementDataStorage::RowHandle TargetRow = GetTargetRow(DataStorage, Row);
+	UE::Editor::DataStorage::RowHandle TargetRow = GetTargetRow(DataStorage, Row);
 
 	TSharedPtr<SWidget> Widget = CreateWidget(DataStorage, DataStorageUi, TargetRow, Row, Arguments);
 	if (Widget)
@@ -159,7 +159,7 @@ TSharedPtr<SWidget> FTypedElementWidgetConstructor::Construct(
 	return nullptr;
 }
 
-TSharedPtr<SWidget> FTypedElementWidgetConstructor::CreateWidget(const TypedElementDataStorage::FMetaDataView& Arguments)
+TSharedPtr<SWidget> FTypedElementWidgetConstructor::CreateWidget(const UE::Editor::DataStorage::FMetaDataView& Arguments)
 {
 	return nullptr;
 }
@@ -167,9 +167,9 @@ TSharedPtr<SWidget> FTypedElementWidgetConstructor::CreateWidget(const TypedElem
 TSharedPtr<SWidget> FTypedElementWidgetConstructor::CreateWidget(
 	ITypedElementDataStorageInterface* DataStorage,
 	ITypedElementDataStorageUiInterface* DataStorageUi,
-	TypedElementDataStorage::RowHandle TargetRow,
-	TypedElementDataStorage::RowHandle UiRow, 
-	const TypedElementDataStorage::FMetaDataView& Arguments)
+	UE::Editor::DataStorage::RowHandle TargetRow,
+	UE::Editor::DataStorage::RowHandle UiRow,
+	const UE::Editor::DataStorage::FMetaDataView& Arguments)
 {
 	return CreateWidget(Arguments);
 }
@@ -226,7 +226,8 @@ void FTypedElementWidgetConstructor::AddDefaultWidgetColumns(RowHandle Row, ITyp
 
 FTypedElementWidgetConstructor::RowHandle FTypedElementWidgetConstructor::GetTargetRow(ITypedElementDataStorageInterface* DataStorage, RowHandle WidgetRow) const
 {
-	TypedElementDataStorage::RowHandle TargetRow = TypedElementDataStorage::InvalidRowHandle;
+	using namespace UE::Editor::DataStorage;
+	RowHandle TargetRow = InvalidRowHandle;
 
 	if(const FTypedElementRowReferenceColumn* RowReferenceColumn = DataStorage->GetColumn<FTypedElementRowReferenceColumn>(WidgetRow))
 	{
@@ -245,8 +246,8 @@ FSimpleWidgetConstructor::FSimpleWidgetConstructor(const UScriptStruct* InTypeIn
 }
 
 TSharedPtr<SWidget> FSimpleWidgetConstructor::CreateWidget(ITypedElementDataStorageInterface* DataStorage,
-	ITypedElementDataStorageUiInterface* DataStorageUi, TypedElementDataStorage::RowHandle TargetRow, TypedElementDataStorage::RowHandle WidgetRow,
-	const TypedElementDataStorage::FMetaDataView& Arguments)
+	ITypedElementDataStorageUiInterface* DataStorageUi, UE::Editor::DataStorage::RowHandle TargetRow, UE::Editor::DataStorage::RowHandle WidgetRow,
+	const UE::Editor::DataStorage::FMetaDataView& Arguments)
 {
 	return nullptr;
 }
@@ -256,7 +257,7 @@ bool FSimpleWidgetConstructor::SetColumns(ITypedElementDataStorageInterface* Dat
 	return FTypedElementWidgetConstructor::SetColumns(DataStorage, Row);
 }
 
-TSharedPtr<SWidget> FSimpleWidgetConstructor::CreateWidget(const TypedElementDataStorage::FMetaDataView& Arguments)
+TSharedPtr<SWidget> FSimpleWidgetConstructor::CreateWidget(const UE::Editor::DataStorage::FMetaDataView& Arguments)
 {
 	// This function is not needed anymore and only exists so derived classes cannot derive from it anymore
 	return nullptr;
@@ -270,9 +271,9 @@ bool FSimpleWidgetConstructor::FinalizeWidget(ITypedElementDataStorageInterface*
 }
 
 TSharedPtr<SWidget> FSimpleWidgetConstructor::Construct(RowHandle WidgetRow, ITypedElementDataStorageInterface* DataStorage,
-	ITypedElementDataStorageUiInterface* DataStorageUi, const TypedElementDataStorage::FMetaDataView& Arguments)
+	ITypedElementDataStorageUiInterface* DataStorageUi, const UE::Editor::DataStorage::FMetaDataView& Arguments)
 {
-	const TypedElementDataStorage::RowHandle TargetRow = GetTargetRow(DataStorage, WidgetRow);
+	const UE::Editor::DataStorage::RowHandle TargetRow = GetTargetRow(DataStorage, WidgetRow);
 
 	// Set any required columns on the widget row first
 	SetColumns(DataStorage, WidgetRow);
diff --git a/Engine/Source/Runtime/TypedElementFramework/Public/Elements/Columns/TypedElementAlertColumns.h b/Engine/Source/Runtime/TypedElementFramework/Public/Elements/Columns/TypedElementAlertColumns.h
--- a/Engine/Source/Runtime/TypedElementFramework/Public/Elements/Columns/TypedElementAlertColumns.h
+++ b/Engine/Source/Runtime/TypedElementFramework/Public/Elements/Columns/TypedElementAlertColumns.h
@@ -30,7 +30,7 @@ struct FTypedElementAlertColumn final : public FEditorDataStorageColumn
 	FText Message;
 
 	// Store a copy of the parent row so it's possible to detect if a row has been reparented.
-	TypedElementDataStorage::RowHandle CachedParent;
+	UE::Editor::DataStorage::RowHandle CachedParent;
 
 	UPROPERTY(meta = (IgnoreForMemberInitializationTest))
 	FTypedElementAlertColumnType AlertType;
@@ -45,7 +45,7 @@ struct FTypedElementChildAlertColumn final : public FEditorDataStorageColumn
 	GENERATED_BODY()
 
 	// Store a copy of the parent row so it's possible to detect if a row has been reparented.
-	TypedElementDataStorage::RowHandle CachedParent;
+	UE::Editor::DataStorage::RowHandle CachedParent;
 
 	uint16 Counts[static_cast<size_t>(FTypedElementAlertColumnType::MAX)];
 };
@@ -58,5 +58,5 @@ struct FTypedElementAlertActionColumn final : public FEditorDataStorageColumn
 {
 	GENERATED_BODY()
 
-	TFunction<void(TypedElementDataStorage::RowHandle)> Action;
+	TFunction<void(UE::Editor::DataStorage::RowHandle)> Action;
 };
\ No newline at end of file
diff --git a/Engine/Source/Runtime/TypedElementFramework/Public/Elements/Columns/TypedElementHiearchyColumns.h b/Engine/Source/Runtime/TypedElementFramework/Public/Elements/Columns/TypedElementHiearchyColumns.h
--- a/Engine/Source/Runtime/TypedElementFramework/Public/Elements/Columns/TypedElementHiearchyColumns.h
+++ b/Engine/Source/Runtime/TypedElementFramework/Public/Elements/Columns/TypedElementHiearchyColumns.h
@@ -16,7 +16,7 @@ struct FTableRowParentColumn final : public FEditorDataStorageColumn
 {
 	GENERATED_BODY()
 
-	TypedElementDataStorage::RowHandle Parent;
+	UE::Editor::DataStorage::RowHandle Parent;
 };
 
 /**
diff --git a/Engine/Source/Runtime/TypedElementFramework/Public/Elements/Common/TypedElementQueryDescription.h b/Engine/Source/Runtime/TypedElementFramework/Public/Elements/Common/TypedElementQueryDescription.h
--- a/Engine/Source/Runtime/TypedElementFramework/Public/Elements/Common/TypedElementQueryDescription.h
+++ b/Engine/Source/Runtime/TypedElementFramework/Public/Elements/Common/TypedElementQueryDescription.h
@@ -16,7 +16,7 @@
 
 class UScriptStruct;
 
-namespace TypedElementDataStorage
+namespace UE::Editor::DataStorage
 {
 	struct FQueryDescription;
 	
@@ -59,7 +59,7 @@ namespace TypedElementDataStorage
 		struct FValueTagData
 		{
 			// The Tag maps to a Mass ConstSharedFragment object
-			UE::Editor::DataStorage::FValueTag Tag;
+			FValueTag Tag;
 
 			// The MatchValue specifies the value that the fragment must have to be matched
 			// If MatchValue is NAME_None, then TEDS will match all values
@@ -97,12 +97,12 @@ namespace TypedElementDataStorage
 		TArray<EOperatorType, TInlineAllocator<NumInlineConditions>> ConditionTypes;
 		TArray<FOperator, TInlineAllocator<NumInlineConditions>> ConditionOperators;
 
-		TArray<UE::Editor::DataStorage::FDynamicColumnDescription, TInlineAllocator<NumInlineSelections>> DynamicSelectionTypes;
+		TArray<FDynamicColumnDescription, TInlineAllocator<NumInlineSelections>> DynamicSelectionTypes;
 		TArray<EQueryAccessType, TInlineAllocator<NumInlineSelections>> DynamicSelectionAccessTypes;
 		TArray<FColumnMetaData::EFlags, TInlineAllocator<NumInlineSelections>> DynamicSelectionMetaData;
 		
 		TArray<EOperatorType, TInlineAllocator<NumInlineConditions>> DynamicConditionOperations;
-		TArray<UE::Editor::DataStorage::FDynamicColumnDescription, TInlineAllocator<NumInlineConditions>> DynamicConditionDescriptions;
+		TArray<FDynamicColumnDescription, TInlineAllocator<NumInlineConditions>> DynamicConditionDescriptions;
 
 		TArray<FValueTagData> ValueTags;
 
@@ -115,4 +115,4 @@ namespace TypedElementDataStorage
 
 		EActionType Action;
 	};
-} // namespace TypedElementDataStorage
+} // namespace UE::Editor::DataStorage
diff --git a/Engine/Source/Runtime/TypedElementFramework/Public/Elements/Framework/TypedElementAttributeBinding.h b/Engine/Source/Runtime/TypedElementFramework/Public/Elements/Framework/TypedElementAttributeBinding.h
--- a/Engine/Source/Runtime/TypedElementFramework/Public/Elements/Framework/TypedElementAttributeBinding.h
+++ b/Engine/Source/Runtime/TypedElementFramework/Public/Elements/Framework/TypedElementAttributeBinding.h
@@ -8,10 +8,8 @@
 
 class ITypedElementDataStorageInterface;
 
-namespace UE::EditorDataStorage
+namespace UE::Editor::DataStorage
 {
-	using namespace UE::Editor::DataStorage;
-
 	// Concept for a conversion function used by the attribute binder to bind a column data member to an attribute of a different type
 	template <typename FunctionType, typename ArgumentType>
 	concept AttributeBinderInvocable = std::is_invocable_v<std::decay_t<FunctionType>, const ArgumentType&>;
@@ -22,7 +20,7 @@ namespace UE::EditorDataStorage
 	 * 
 	 * Usage Example:
 	 * 
-	 * TypedElementDataStorage::FAttributeBinder Binder(RowHandle);
+	 * FAttributeBinder Binder(RowHandle);
 	 * TAttribute<int> TestAttribute(Binder.BindData(&FTestColumnInt::TestInt))
 	 */
 	class FAttributeBinder
@@ -30,10 +28,10 @@ namespace UE::EditorDataStorage
 	public:
 
 		/* Create an attribute binder for a given row. */
-		TYPEDELEMENTFRAMEWORK_API explicit FAttributeBinder(TypedElementDataStorage::RowHandle InTargetRow);
+		TYPEDELEMENTFRAMEWORK_API FAttributeBinder(RowHandle InTargetRow);
 		
 		/* Create an attribute binder for a given row. */
-		TYPEDELEMENTFRAMEWORK_API FAttributeBinder(TypedElementDataStorage::RowHandle InTargetRow, ITypedElementDataStorageInterface* InDataStorage);
+		TYPEDELEMENTFRAMEWORK_API FAttributeBinder(RowHandle InTargetRow, ITypedElementDataStorageInterface* InDataStorage);
 
 		/**
 		 * Bind a specific data member inside a TEDS column to an attribute of the same type as the data
@@ -43,7 +41,7 @@ namespace UE::EditorDataStorage
 		 * @return A TAttribute bound to the row, column pair specified
 		 *
 		 * Example:
-		 * TypedElementDataStorage::FAttributeBinder Binder(RowHandle);
+		 * FAttributeBinder Binder(RowHandle);
 		 * TAttribute<FString> TestAttribute(Binder.BindData(&FTypedElementLabelColumn::Label))
 		 */
 		template <typename AttributeType, TDataColumnType ColumnType>
@@ -59,7 +57,7 @@ namespace UE::EditorDataStorage
 		 * @return A TAttribute bound to the row, column pair specified
 		 *
 		 * Example:
-		 * TypedElementDataStorage::FAttributeBinder Binder(RowHandle);
+		 * UE::Editor::DataStorage::FAttributeBinder Binder(RowHandle);
 		 * TAttribute<FText> TestAttribute(Binder.BindData(&FTypedElementLabelColumn::Label),
 		 *                                 [](const FString& Data)
 		 *                                   {
@@ -109,11 +107,11 @@ namespace UE::EditorDataStorage
 	private:
 
 		// The target row for this binder
-		TypedElementDataStorage::RowHandle TargetRow;
+		RowHandle TargetRow;
 
 		// A ptr to the data storage for quick access
 		ITypedElementDataStorageInterface* DataStorage;
 	};
-}
+} // namespace UE::Editor::DataStorage
 
 #include "TypedElementAttributeBinding.inl"
\ No newline at end of file
diff --git a/Engine/Source/Runtime/TypedElementFramework/Public/Elements/Framework/TypedElementAttributeBinding.inl b/Engine/Source/Runtime/TypedElementFramework/Public/Elements/Framework/TypedElementAttributeBinding.inl
--- a/Engine/Source/Runtime/TypedElementFramework/Public/Elements/Framework/TypedElementAttributeBinding.inl
+++ b/Engine/Source/Runtime/TypedElementFramework/Public/Elements/Framework/TypedElementAttributeBinding.inl
@@ -4,10 +4,8 @@
 
 #include "TypedElementRegistry.h"
 
-namespace UE::EditorDataStorage
+namespace UE::Editor::DataStorage
 {
-	using namespace UE::Editor::DataStorage;
-
 	// A property that can be directly be accessed from an object
 	template <typename PropertyType>
 	struct DirectProperty final
@@ -220,6 +218,6 @@ namespace UE::EditorDataStorage
 		});
 	}
 	
-}
+} // namespace UE::Editor::DataStorage
 
 	
\ No newline at end of file
diff --git a/Engine/Source/Runtime/TypedElementFramework/Public/Elements/Framework/TypedElementDataStorageWidget.h b/Engine/Source/Runtime/TypedElementFramework/Public/Elements/Framework/TypedElementDataStorageWidget.h
--- a/Engine/Source/Runtime/TypedElementFramework/Public/Elements/Framework/TypedElementDataStorageWidget.h
+++ b/Engine/Source/Runtime/TypedElementFramework/Public/Elements/Framework/TypedElementDataStorageWidget.h
@@ -20,13 +20,13 @@ class STedsWidget : public SCompoundWidget
 public:
 
 	SLATE_BEGIN_ARGS(STedsWidget)
-		: _UiRowHandle(TypedElementDataStorage::InvalidRowHandle)
+		: _UiRowHandle(UE::Editor::DataStorage::InvalidRowHandle)
 		, _Content()
 	{
 	}
 
 	// The UI Row this widget will be assigned to
-	SLATE_ARGUMENT(TypedElementDataStorage::RowHandle, UiRowHandle)
+	SLATE_ARGUMENT(UE::Editor::DataStorage::RowHandle, UiRowHandle)
 	
 	/** The actual widget content */
 	SLATE_DEFAULT_SLOT(FArguments, Content)
@@ -39,7 +39,7 @@ public:
 	
 	TYPEDELEMENTFRAMEWORK_API void SetContent(const TSharedRef< SWidget >& InContent);
 
-	TYPEDELEMENTFRAMEWORK_API TypedElementDataStorage::RowHandle GetRowHandle() const;
+	TYPEDELEMENTFRAMEWORK_API UE::Editor::DataStorage::RowHandle GetRowHandle() const;
 
 private:
 
@@ -49,5 +49,5 @@ private:
 
 private:
 	
-    TypedElementDataStorage::RowHandle UiRowHandle;
+	UE::Editor::DataStorage::RowHandle UiRowHandle;
 };
diff --git a/Engine/Source/Runtime/TypedElementFramework/Public/Elements/Framework/TypedElementMetaData.h b/Engine/Source/Runtime/TypedElementFramework/Public/Elements/Framework/TypedElementMetaData.h
--- a/Engine/Source/Runtime/TypedElementFramework/Public/Elements/Framework/TypedElementMetaData.h
+++ b/Engine/Source/Runtime/TypedElementFramework/Public/Elements/Framework/TypedElementMetaData.h
@@ -13,7 +13,7 @@
 
 class UScriptStruct;
 
-namespace TypedElementDataStorage
+namespace UE::Editor::DataStorage
 {
 	struct FQueryDescription;
 
@@ -124,14 +124,14 @@ namespace TypedElementDataStorage
 	public:
 		~FQueryMetaDataView() override = default;
 
-		TYPEDELEMENTFRAMEWORK_API explicit FQueryMetaDataView(const TypedElementDataStorage::FQueryDescription& InQuery);
+		TYPEDELEMENTFRAMEWORK_API explicit FQueryMetaDataView(const FQueryDescription& InQuery);
 		
 		TYPEDELEMENTFRAMEWORK_API FMetaDataEntryView FindGeneric(FName AttributeName) const override;
 		TYPEDELEMENTFRAMEWORK_API FMetaDataEntryView FindForColumn(
 			TWeakObjectPtr<const UScriptStruct> Column, FName AttributeName) const override;
 
 	private:
-		const TypedElementDataStorage::FQueryDescription& Query;
+		const FQueryDescription& Query;
 	};
 
 	/**
@@ -201,6 +201,6 @@ namespace TypedElementDataStorage
 		TTuple<ViewTypes...> Views;
 	};
 
-} // TypedElementDataStorage
+} // UE::Editor::DataStorage
 
 #include "Elements/Framework/TypedElementMetaData.inl"
\ No newline at end of file
diff --git a/Engine/Source/Runtime/TypedElementFramework/Public/Elements/Framework/TypedElementMetaData.inl b/Engine/Source/Runtime/TypedElementFramework/Public/Elements/Framework/TypedElementMetaData.inl
--- a/Engine/Source/Runtime/TypedElementFramework/Public/Elements/Framework/TypedElementMetaData.inl
+++ b/Engine/Source/Runtime/TypedElementFramework/Public/Elements/Framework/TypedElementMetaData.inl
@@ -1,8 +1,8 @@
 // Copyright Epic Games, Inc. All Rights Reserved.
 
-namespace TypedElementDataStorage
+namespace UE::Editor::DataStorage
 {
-	namespace Internal
+	namespace Private
 	{
 		template<typename TContainer, typename T>
 		void AddDataToContainer(TContainer& Container, FName Name, T Value)
@@ -43,14 +43,14 @@ namespace TypedElementDataStorage
 		}
 	}
 
-	ENUM_CLASS_FLAGS(TypedElementDataStorage::FColumnMetaData::EFlags)
+	ENUM_CLASS_FLAGS(FColumnMetaData::EFlags)
 
 		template<typename T>
 	bool FMetaDataBase::AddImmutableData(FName Name, T&& Value)
 	{
 		if (!ImmutableData.Contains(Name))
 		{
-			Internal::AddDataToContainer(ImmutableData, Name, Forward<T>(Value));
+			Private::AddDataToContainer(ImmutableData, Name, Forward<T>(Value));
 			return true;
 		}
 		else
@@ -62,7 +62,7 @@ namespace TypedElementDataStorage
 	template<typename T>
 	void FMetaDataBase::AddOrSetMutableData(FName Name, T&& Value)
 	{
-		Internal::AddDataToContainer(MutableData, Name, Forward<T>(Value));
+		Private::AddDataToContainer(MutableData, Name, Forward<T>(Value));
 	}
 
 
@@ -163,4 +163,4 @@ namespace TypedElementDataStorage
 				return Result;
 			});
 	}
-} // namespace TypedElementDataStorage
+} // namespace UE::Editor::DataStorage
diff --git a/Engine/Source/Runtime/TypedElementFramework/Public/Elements/Framework/TypedElementQueryBuilder.h b/Engine/Source/Runtime/TypedElementFramework/Public/Elements/Framework/TypedElementQueryBuilder.h
--- a/Engine/Source/Runtime/TypedElementFramework/Public/Elements/Framework/TypedElementQueryBuilder.h
+++ b/Engine/Source/Runtime/TypedElementFramework/Public/Elements/Framework/TypedElementQueryBuilder.h
@@ -334,9 +334,9 @@ namespace TypedElementQueryBuilder
 
 		inline TypedElementDataStorage::FQueryResult RunQuery(QueryHandle Query) override;
 		inline TypedElementDataStorage::FQueryResult RunSubquery(int32 SubqueryIndex) override;
-		inline TypedElementDataStorage::FQueryResult RunSubquery(int32 SubqueryIndex, TypedElementDataStorage::SubqueryCallbackRef Callback) override;
+		inline TypedElementDataStorage::FQueryResult RunSubquery(int32 SubqueryIndex, SubqueryCallbackRef Callback) override;
 		inline TypedElementDataStorage::FQueryResult RunSubquery(int32 SubqueryIndex, RowHandle Row,
-			TypedElementDataStorage::SubqueryCallbackRef Callback) override;
+			SubqueryCallbackRef Callback) override;
 
 		ITypedElementDataStorageInterface::IQueryContext& ParentContext;
 		const ITypedElementDataStorageInterface::FQueryDescription& Description;
@@ -429,9 +429,9 @@ namespace TypedElementQueryBuilder
 	};
 
 	template<typename Function>
-	TypedElementDataStorage::DirectQueryCallback CreateDirectQueryCallbackBinding(Function&& Callback);
+	DirectQueryCallback CreateDirectQueryCallbackBinding(Function&& Callback);
 	template<typename Function>
-	TypedElementDataStorage::SubqueryCallback CreateSubqueryCallbackBinding(Function&& Callback);
+	SubqueryCallback CreateSubqueryCallbackBinding(Function&& Callback);
 
 } // namespace TypedElementQueryBuilder
 
diff --git a/Engine/Source/Runtime/TypedElementFramework/Public/Elements/Framework/TypedElementQueryBuilder.inl b/Engine/Source/Runtime/TypedElementFramework/Public/Elements/Framework/TypedElementQueryBuilder.inl
--- a/Engine/Source/Runtime/TypedElementFramework/Public/Elements/Framework/TypedElementQueryBuilder.inl
+++ b/Engine/Source/Runtime/TypedElementFramework/Public/Elements/Framework/TypedElementQueryBuilder.inl
@@ -14,7 +14,7 @@ namespace TypedElementQueryBuilder
 {
 	using namespace UE::Editor::DataStorage;
 
-	namespace Internal
+	namespace Private
 	{
 		// This assumes that the types are unique, but for queries this should be true and otherwise
 		// both results would point to the first found index.
@@ -67,7 +67,7 @@ namespace TypedElementQueryBuilder
 			
 			return Result;
 		}
-	} // namespace Internal
+	} // namespace Private
 
 	//
 	// FDependecy
@@ -116,8 +116,8 @@ namespace TypedElementQueryBuilder
 	//
 
 	FQueryContextForwarder::FQueryContextForwarder(
-		const TypedElementDataStorage::FQueryDescription& InDescription,
-		TypedElementDataStorage::IQueryContext& InParentContext)
+		const UE::Editor::DataStorage::FQueryDescription& InDescription,
+		UE::Editor::DataStorage::IQueryContext& InParentContext)
 		: ParentContext(InParentContext)
 		, Description(InDescription)
 	{}
@@ -216,13 +216,13 @@ namespace TypedElementQueryBuilder
 	}
 
 	TypedElementDataStorage::FQueryResult FQueryContextForwarder::RunSubquery(int32 SubqueryIndex,
-		TypedElementDataStorage::SubqueryCallbackRef Callback)
+		UE::Editor::DataStorage::SubqueryCallbackRef Callback)
 	{
 		return ParentContext.RunSubquery(SubqueryIndex, Callback);
 	}
 
 	TypedElementDataStorage::FQueryResult FQueryContextForwarder::RunSubquery(int32 SubqueryIndex,
-		RowHandle Row, TypedElementDataStorage::SubqueryCallbackRef Callback)
+		RowHandle Row, UE::Editor::DataStorage::SubqueryCallbackRef Callback)
 	{
 		return ParentContext.RunSubquery(SubqueryIndex, Row, Callback);
 	}
@@ -233,19 +233,19 @@ namespace TypedElementQueryBuilder
 	
 	template<typename... Dependencies>
 	FCachedQueryContext<Dependencies...>::FCachedQueryContext(
-		const TypedElementDataStorage::FQueryDescription& InDescription,
-		TypedElementDataStorage::IQueryContext& InParentContext)
+		const UE::Editor::DataStorage::FQueryDescription& InDescription,
+		UE::Editor::DataStorage::IQueryContext& InParentContext)
 		: FQueryContextForwarder(InDescription, InParentContext)
 	{}
 
 	template<typename... Dependencies>
-	void FCachedQueryContext<Dependencies...>::Register(TypedElementDataStorage::FQueryDescription& Query)
+	void FCachedQueryContext<Dependencies...>::Register(UE::Editor::DataStorage::FQueryDescription& Query)
 	{
 		Query.DependencyTypes.Reserve(sizeof...(Dependencies));
 		Query.DependencyFlags.Reserve(sizeof...(Dependencies));
 		
 		( Query.DependencyTypes.Emplace(Dependencies::StaticClass()), ... );
-		( Query.DependencyFlags.Emplace(Internal::GetDependencyFlags<Dependencies>()), ... );
+		( Query.DependencyFlags.Emplace(Private::GetDependencyFlags<Dependencies>()), ... );
 		
 		Query.CachedDependencies.AddDefaulted(sizeof...(Dependencies));
 	}
@@ -255,7 +255,7 @@ namespace TypedElementQueryBuilder
 	Dependency& FCachedQueryContext<Dependencies...>::GetCachedMutableDependency()
 	{
 		// Don't allow a dependency registered as const to be found.
-		constexpr uint32 Index = Internal::GetVarArgIndex<std::remove_const_t<Dependency>, Dependencies...>();
+		constexpr uint32 Index = Private::GetVarArgIndex<std::remove_const_t<Dependency>, Dependencies...>();
 		static_assert(Index < sizeof...(Dependencies), "Requested dependency isn't part of the query context cache.");
 		return *static_cast<Dependency*>(Description.CachedDependencies[Index].Get());
 	}
@@ -265,14 +265,14 @@ namespace TypedElementQueryBuilder
 	const Dependency& FCachedQueryContext<Dependencies...>::GetCachedDependency() const
 	{
 		// Allow access to dependencies registered with and without const.
-		constexpr uint32 Index = Internal::GetVarArgIndex<Dependency, Dependencies...>();
+		constexpr uint32 Index = Private::GetVarArgIndex<Dependency, Dependencies...>();
 		if constexpr (Index < sizeof...(Dependencies))
 		{
 			return *static_cast<Dependency*>(Description.CachedDependencies[Index].Get());
 		}
 		else
 		{
-			constexpr uint32 ConstIndex = Internal::GetVarArgIndex<std::add_const_t<Dependency>, Dependencies...>();
+			constexpr uint32 ConstIndex = Private::GetVarArgIndex<std::add_const_t<Dependency>, Dependencies...>();
 			static_assert(ConstIndex < sizeof...(Dependencies), "Requested dependency isn't part of the query context cache.");
 			return *static_cast<Dependency*>(Description.CachedDependencies[ConstIndex].Get());
 		}
@@ -281,7 +281,7 @@ namespace TypedElementQueryBuilder
 	//
 	// Select
 	//
-	namespace Internal
+	namespace Private
 	{
 		template<typename T>
 		concept HasStaticStructMethod = requires
@@ -386,7 +386,7 @@ namespace TypedElementQueryBuilder
 			static constexpr bool bArePointerColumns = AreAllColumnsPointers<ColumnTypes...>();
 
 			template<typename Column>
-			bool CheckValidity(const TypedElementDataStorage::FQueryDescription& Description)
+			bool CheckValidity(const UE::Editor::DataStorage::FQueryDescription& Description)
 			{
 				const UScriptStruct* ColumnType = Column::StaticStruct();
 				if (const TWeakObjectPtr<const UScriptStruct>* Match = Description.SelectionTypes.FindByPredicate(
@@ -424,7 +424,7 @@ namespace TypedElementQueryBuilder
 				}
 			}
 
-			FFunctionColumnInfo(const TypedElementDataStorage::FQueryDescription& Description, SourceContext& Context)
+			FFunctionColumnInfo(const UE::Editor::DataStorage::FQueryDescription& Description, SourceContext& Context)
 			{
 				if constexpr (ValidateColumns)
 				{
@@ -448,7 +448,7 @@ namespace TypedElementQueryBuilder
 		template<SourceQueryContext SourceContext, bool ValidateColumns>
 		struct FFunctionColumnInfo<SourceContext, ValidateColumns>
 		{
-			FFunctionColumnInfo(const TypedElementDataStorage::FQueryDescription&, SourceContext&) {}
+			FFunctionColumnInfo(const UE::Editor::DataStorage::FQueryDescription&, SourceContext&) {}
 			static constexpr bool bArePointerColumns = false;
 		};
 
@@ -459,7 +459,7 @@ namespace TypedElementQueryBuilder
 			using BaseContextType = std::remove_reference_t<TargetType>;
 			BaseContextType ContextWrapper;
 
-			FContextInfo(const TypedElementDataStorage::FQueryDescription& Description, SourceType& Context)
+			FContextInfo(const UE::Editor::DataStorage::FQueryDescription& Description, SourceType& Context)
 				: ContextWrapper(Description, Context)
 			{}
 		};
@@ -469,7 +469,7 @@ namespace TypedElementQueryBuilder
 		{
 			const SourceType& ContextWrapper;
 
-			FContextInfo(const TypedElementDataStorage::FQueryDescription& Description, SourceType& Context)
+			FContextInfo(const UE::Editor::DataStorage::FQueryDescription& Description, SourceType& Context)
 				: ContextWrapper(Context)
 			{}
 		};
@@ -479,7 +479,7 @@ namespace TypedElementQueryBuilder
 		{
 			SourceType& ContextWrapper;
 
-			FContextInfo(const TypedElementDataStorage::FQueryDescription& Description, SourceType& Context)
+			FContextInfo(const UE::Editor::DataStorage::FQueryDescription& Description, SourceType& Context)
 				: ContextWrapper(Context)
 			{}
 		};
@@ -493,7 +493,7 @@ namespace TypedElementQueryBuilder
 			using SuperColumn = FFunctionColumnInfo<SourceContext, ValidateColumns, Args...>;
 			using SuperContext = FContextInfo<SourceContext, TargetContext>;
 
-			FContextRowHandleColumnsFunction(const TypedElementDataStorage::FQueryDescription& Description, SourceContext& Context)
+			FContextRowHandleColumnsFunction(const UE::Editor::DataStorage::FQueryDescription& Description, SourceContext& Context)
 				: SuperColumn(Description, Context)
 				, SuperContext(Description, Context)
 			{}
@@ -546,7 +546,7 @@ namespace TypedElementQueryBuilder
 			using SuperColumn = FFunctionColumnInfo<SourceContext, ValidateColumns, Args...>;
 			using SuperContext = FContextInfo<SourceContext, TargetContext>;
 
-			FContextColumnsFunction(const TypedElementDataStorage::FQueryDescription& Description, SourceContext& Context)
+			FContextColumnsFunction(const UE::Editor::DataStorage::FQueryDescription& Description, SourceContext& Context)
 				: SuperColumn(Description, Context)
 				, SuperContext(Description, Context)
 			{}
@@ -594,7 +594,7 @@ namespace TypedElementQueryBuilder
 		{
 			using Super = FFunctionColumnInfo<SourceContext, ValidateColumns, Columns...>;
 
-			FRowHandleColumnsFunction(const TypedElementDataStorage::FQueryDescription& Description, SourceContext& Context)
+			FRowHandleColumnsFunction(const UE::Editor::DataStorage::FQueryDescription& Description, SourceContext& Context)
 				: Super(Description, Context)
 			{}
 
@@ -642,7 +642,7 @@ namespace TypedElementQueryBuilder
 		{
 			using Super = FFunctionColumnInfo<SourceContext, ValidateColumns, Args...>;
 			
-			FColumnsFunction(const TypedElementDataStorage::FQueryDescription& Description, SourceContext& Context)
+			FColumnsFunction(const UE::Editor::DataStorage::FQueryDescription& Description, SourceContext& Context)
 				: Super(Description, Context)
 			{}
 
@@ -716,105 +716,105 @@ namespace TypedElementQueryBuilder
 		{
 			using Super = typename FFunctionInfoHelper<SourceContext, ValidateColumns, Args...>::BaseClass;
 
-			FFunctionInfo(const TypedElementDataStorage::FQueryDescription& Description, SourceContext& Context)
+			FFunctionInfo(const UE::Editor::DataStorage::FQueryDescription& Description, SourceContext& Context)
 				: Super(Description, Context)
 			{}
 		};
 
 		template<SourceQueryContext SourceContext, bool ValidateColumns, typename... Args>
-		void BindQueryFunction(TypedElementDataStorage::QueryCallback& Function,
+		void BindQueryFunction(UE::Editor::DataStorage::QueryCallback& Function,
 			void (*Callback)(Args...))
 		{
-			Function = [Callback](const TypedElementDataStorage::FQueryDescription& Description, SourceContext& Context)
+			Function = [Callback](const UE::Editor::DataStorage::FQueryDescription& Description, SourceContext& Context)
 			{
 				FFunctionInfo<SourceContext, ValidateColumns, Args...>(Description, Context).Call(Context, Callback);
 			};
 		}
 
 		template<SourceQueryContext SourceContext, bool ValidateColumns, typename Class, typename... Args>
-		void BindQueryFunction(TypedElementDataStorage::QueryCallback& Function, Class* Target,
+		void BindQueryFunction(UE::Editor::DataStorage::QueryCallback& Function, Class* Target,
 			void (Class::*Callback)(Args...))
 		{
-			Function = [Target, Callback](const TypedElementDataStorage::FQueryDescription& Description, SourceContext& Context)
+			Function = [Target, Callback](const UE::Editor::DataStorage::FQueryDescription& Description, SourceContext& Context)
 			{
 				FFunctionInfo<SourceContext, ValidateColumns, Args...>(Description, Context).Call(Context, (Target->*Callback));
 			};
 		}
 
 		template<SourceQueryContext SourceContext, bool ValidateColumns, typename Class, typename... Args>
-		void BindQueryFunction(TypedElementDataStorage::QueryCallback& Function, Class* Target,
+		void BindQueryFunction(UE::Editor::DataStorage::QueryCallback& Function, Class* Target,
 			void (Class::*Callback)(Args...) const)
 		{
-			Function = [Target, Callback](const TypedElementDataStorage::FQueryDescription& Description, SourceContext& Context)
+			Function = [Target, Callback](const UE::Editor::DataStorage::FQueryDescription& Description, SourceContext& Context)
 			{
 				FFunctionInfo<SourceContext, ValidateColumns, Args...>(Description, Context).Call(Context, (Target->*Callback));
 			};
 		}
 
 		template<SourceQueryContext SourceContext, bool ValidateColumns, typename Functor, typename Class, typename... Args>
-		void BindQueryFunction_Expand(TypedElementDataStorage::QueryCallback& Function, Functor&& CallbackObject,
+		void BindQueryFunction_Expand(UE::Editor::DataStorage::QueryCallback& Function, Functor&& CallbackObject,
 			void (Class::*Callback)(Args...) const)
 		{
-			Function = [CallbackObject = std::forward<Functor>(CallbackObject)](const TypedElementDataStorage::FQueryDescription& Description, SourceContext& Context)
+			Function = [CallbackObject = std::forward<Functor>(CallbackObject)](const UE::Editor::DataStorage::FQueryDescription& Description, SourceContext& Context)
 			{
 				FFunctionInfo<SourceContext, ValidateColumns, Args...>(Description, Context).Call(Context, CallbackObject);
 			};
 		}
 
 		template<SourceQueryContext SourceContext, bool ValidateColumns, typename Functor, typename Class, typename... Args>
-		void BindQueryFunction_Expand(TypedElementDataStorage::QueryCallback& Function, Functor&& CallbackObject,
+		void BindQueryFunction_Expand(UE::Editor::DataStorage::QueryCallback& Function, Functor&& CallbackObject,
 			void (Class::*Callback)(Args...))
 		{
-			Function = [CallbackObject = std::forward<Functor>(CallbackObject)](const TypedElementDataStorage::FQueryDescription& Description, SourceContext& Context)
+			Function = [CallbackObject = std::forward<Functor>(CallbackObject)](const UE::Editor::DataStorage::FQueryDescription& Description, SourceContext& Context)
 			{
 				FFunctionInfo<SourceContext, ValidateColumns, Args...>(Description, Context).Call(Context, CallbackObject);
 			};
 		}
 
 		template<SourceQueryContext SourceContext, bool ValidateColumns, typename Functor>
-		void BindQueryFunction(TypedElementDataStorage::QueryCallback& Function, Functor&& Callback)
+		void BindQueryFunction(UE::Editor::DataStorage::QueryCallback& Function, Functor&& Callback)
 		{
 			BindQueryFunction_Expand<SourceContext, ValidateColumns>(Function, std::forward<Functor>(Callback), &Functor::operator());
 		}
 
 		
 		template<bool ValidateColumns, SourceQueryContext SourceContext, typename... Args>
-		void CallQueryFunction(const TypedElementDataStorage::FQueryDescription& Description, SourceContext& Context,
-			TypedElementDataStorage::QueryCallback& Function, void (*Callback)(Args...))
+		void CallQueryFunction(const UE::Editor::DataStorage::FQueryDescription& Description, SourceContext& Context,
+			UE::Editor::DataStorage::QueryCallback& Function, void (*Callback)(Args...))
 		{
 			FFunctionInfo<SourceContext, ValidateColumns, Args...>(Description, Context).Call(Context, Callback);
 		}
 
 		template<bool ValidateColumns, SourceQueryContext SourceContext, typename Class, typename... Args>
-		void CallQueryFunction(const TypedElementDataStorage::FQueryDescription& Description, SourceContext& Context,
+		void CallQueryFunction(const UE::Editor::DataStorage::FQueryDescription& Description, SourceContext& Context,
 			Class* Target, void (Class::* Callback)(Args...))
 		{
 			FFunctionInfo<SourceContext, ValidateColumns, Args...>(Description, Context).Call(Context, (Target->*Callback));
 		}
 
 		template<bool ValidateColumns, SourceQueryContext SourceContext, typename Class, typename... Args>
-		void CallQueryFunction(const TypedElementDataStorage::FQueryDescription& Description, SourceContext& Context,
+		void CallQueryFunction(const UE::Editor::DataStorage::FQueryDescription& Description, SourceContext& Context,
 			Class* Target, void (Class::* Callback)(Args...) const)
 		{
 			FFunctionInfo<SourceContext, ValidateColumns, Args...>(Description, Context).Call(Context, (Target->*Callback));
 		}
 
 		template<bool ValidateColumns, SourceQueryContext SourceContext, typename Functor, typename Class, typename... Args>
-		void CallQueryFunction_Expand(const TypedElementDataStorage::FQueryDescription& Description, SourceContext& Context,
+		void CallQueryFunction_Expand(const UE::Editor::DataStorage::FQueryDescription& Description, SourceContext& Context,
 			const Functor& CallbackObject, void (Class::* Callback)(Args...) const)
 		{
 			FFunctionInfo<SourceContext, ValidateColumns, Args...>(Description, Context).Call(Context, CallbackObject);
 		}
 
 		template<bool ValidateColumns, SourceQueryContext SourceContext, typename Functor, typename Class, typename... Args>
-		void CallQueryFunction_Expand(const TypedElementDataStorage::FQueryDescription& Description, SourceContext& Context,
+		void CallQueryFunction_Expand(const UE::Editor::DataStorage::FQueryDescription& Description, SourceContext& Context,
 			const Functor& CallbackObject, void (Class::* Callback)(Args...))
 		{
 			FFunctionInfo<SourceContext, ValidateColumns, Args...>(Description, Context).Call(Context, CallbackObject);
 		}
 
 		template<bool ValidateColumns, SourceQueryContext SourceContext, typename Functor>
-		void CallQueryFunction(const TypedElementDataStorage::FQueryDescription& Description, SourceContext& Context,
+		void CallQueryFunction(const UE::Editor::DataStorage::FQueryDescription& Description, SourceContext& Context,
 			const Functor& Callback)
 		{
 			CallQueryFunction_Expand<ValidateColumns>(Description, Context, Callback, &Functor::operator());
@@ -834,7 +834,7 @@ namespace TypedElementQueryBuilder
 			}
 		}
 
-		template<typename ContextType> void RegisterDependencies(TypedElementDataStorage::FQueryDescription& Query)
+		template<typename ContextType> void RegisterDependencies(UE::Editor::DataStorage::FQueryDescription& Query)
 		{
 			using BaseType = typename TRemoveReference<ContextType>::Type;
 			if constexpr (TIsDerivedFrom<BaseType, FQueryContextForwarder>::Value)
@@ -846,7 +846,7 @@ namespace TypedElementQueryBuilder
 		template<SourceQueryContext SourceContext, typename... Args>
 		struct RegisterFunctionArgumentsHelper
 		{
-			void Register(TypedElementDataStorage::FQueryDescription&, Select& Target)
+			void Register(UE::Editor::DataStorage::FQueryDescription&, Select& Target)
 			{
 				(AddColumnToSelect<Args>(Target), ...);
 			}
@@ -855,13 +855,13 @@ namespace TypedElementQueryBuilder
 		template<SourceQueryContext SourceContext>
 		struct RegisterFunctionArgumentsHelper<SourceContext>
 		{
-			static void Register(TypedElementDataStorage::FQueryDescription&, Select&) {}
+			static void Register(UE::Editor::DataStorage::FQueryDescription&, Select&) {}
 		};
 
 		template<SourceQueryContext SourceContext, typename Arg0>
 		struct RegisterFunctionArgumentsHelper<SourceContext, Arg0>
 		{
-			static void Register(TypedElementDataStorage::FQueryDescription& Query, Select& Target)
+			static void Register(UE::Editor::DataStorage::FQueryDescription& Query, Select& Target)
 			{
 				if constexpr (IsCompatibleTargetContextType<SourceContext, Arg0>)
 				{
@@ -878,7 +878,7 @@ namespace TypedElementQueryBuilder
 		template<SourceQueryContext SourceContext, typename Arg0, typename Arg1, typename... Args>
 		struct RegisterFunctionArgumentsHelper<SourceContext, Arg0, Arg1, Args...>
 		{
-			static void Register(TypedElementDataStorage::FQueryDescription& Query, Select& Target)
+			static void Register(UE::Editor::DataStorage::FQueryDescription& Query, Select& Target)
 			{
 				if constexpr (IsCompatibleTargetContextType<SourceContext, Arg0>)
 				{
@@ -901,25 +901,25 @@ namespace TypedElementQueryBuilder
 		};
 
 		template<SourceQueryContext SourceContext, typename... Args>
-		void RegisterFunctionArguments(TypedElementDataStorage::FQueryDescription& Query, Select& Target, void (*)(Args...))
+		void RegisterFunctionArguments(UE::Editor::DataStorage::FQueryDescription& Query, Select& Target, void (*)(Args...))
 		{
 			RegisterFunctionArgumentsHelper<SourceContext, Args...>::Register(Query, Target);
 		}
 
 		template<SourceQueryContext SourceContext, typename Class, typename... Args>
-		void RegisterFunctionArguments(TypedElementDataStorage::FQueryDescription& Query, Select& Target, void (Class::*)(Args...))
+		void RegisterFunctionArguments(UE::Editor::DataStorage::FQueryDescription& Query, Select& Target, void (Class::*)(Args...))
 		{
 			RegisterFunctionArgumentsHelper<SourceContext, Args...>::Register(Query, Target);
 		}
 
 		template<SourceQueryContext SourceContext, typename Class, typename... Args>
-		void RegisterFunctionArguments(TypedElementDataStorage::FQueryDescription& Query, Select& Target, void (Class::*)(Args...) const)
+		void RegisterFunctionArguments(UE::Editor::DataStorage::FQueryDescription& Query, Select& Target, void (Class::*)(Args...) const)
 		{
 			RegisterFunctionArgumentsHelper<SourceContext, Args...>::Register(Query, Target);
 		}
 
 		template<SourceQueryContext SourceContext, typename Functor>
-		void RegisterFunctionArguments(TypedElementDataStorage::FQueryDescription& Query, Select& Target, Functor)
+		void RegisterFunctionArguments(UE::Editor::DataStorage::FQueryDescription& Query, Select& Target, Functor)
 		{
 			RegisterFunctionArguments<SourceContext>(Query, Target, &Functor::operator());
 		}
@@ -1025,7 +1025,7 @@ namespace TypedElementQueryBuilder
 			}
 		};
 		
-		inline void PrepareForQueryBinding(TypedElementDataStorage::FQueryDescription& Query, const FProcessor& Processor)
+		inline void PrepareForQueryBinding(UE::Editor::DataStorage::FQueryDescription& Query, const FProcessor& Processor)
 		{
 			Query.Callback.Type = TypedElementDataStorage::EQueryCallbackType::Processor;
 			Query.Callback.Phase = Processor.Phase;
@@ -1043,7 +1043,7 @@ namespace TypedElementQueryBuilder
 			Query.Callback.ExecutionMode = Processor.ExecutionMode;
 		}
 
-		inline void PrepareForQueryBinding(TypedElementDataStorage::FQueryDescription& Query, const FObserver& Observer)
+		inline void PrepareForQueryBinding(UE::Editor::DataStorage::FQueryDescription& Query, const FObserver& Observer)
 		{
 			switch (Observer.Event)
 			{
@@ -1060,7 +1060,7 @@ namespace TypedElementQueryBuilder
 			Query.Callback.ExecutionMode = Observer.ExecutionMode;
 		}
 
-		inline void PrepareForQueryBinding(TypedElementDataStorage::FQueryDescription& Query, const FPhaseAmble& PhaseAmble)
+		inline void PrepareForQueryBinding(UE::Editor::DataStorage::FQueryDescription& Query, const FPhaseAmble& PhaseAmble)
 		{
 			switch (PhaseAmble.Location)
 			{
@@ -1078,12 +1078,12 @@ namespace TypedElementQueryBuilder
 		}
 		
 		template<typename CallbackType, typename Function>
-		void PrepareForQueryBinding(Select& Target, TypedElementDataStorage::FQueryDescription& Query, FName Name,
+		void PrepareForQueryBinding(Select& Target, UE::Editor::DataStorage::FQueryDescription& Query, FName Name,
 			const CallbackType& Type, Function Callback)
 		{
 			static_assert(TIsDerivedFrom<CallbackType, FQueryCallbackType>::Value, "The callback type provided isn't one of the available "
 				"classes derived from FQueryCallbackType.");
-			static_assert(IsValidSelectFunctionSignature<TypedElementDataStorage::IQueryContext, Function>(),
+			static_assert(IsValidSelectFunctionSignature<UE::Editor::DataStorage::IQueryContext, Function>(),
 				R"(The function provided to the Query Builder's Select call wasn't invocable or doesn't contain a supported combination of arguments.
 The following options are supported:
 - void([const]Column&...) 
@@ -1092,22 +1092,22 @@ The following options are supported:
 - void(<Context>&, RowHandle, [const]Column&...) 
 - void(<Context>&, [const]Column*...) 
 - void(<Context>&, const RowHandle*, [const]Column*...) 
-Where <Context> is TypedElementDataStorage::IQueryContext or FCachedQueryContext<...>
+Where <Context> is UE::Editor::DataStorage::IQueryContext or FCachedQueryContext<...>
 e.g. void(FCachedQueryContext<Subsystem1, const Subsystem2>& Context, RowHandle Row, ColumnType0& ColumnA, const ColumnType1& ColumnB) {...}
 )");
-			RegisterFunctionArguments<TypedElementDataStorage::IQueryContext>(Query, Target, Callback);
+			RegisterFunctionArguments<UE::Editor::DataStorage::IQueryContext>(Query, Target, Callback);
 			PrepareForQueryBinding(Query, Type);
 			Query.Callback.Name = Name;
 		}
-	} // namespace Internal
+	} // namespace Private
 
 	template<typename CallbackType, typename Function>
 	Select::Select(FName Name, const CallbackType& Type, Function&& Callback)
 		: Select()
 	{
 		static constexpr bool ValidateColumns = false;
-		Internal::PrepareForQueryBinding(*this, Query, Name, Type, Callback);
-		Internal::BindQueryFunction<TypedElementDataStorage::IQueryContext, ValidateColumns>(Query.Callback.Function, std::forward<Function>(Callback));
+		Private::PrepareForQueryBinding(*this, Query, Name, Type, Callback);
+		Private::BindQueryFunction<UE::Editor::DataStorage::IQueryContext, ValidateColumns>(Query.Callback.Function, std::forward<Function>(Callback));
 	}
 
 	template<typename CallbackType, typename Class, typename Function>
@@ -1115,8 +1115,8 @@ e.g. void(FCachedQueryContext<Subsystem1, const Subsystem2>& Context, RowHandle
 		: Select()
 	{
 		static constexpr bool ValidateColumns = false;
-		Internal::PrepareForQueryBinding(*this, Query, Name, Type, Callback);
-		Internal::BindQueryFunction<TypedElementDataStorage::IQueryContext, ValidateColumns>(Query.Callback.Function, Instance, Callback);
+		Private::PrepareForQueryBinding(*this, Query, Name, Type, Callback);
+		Private::BindQueryFunction<UE::Editor::DataStorage::IQueryContext, ValidateColumns>(Query.Callback.Function, Instance, Callback);
 	}
 
 	template<TDataColumnType... TargetTypes>
@@ -1252,9 +1252,9 @@ e.g. void(FCachedQueryContext<Subsystem1, const Subsystem2>& Context, RowHandle
 	// External query bindings
 	//
 	template<typename Function>
-	TypedElementDataStorage::DirectQueryCallback CreateDirectQueryCallbackBinding(Function&& Callback)
+	UE::Editor::DataStorage::DirectQueryCallback CreateDirectQueryCallbackBinding(Function&& Callback)
 	{
-		static_assert(Internal::IsValidSelectFunctionSignature<TypedElementDataStorage::IDirectQueryContext, Function>(),
+		static_assert(Private::IsValidSelectFunctionSignature<UE::Editor::DataStorage::IDirectQueryContext, Function>(),
 			R"(The function provided to the Query Builder's CreateDirectQueryCallbackBinding call wasn't invocable or doesn't contain a supported combination of arguments.
 The following options are supported:
 - void([const]Column&...) 
@@ -1263,23 +1263,23 @@ The following options are supported:
 - void(<Context>&, RowHandle, [const]Column&...) 
 - void(<Context>&, [const]Column*...) 
 - void(<Context>&, const RowHandle*, [const]Column*...) 
-Where <Context> is TypedElementDataStorage::IDirectQueryContext
+Where <Context> is UE::Editor::DataStorage::IDirectQueryContext
 e.g. void(IDirectQueryContext& Context, RowHandle Row, ColumnType0& ColumnA, const ColumnType1& ColumnB) {...}
 )");
 
 		return [Callback = Forward<Function>(Callback)](
-			const TypedElementDataStorage::FQueryDescription& Description,
-			TypedElementDataStorage::IDirectQueryContext& Context)
+			const UE::Editor::DataStorage::FQueryDescription& Description,
+			UE::Editor::DataStorage::IDirectQueryContext& Context)
 		{
 			static constexpr bool ValidateColumns = true;
-			Internal::CallQueryFunction<ValidateColumns>(Description, Context, Callback);
+			Private::CallQueryFunction<ValidateColumns>(Description, Context, Callback);
 		};
 	}
 
 	template<typename Function>
-	TypedElementDataStorage::SubqueryCallback CreateSubqueryCallbackBinding(Function&& Callback)
+	UE::Editor::DataStorage::SubqueryCallback CreateSubqueryCallbackBinding(Function&& Callback)
 	{
-		static_assert(Internal::IsValidSelectFunctionSignature<ITypedElementDataStorageInterface::ISubqueryContext, Function>(),
+		static_assert(Private::IsValidSelectFunctionSignature<ITypedElementDataStorageInterface::ISubqueryContext, Function>(),
 			R"(The function provided to the Query Builder's CreateSubqueryCallbackBinding call wasn't invocable or doesn't contain a supported combination of arguments.
 The following options are supported:
 - void([const]Column&...) 
@@ -1288,16 +1288,16 @@ The following options are supported:
 - void(<Context>&, RowHandle, [const]Column&...) 
 - void(<Context>&, [const]Column*...) 
 - void(<Context>&, const RowHandle*, [const]Column*...) 
-Where <Context> is TypedElementDataStorage::ISubqueryContext
+Where <Context> is UE::Editor::DataStorage::ISubqueryContext
 e.g. void(ISubqueryContext& Context, RowHandle Row, ColumnType0& ColumnA, const ColumnType1& ColumnB) {...}
 )");
 
 		return [Callback = Forward<Function>(Callback)](
-			const TypedElementDataStorage::FQueryDescription& Description,
-			TypedElementDataStorage::ISubqueryContext& Context)
+			const UE::Editor::DataStorage::FQueryDescription& Description,
+			UE::Editor::DataStorage::ISubqueryContext& Context)
 		{
 			static constexpr bool ValidateColumns = true;
-			Internal::CallQueryFunction<ValidateColumns>(Description, Context, Callback);
+			Private::CallQueryFunction<ValidateColumns>(Description, Context, Callback);
 		};
 	}
 } // namespace TypedElementQueryBuilder
