[Teds] Update namespace from TypedElementDataStorage to UE::Editor::DataStorage for metadata and query store.

diff --git a/Engine/Source/Runtime/TypedElementFramework/Public/Elements/Interfaces/TypedElementDataStorageInterface.h b/Engine/Source/Runtime/TypedElementFramework/Public/Elements/Interfaces/TypedElementDataStorageInterface.h
--- a/Engine/Source/Runtime/TypedElementFramework/Public/Elements/Interfaces/TypedElementDataStorageInterface.h
+++ b/Engine/Source/Runtime/TypedElementFramework/Public/Elements/Interfaces/TypedElementDataStorageInterface.h
@@ -200,7 +200,7 @@ public:
 	virtual void AddColumn(RowHandle Row, const UE::Editor::DataStorage::FValueTag& Tag, const FName& Value) = 0;
 
 	template<typename T>
-	void AddColumn(TypedElementDataStorage::RowHandle Row, const FName& Tag) = delete;
+	void AddColumn(RowHandle Row, const FName& Tag) = delete;
 	
 	template<typename T>
 	void AddColumn(RowHandle Row, const FName& Tag, const FName& Value) = delete;
@@ -215,10 +215,10 @@ public:
 	void AddColumn(RowHandle Row);
 
 	template<UE::Editor::DataStorage::TColumnType DynamicColumnTemplate>
-	void AddColumn(TypedElementDataStorage::RowHandle Row, const FName& Identifier);
+	void AddColumn(RowHandle Row, const FName& Identifier);
 	
 	template<UE::Editor::DataStorage::TColumnType DynamicColumnTemplate>
-	void AddColumn(TypedElementDataStorage::RowHandle Row, const FName& Identifier, DynamicColumnTemplate&& TemplateInstance);
+	void AddColumn(RowHandle Row, const FName& Identifier, DynamicColumnTemplate&& TemplateInstance);
 
 	/**
 	 * Adds multiple columns from a row. This is typically more efficient than adding columns one 
@@ -246,7 +246,7 @@ public:
 	void RemoveColumn(RowHandle Row, const FName& Tag) = delete;
 
 	template<UE::Editor::DataStorage::TColumnType DynamicColumnTemplateType>
-	void RemoveColumn(TypedElementDataStorage::RowHandle Row, const FName& Identifier);
+	void RemoveColumn(RowHandle Row, const FName& Identifier);
 	
 	template<>
 	void RemoveColumn<UE::Editor::DataStorage::FValueTag>(RowHandle Row, const FName& Tag);
@@ -335,17 +335,17 @@ public:
 	using EQueryDependencyFlags = TypedElementDataStorage::EQueryDependencyFlags;
 	using FQueryResult = TypedElementDataStorage::FQueryResult;
 
-	using IQueryContext = TypedElementDataStorage::IQueryContext;
-	using IDirectQueryContext = TypedElementDataStorage::IDirectQueryContext;
-	using ISubqueryContext = TypedElementDataStorage::ISubqueryContext;
+	using IQueryContext = UE::Editor::DataStorage::IQueryContext;
+	using IDirectQueryContext = UE::Editor::DataStorage::IDirectQueryContext;
+	using ISubqueryContext = UE::Editor::DataStorage::ISubqueryContext;
 
-	using FQueryDescription = TypedElementDataStorage::FQueryDescription;
-	using QueryCallback = TypedElementDataStorage::QueryCallback;
-	using QueryCallbackRef = TypedElementDataStorage::QueryCallbackRef;
-	using DirectQueryCallback = TypedElementDataStorage::DirectQueryCallback;
-	using DirectQueryCallbackRef = TypedElementDataStorage::DirectQueryCallbackRef;
-	using SubqueryCallback = TypedElementDataStorage::SubqueryCallback;
-	using SubqueryCallbackRef = TypedElementDataStorage::SubqueryCallbackRef;
+	using FQueryDescription = UE::Editor::DataStorage::FQueryDescription;
+	using QueryCallback = UE::Editor::DataStorage::QueryCallback;
+	using QueryCallbackRef = UE::Editor::DataStorage::QueryCallbackRef;
+	using DirectQueryCallback = UE::Editor::DataStorage::DirectQueryCallback;
+	using DirectQueryCallbackRef = UE::Editor::DataStorage::DirectQueryCallbackRef;
+	using SubqueryCallback = UE::Editor::DataStorage::SubqueryCallback;
+	using SubqueryCallbackRef = UE::Editor::DataStorage::SubqueryCallbackRef;
 
 	/** 
 	 * Registers a query with the data storage. The description is processed into an internal format and may be changed. If no valid
@@ -374,7 +374,7 @@ public:
 	 * Directly runs a query. The callback will be called for batches of matching rows. During a single call to RunQuery the callback
 	 * may be called multiple times. If the query handle is invalid or has been deleted nothing happens and the callback won't be called.
 	 */
-	virtual FQueryResult RunQuery(QueryHandle Query, TypedElementDataStorage::EDirectQueryExecutionFlags Flags, 
+	virtual FQueryResult RunQuery(QueryHandle Query, UE::Editor::DataStorage::EDirectQueryExecutionFlags Flags,
 		DirectQueryCallbackRef Callback) = 0;
 	/**
 	 * Triggers all queries registered under the activation name to run for one update cycle. The activatable queries will be activated at
@@ -496,7 +496,7 @@ inline void ITypedElementDataStorageInterface::RemoveColumn<UE::Editor::DataStor
 }
 
 template<UE::Editor::DataStorage::TColumnType DynamicColumnTemplate>
-void ITypedElementDataStorageInterface::RemoveColumn(TypedElementDataStorage::RowHandle Row, const FName& Identifier)
+void ITypedElementDataStorageInterface::RemoveColumn(RowHandle Row, const FName& Identifier)
 {
 	const UE::Editor::DataStorage::FDynamicColumnDescription Description
 	{
@@ -525,7 +525,7 @@ void ITypedElementDataStorageInterface::AddColumn(RowHandle Row)
 }
 
 template <UE::Editor::DataStorage::TColumnType DynamicColumnTemplate>
-void ITypedElementDataStorageInterface::AddColumn(TypedElementDataStorage::RowHandle Row, const FName& Identifier)
+void ITypedElementDataStorageInterface::AddColumn(RowHandle Row, const FName& Identifier)
 {
 	static_assert(UE::Editor::DataStorage::TDataColumnType<DynamicColumnTemplate> || UE::Editor::DataStorage::TTagColumnType<DynamicColumnTemplate>,
 		"DynamicColumnTemplate must be derived from either a Tag or Column");
@@ -539,7 +539,7 @@ void ITypedElementDataStorageInterface::AddColumn(TypedElementDataStorage::RowHa
 }
 
 template <UE::Editor::DataStorage::TColumnType DynamicColumnTemplate>
-void ITypedElementDataStorageInterface::AddColumn(TypedElementDataStorage::RowHandle Row, const FName& Identifier, DynamicColumnTemplate&& TemplateInstance)
+void ITypedElementDataStorageInterface::AddColumn(RowHandle Row, const FName& Identifier, DynamicColumnTemplate&& TemplateInstance)
 {
 	static_assert(UE::Editor::DataStorage::TDataColumnType<DynamicColumnTemplate> || UE::Editor::DataStorage::TTagColumnType<DynamicColumnTemplate>,
 		"DynamicColumnTemplate must be derived from either a Tag or Column");
diff --git a/Engine/Source/Runtime/TypedElementFramework/Public/Elements/Interfaces/TypedElementDataStorageUiInterface.h b/Engine/Source/Runtime/TypedElementFramework/Public/Elements/Interfaces/TypedElementDataStorageUiInterface.h
--- a/Engine/Source/Runtime/TypedElementFramework/Public/Elements/Interfaces/TypedElementDataStorageUiInterface.h
+++ b/Engine/Source/Runtime/TypedElementFramework/Public/Elements/Interfaces/TypedElementDataStorageUiInterface.h
@@ -35,15 +35,15 @@ public:
 	virtual ~FTypedElementWidgetConstructor() = default;
 
 	/** Initializes a new constructor based on the provided arguments.. */
-	TYPEDELEMENTFRAMEWORK_API virtual bool Initialize(const TypedElementDataStorage::FMetaDataView& InArguments,
-		TArray<TWeakObjectPtr<const UScriptStruct>> InMatchedColumnTypes, const TypedElementDataStorage::FQueryConditions& InQueryConditions);
+	TYPEDELEMENTFRAMEWORK_API virtual bool Initialize(const UE::Editor::DataStorage::FMetaDataView& InArguments,
+		TArray<TWeakObjectPtr<const UScriptStruct>> InMatchedColumnTypes, const UE::Editor::DataStorage::FQueryConditions& InQueryConditions);
 
 	/** Retrieves the type information for the constructor type. */
 	TYPEDELEMENTFRAMEWORK_API virtual const UScriptStruct* GetTypeInfo() const;
 	/** Retrieves the columns, if any, that were matched to this constructor when it was created. */
 	TYPEDELEMENTFRAMEWORK_API virtual const TArray<TWeakObjectPtr<const UScriptStruct>>& GetMatchedColumns() const;
 	/** Retrieves the query conditions that need to match for this widget constructor to produce a widget. */
-	TYPEDELEMENTFRAMEWORK_API virtual const TypedElementDataStorage::FQueryConditions* GetQueryConditions() const;
+	TYPEDELEMENTFRAMEWORK_API virtual const UE::Editor::DataStorage::FQueryConditions* GetQueryConditions() const;
 
 	/** Returns a list of additional columns the widget requires to be added to its rows. */
 	TYPEDELEMENTFRAMEWORK_API virtual TConstArrayView<const UScriptStruct*> GetAdditionalColumnsList() const;
@@ -55,7 +55,7 @@ public:
 	 * Individual widget constructors can override this function with a name specific to them.
 	 */
 	TYPEDELEMENTFRAMEWORK_API virtual FString CreateWidgetDisplayName(
-		ITypedElementDataStorageInterface* DataStorage, TypedElementDataStorage::RowHandle Row) const;
+		ITypedElementDataStorageInterface* DataStorage, RowHandle Row) const;
 	
 	/**
 	 *	Calls Construct() to create the internal widget, and then stores it in a container before returning.
@@ -69,7 +69,7 @@ public:
 		RowHandle Row, /** The row the widget will be stored in. */
 		ITypedElementDataStorageInterface* DataStorage,
 		ITypedElementDataStorageUiInterface* DataStorageUi,
-		const TypedElementDataStorage::FMetaDataView& Arguments);
+		const UE::Editor::DataStorage::FMetaDataView& Arguments);
 
 	/**
 	 * Constructs the widget according to the provided information. Information is collected by calling
@@ -82,18 +82,18 @@ public:
 		RowHandle Row, /** The row the widget will be stored in. */
 		ITypedElementDataStorageInterface* DataStorage,
 		ITypedElementDataStorageUiInterface* DataStorageUi,
-		const TypedElementDataStorage::FMetaDataView& Arguments);
+		const UE::Editor::DataStorage::FMetaDataView& Arguments);
 
 protected:
 	/** Create a new instance of the target widget. This is a required function. */
-	TYPEDELEMENTFRAMEWORK_API virtual TSharedPtr<SWidget> CreateWidget(const TypedElementDataStorage::FMetaDataView& Arguments);
+	TYPEDELEMENTFRAMEWORK_API virtual TSharedPtr<SWidget> CreateWidget(const UE::Editor::DataStorage::FMetaDataView& Arguments);
 	/** Create a new instance of the target widget. This is a required function. */
 	TYPEDELEMENTFRAMEWORK_API virtual TSharedPtr<SWidget> CreateWidget(
 		ITypedElementDataStorageInterface* DataStorage,
 		ITypedElementDataStorageUiInterface* DataStorageUi,
-		TypedElementDataStorage::RowHandle TargetRow,
-		TypedElementDataStorage::RowHandle WidgetRow, 
-		const TypedElementDataStorage::FMetaDataView& Arguments);
+		UE::Editor::DataStorage::RowHandle TargetRow,
+		UE::Editor::DataStorage::RowHandle WidgetRow, 
+		const UE::Editor::DataStorage::FMetaDataView& Arguments);
 	/** Set any values in columns if needed. The columns provided through GetAdditionalColumnsList() will have already been created. */
 	TYPEDELEMENTFRAMEWORK_API virtual bool SetColumns(ITypedElementDataStorageInterface* DataStorage, RowHandle Row);
 	
@@ -122,7 +122,7 @@ protected:
 protected:
 
 	TArray<TWeakObjectPtr<const UScriptStruct>> MatchedColumnTypes;
-	const TypedElementDataStorage::FQueryConditions* QueryConditions = nullptr;
+	const UE::Editor::DataStorage::FQueryConditions* QueryConditions = nullptr;
 	const UScriptStruct* TypeInfo = nullptr;
 };
 
@@ -155,9 +155,9 @@ struct FSimpleWidgetConstructor : public FTypedElementWidgetConstructor
 	TYPEDELEMENTFRAMEWORK_API virtual TSharedPtr<SWidget> CreateWidget(
 		ITypedElementDataStorageInterface* DataStorage,
 		ITypedElementDataStorageUiInterface* DataStorageUi,
-		RowHandle TargetRow,
-		RowHandle WidgetRow, 
-		const TypedElementDataStorage::FMetaDataView& Arguments) override;
+		UE::Editor::DataStorage::RowHandle TargetRow,
+		UE::Editor::DataStorage::RowHandle WidgetRow,
+		const UE::Editor::DataStorage::FMetaDataView& Arguments) override;
 
 	/*
 	 * Override this function to add any columns to the WidgetRow before CreateWidget is called
@@ -171,7 +171,7 @@ struct FSimpleWidgetConstructor : public FTypedElementWidgetConstructor
 protected:
 
 	/** Old CreateWidget overload that exists for backwards compatibility, you should use the overload that provides the row instead */
-	TYPEDELEMENTFRAMEWORK_API virtual TSharedPtr<SWidget> CreateWidget(const TypedElementDataStorage::FMetaDataView& Arguments) override final;
+	TYPEDELEMENTFRAMEWORK_API virtual TSharedPtr<SWidget> CreateWidget(const UE::Editor::DataStorage::FMetaDataView& Arguments) override final;
 	
 	/** 
 	 * Old function in the widget creation pipeline that isn't used anymore. All your logic should go in CreateWidget() itself
@@ -189,7 +189,7 @@ protected:
 		RowHandle WidgetRow, 
 		ITypedElementDataStorageInterface* DataStorage,
 		ITypedElementDataStorageUiInterface* DataStorageUi,
-		const TypedElementDataStorage::FMetaDataView& Arguments) override final;
+		const UE::Editor::DataStorage::FMetaDataView& Arguments) override final;
 };
 
 template<>
@@ -284,7 +284,7 @@ public:
 	 * If registration is successful true will be returned otherwise false.
 	 */
 	virtual bool RegisterWidgetFactory(FName Purpose, const UScriptStruct* Constructor,
-		TypedElementDataStorage::FQueryConditions Columns) = 0;
+		UE::Editor::DataStorage::FQueryConditions Columns) = 0;
 	/**
 	 * Registers a widget factory that will be called when the purpose it's registered under is requested.
 	 * This version registers a generic type. Construction using these are typically cheaper as they can avoid
@@ -294,7 +294,7 @@ public:
 	 * If registration is successful true will be returned otherwise false.
 	 */
 	template<typename ConstructorType>
-	bool RegisterWidgetFactory(FName Purpose, TypedElementDataStorage::FQueryConditions Columns);
+	bool RegisterWidgetFactory(FName Purpose, UE::Editor::DataStorage::FQueryConditions Columns);
 	/**
 	 * Registers a widget factory that will be called when the purpose it's registered under is requested.
 	 * This version uses a previously created instance of the Constructor. The benefit of this is that it store
@@ -312,7 +312,7 @@ public:
 	 * If registration is successful true will be returned otherwise false.
 	 */
 	virtual bool RegisterWidgetFactory(FName Purpose, TUniquePtr<FTypedElementWidgetConstructor>&& Constructor, 
-		TypedElementDataStorage::FQueryConditions Columns) = 0;
+		UE::Editor::DataStorage::FQueryConditions Columns) = 0;
 	
 	/** 
 	 * Creates widget constructors for the requested purpose.
@@ -320,7 +320,7 @@ public:
 	 * widgets created from the constructor, if applicable.
 	 */
 	virtual void CreateWidgetConstructors(FName Purpose, 
-		const TypedElementDataStorage::FMetaDataView& Arguments, const WidgetConstructorCallback& Callback) = 0;
+		const UE::Editor::DataStorage::FMetaDataView& Arguments, const WidgetConstructorCallback& Callback) = 0;
 	/** 
 	 * Finds matching widget constructors for provided columns, preferring longer matches over shorter matches.
 	 * The provided list of columns will be updated to contain all columns that couldn't be matched.
@@ -328,14 +328,14 @@ public:
 	 * widgets created from the constructor, if applicable.
 	 */
 	virtual void CreateWidgetConstructors(FName Purpose, EMatchApproach MatchApproach, TArray<TWeakObjectPtr<const UScriptStruct>>& Columns,
-		const TypedElementDataStorage::FMetaDataView& Arguments, const WidgetConstructorCallback& Callback) = 0;
+		const UE::Editor::DataStorage::FMetaDataView& Arguments, const WidgetConstructorCallback& Callback) = 0;
 
 	/**
 	 * Creates all the widgets registered under the provided name. This may be a large number of widgets for a wide name
 	 * or exactly one when the exact name of the widget is registered. Arguments can be provided, but widgets are free
 	 * to ignore them.
 	 */
-	virtual void ConstructWidgets(FName Purpose, const TypedElementDataStorage::FMetaDataView& Arguments,
+	virtual void ConstructWidgets(FName Purpose, const UE::Editor::DataStorage::FMetaDataView& Arguments,
 		const WidgetCreatedCallback& ConstructionCallback) = 0;
 
 	/** 
@@ -344,7 +344,7 @@ public:
 	 * constructor is free to use that to configure the widget. Arguments are used by the constructor to configure the widget.
 	 */
 	virtual TSharedPtr<SWidget> ConstructWidget(RowHandle Row, FTypedElementWidgetConstructor& Constructor,
-		const TypedElementDataStorage::FMetaDataView& Arguments) = 0;
+		const UE::Editor::DataStorage::FMetaDataView& Arguments) = 0;
 
 	/** Calls the provided callback for all known registered widget purposes. */
 	virtual void ListWidgetPurposes(const WidgetPurposeCallback& Callback) const = 0;
diff --git a/Engine/Source/Runtime/TypedElementFramework/Public/Elements/Interfaces/TypedElementQueryStorageInterfaces.h b/Engine/Source/Runtime/TypedElementFramework/Public/Elements/Interfaces/TypedElementQueryStorageInterfaces.h
--- a/Engine/Source/Runtime/TypedElementFramework/Public/Elements/Interfaces/TypedElementQueryStorageInterfaces.h
+++ b/Engine/Source/Runtime/TypedElementFramework/Public/Elements/Interfaces/TypedElementQueryStorageInterfaces.h
@@ -14,11 +14,11 @@ class UClass;
 class UObject;
 class UScriptStruct;
 
-namespace TypedElementDataStorage
+namespace UE::Editor::DataStorage
 {
-	using namespace UE::Editor::DataStorage;
+	using namespace TypedElementDataStorage;
 
-	struct FQueryDescription; 
+	struct FQueryDescription;
 	struct ISubqueryContext;
 
 	using SubqueryCallback = TFunction<void(const FQueryDescription&, ISubqueryContext&)>;
@@ -44,14 +44,14 @@ namespace TypedElementDataStorage
 		/** Return the address of a immutable column matching the requested type or a nullptr if not found. */
 		template<typename Column>
 		const Column* GetColumn() const;
-		template<UE::Editor::DataStorage::TDataColumnType TemplateType>
+		template<TDataColumnType TemplateType>
 		const TemplateType* GetColumn(const FName& Identifier) const;
 		/** Return the address of a mutable column matching the requested type or a nullptr if not found. */
 		virtual void* GetMutableColumn(const UScriptStruct* ColumnType) = 0;
 		/** Return the address of a mutable column matching the requested type or a nullptr if not found. */
 		template<typename Column>
 		Column* GetMutableColumn();
-		template<UE::Editor::DataStorage::TDataColumnType TemplateType>
+		template<TDataColumnType TemplateType>
 		TemplateType* GetMutableColumn(const FName& Identifier);
 
 		/**
@@ -85,24 +85,24 @@ namespace TypedElementDataStorage
 		 * table that's set in the context, for instance because it's the current row, then the version that doesn't take a row
 		 * as an argument is recommended.
 		 */
-		virtual bool HasColumn(TypedElementDataStorage::RowHandle Row, const UScriptStruct* ColumnType) const = 0;
+		virtual bool HasColumn(RowHandle Row, const UScriptStruct* ColumnType) const = 0;
 		/*
 		 * Return whether a column matches the requested type or not. This can be used for arbitrary rows. If the row is in the
 		 * table that's set in the context, for instance because it's the current row, then the version that doesn't take a row
 		 * as an argument is recommended.
 		 */
 		template<typename Column>
-		bool HasColumn(TypedElementDataStorage::RowHandle Row) const;
+		bool HasColumn(RowHandle Row) const;
 
-		template<UE::Editor::DataStorage::TColumnType DynamicColumnTemplate>
-		bool HasColumn(TypedElementDataStorage::RowHandle Row, const FName& Identifier) const;
+		template<TColumnType DynamicColumnTemplate>
+		bool HasColumn(RowHandle Row, const FName& Identifier) const;
 
 		/**
 		 * Finds the type of a dynamic column
 		 * @return The UScriptStruct for a previously created Dynamic Column.  If no column exists, then nullptr
 		 */
-		virtual const UScriptStruct* FindDynamicColumnType(const UE::Editor::DataStorage::FDynamicColumnDescription& Description) const = 0;
-		template<UE::Editor::DataStorage::TColumnType TemplateType>
+		virtual const UScriptStruct* FindDynamicColumnType(const FDynamicColumnDescription& Description) const = 0;
+		template<TColumnType TemplateType>
 		const UScriptStruct* FindDynamicColumnType(const FName& Identifier) const;
 	};
 
@@ -230,7 +230,7 @@ namespace TypedElementDataStorage
 		 */
 		virtual void AddColumns(TConstArrayView<RowHandle> Rows, TConstArrayView<const UScriptStruct*> ColumnTypes) = 0;
 		
-		virtual void AddColumns(TConstArrayView<RowHandle> Rows, TConstArrayView<UE::Editor::DataStorage::FDynamicColumnDescription> DynamicColumnDescriptions) = 0;
+		virtual void AddColumns(TConstArrayView<RowHandle> Rows, TConstArrayView<FDynamicColumnDescription> DynamicColumnDescriptions) = 0;
 		/**
 		 * Add a new uninitialized column of the provided type if one does not exist.
 		 * Returns a staged column which is used to copy into the database at a later time via the UStructScript Copy operator at the end
@@ -250,8 +250,8 @@ namespace TypedElementDataStorage
 		 */
 		virtual void* AddColumnUninitialized(RowHandle Row, const UScriptStruct* ObjectType, ObjectCopyOrMove Relocator) = 0;
 
-		virtual void* AddColumnUninitialized(RowHandle Row, const UE::Editor::DataStorage::FDynamicColumnDescription& DynamicColumnDescription, ObjectCopyOrMove Relocator) = 0;
-		virtual void* AddColumnUninitialized(RowHandle Row, const UE::Editor::DataStorage::FDynamicColumnDescription& DynamicColumnDescription) = 0;
+		virtual void* AddColumnUninitialized(RowHandle Row, const FDynamicColumnDescription& DynamicColumnDescription, ObjectCopyOrMove Relocator) = 0;
+		virtual void* AddColumnUninitialized(RowHandle Row, const FDynamicColumnDescription& DynamicColumnDescription) = 0;
 		
 
 		/**
@@ -349,7 +349,7 @@ namespace TypedElementDataStorage
 		 */
 		virtual FQueryResult RunSubquery(int32 SubqueryIndex, RowHandle Row, SubqueryCallbackRef Callback) = 0;
 	};
-} // namespace TypedElementDataStorage
+} // namespace UE::Editor::DataStorage
 
 
 
@@ -358,7 +358,7 @@ namespace TypedElementDataStorage
 // Implementations
 //
 
-namespace TypedElementDataStorage
+namespace UE::Editor::DataStorage
 {
 	//
 	// ICommonQueryContext
@@ -370,7 +370,7 @@ namespace TypedElementDataStorage
 		return reinterpret_cast<const Column*>(GetColumn(Column::StaticStruct()));
 	}
 
-	template <UE::Editor::DataStorage::TDataColumnType TemplateType>
+	template <TDataColumnType TemplateType>
 	const TemplateType* ICommonQueryContext::GetColumn(const FName& Identifier) const
 	{
 		const UScriptStruct* DynamicColumnType = FindDynamicColumnType<TemplateType>(Identifier);
@@ -383,7 +383,7 @@ namespace TypedElementDataStorage
 		return reinterpret_cast<Column*>(GetMutableColumn(Column::StaticStruct()));
 	}
 
-	template <UE::Editor::DataStorage::TDataColumnType TemplateType>
+	template <TDataColumnType TemplateType>
 	TemplateType* ICommonQueryContext::GetMutableColumn(const FName& Identifier)
 	{
 		const UScriptStruct* DynamicColumnType = FindDynamicColumnType<TemplateType>(Identifier);
@@ -397,13 +397,13 @@ namespace TypedElementDataStorage
 	}
 
 	template <typename Column>
-	bool ICommonQueryContext::HasColumn(TypedElementDataStorage::RowHandle Row) const
+	bool ICommonQueryContext::HasColumn(RowHandle Row) const
 	{
 		return HasColumn(Row, Column::StaticStruct());
 	}
 
-	template <UE::Editor::DataStorage::TColumnType DynamicColumnTemplate>
-	bool ICommonQueryContext::HasColumn(TypedElementDataStorage::RowHandle Row, const FName& Identifier) const
+	template <TColumnType DynamicColumnTemplate>
+	bool ICommonQueryContext::HasColumn(RowHandle Row, const FName& Identifier) const
 	{
 		if (const UScriptStruct* DynamicColumnType = FindDynamicColumnType<DynamicColumnTemplate>(Identifier))
 		{
@@ -412,10 +412,10 @@ namespace TypedElementDataStorage
 		return false;
 	}
 
-	template <UE::Editor::DataStorage::TColumnType TemplateType>
+	template <TColumnType TemplateType>
 	const UScriptStruct* ICommonQueryContext::FindDynamicColumnType(const FName& Identifier) const
 	{
-		return FindDynamicColumnType(UE::Editor::DataStorage::FDynamicColumnDescription
+		return FindDynamicColumnType(FDynamicColumnDescription
 			{
 				.TemplateType = TemplateType::StaticStruct(),
 				.Identifier = Identifier
@@ -522,19 +522,19 @@ namespace TypedElementDataStorage
 	ColumnTypeTemplate* ICommonQueryWithEnvironmentContext::AddColumn(RowHandle Row, const FName& Identifier)
 	{
 		UScriptStruct* TemplateType = ColumnTypeTemplate::StaticStruct();
-		const UE::Editor::DataStorage::FDynamicColumnDescription Description
+		const FDynamicColumnDescription Description
 		{
 			.TemplateType = TemplateType,
 			.Identifier = Identifier
 		};
 
-		if constexpr (UE::Editor::DataStorage::TDataColumnType<ColumnTypeTemplate>)
+		if constexpr (TDataColumnType<ColumnTypeTemplate>)
 		{
 			ColumnTypeTemplate* ColumnData = static_cast<ColumnTypeTemplate*>(AddColumnUninitialized(Row, Description));
 			new (ColumnData) ColumnTypeTemplate();
 			return ColumnData;
 		}
-		if constexpr (UE::Editor::DataStorage::TTagColumnType<ColumnTypeTemplate>)
+		if constexpr (TTagColumnType<ColumnTypeTemplate>)
 		{
 			AddColumns({Row}, {Description});
 			return nullptr;
@@ -547,7 +547,7 @@ namespace TypedElementDataStorage
 	ColumnTypeTemplate* ICommonQueryWithEnvironmentContext::AddColumn(RowHandle Row, const FName& Identifier, ColumnTypeTemplate&& Column)
 	{
 		UScriptStruct* TemplateType = ColumnTypeTemplate::StaticStruct();
-		const UE::Editor::DataStorage::FDynamicColumnDescription Description
+		const FDynamicColumnDescription Description
 		{
 			.TemplateType = TemplateType,
 			.Identifier = Identifier
@@ -593,4 +593,4 @@ namespace TypedElementDataStorage
 	{
 		RemoveColumns(Rows, { Columns::StaticStruct()... });
 	}
-} // namespace TypedElementDataStorage
+} // namespace UE::Editor::DataStorage
diff --git a/Engine/Source/Runtime/TypedElementFramework/Tests/Elements/Framework/TypedElementDataStoragePerformanceTestCommands.cpp b/Engine/Source/Runtime/TypedElementFramework/Tests/Elements/Framework/TypedElementDataStoragePerformanceTestCommands.cpp
--- a/Engine/Source/Runtime/TypedElementFramework/Tests/Elements/Framework/TypedElementDataStoragePerformanceTestCommands.cpp
+++ b/Engine/Source/Runtime/TypedElementFramework/Tests/Elements/Framework/TypedElementDataStoragePerformanceTestCommands.cpp
@@ -12,61 +12,57 @@ namespace UE::Editor::DataStorage
 	namespace Private
 	{
 		TableHandle PerformanceTestCommandTable = InvalidTableHandle;
-	}
-}
+	} // namespace Private
 
-FAutoConsoleCommand CVarAddDebugRows = FAutoConsoleCommand(
-	TEXT("Teds.Debug.PerformanceTest.AddRows"),
-	TEXT("Teds.Debug.PerformanceTest.AddRows <NumRows>;  NumRows = number of rows to add"),
-	FConsoleCommandWithArgsDelegate::CreateLambda([](const TArray<FString>& Args)
-	{
-		if (Args.Num() > 0)
+	FAutoConsoleCommand CVarAddDebugRows = FAutoConsoleCommand(
+		TEXT("Teds.Debug.PerformanceTest.AddRows"),
+		TEXT("Teds.Debug.PerformanceTest.AddRows <NumRows>;  NumRows = number of rows to add"),
+		FConsoleCommandWithArgsDelegate::CreateLambda([](const TArray<FString>& Args)
 		{
-			int32 EntitiesToAdd;
-			LexFromString(EntitiesToAdd, *Args[0]);
+			if (Args.Num() > 0)
+			{
+				int32 EntitiesToAdd;
+				LexFromString(EntitiesToAdd, *Args[0]);
 
-			ITypedElementDataStorageInterface* DataStorageInterface = UTypedElementRegistry::GetInstance()->GetMutableDataStorage();
+				ITypedElementDataStorageInterface* DataStorageInterface = UTypedElementRegistry::GetInstance()->GetMutableDataStorage();
 			
-			DataStorageInterface->BatchAddRow(UE::Editor::DataStorage::Private::PerformanceTestCommandTable, EntitiesToAdd, [DataStorageInterface](TypedElementDataStorage::RowHandle Row)
-				{
-					FTest_PingPongPrePhys* Column = DataStorageInterface->GetColumn<FTest_PingPongPrePhys>(Row);
-					Column->Value = 0;
-				});
+				DataStorageInterface->BatchAddRow(Private::PerformanceTestCommandTable, EntitiesToAdd, [DataStorageInterface](RowHandle Row)
+					{
+						FTest_PingPongPrePhys* Column = DataStorageInterface->GetColumn<FTest_PingPongPrePhys>(Row);
+						Column->Value = 0;
+					});
+			}
 		}
-	}
-));
-
-FAutoConsoleCommand CVarResetDebugEntities = FAutoConsoleCommand(
-	TEXT("Teds.Debug.PerformanceTest.RemoveAllRows"),
-	TEXT("Removes all added rows for the performance test"),
-	FConsoleCommandWithArgsDelegate::CreateLambda([](const TArray<FString>& Args)
-	{
-		using namespace TypedElementQueryBuilder;
-		using namespace TypedElementDataStorage;
+	));
 
-		using DSI = ITypedElementDataStorageInterface;
-		namespace DS = TypedElementDataStorage;
+	FAutoConsoleCommand CVarResetDebugEntities = FAutoConsoleCommand(
+		TEXT("Teds.Debug.PerformanceTest.RemoveAllRows"),
+		TEXT("Removes all added rows for the performance test"),
+		FConsoleCommandWithArgsDelegate::CreateLambda([](const TArray<FString>& Args)
+		{
+			using namespace TypedElementQueryBuilder;
 
-		ITypedElementDataStorageInterface* DataStorageInterface = UTypedElementRegistry::GetInstance()->GetMutableDataStorage();
+			ITypedElementDataStorageInterface* DataStorageInterface = UTypedElementRegistry::GetInstance()->GetMutableDataStorage();
 		
-		TypedElementDataStorage::QueryHandle Query = DataStorageInterface->RegisterQuery(
-			Select().Where().All<FTest_PingPongPrePhys>().Compile());
+			QueryHandle Query = DataStorageInterface->RegisterQuery(
+				Select().Where().All<FTest_PingPongPrePhys>().Compile());
 
-		TArray<TypedElementDataStorage::RowHandle> RowsToDelete;
+			TArray<RowHandle> RowsToDelete;
 
-		DataStorageInterface->RunQuery(Query, CreateDirectQueryCallbackBinding([&RowsToDelete](const ITypedElementDataStorageInterface::IDirectQueryContext& Context, const TypedElementDataStorage::RowHandle* RowArray)
-		{
-			RowsToDelete.Insert(RowArray, Context.GetRowCount(), RowsToDelete.Num());
-		}));
+			DataStorageInterface->RunQuery(Query, CreateDirectQueryCallbackBinding([&RowsToDelete](const ITypedElementDataStorageInterface::IDirectQueryContext& Context, const RowHandle* RowArray)
+			{
+				RowsToDelete.Insert(RowArray, Context.GetRowCount(), RowsToDelete.Num());
+			}));
 
-		for (TypedElementDataStorage::RowHandle Row : RowsToDelete)
-		{
-			DataStorageInterface->RemoveRow(Row);
-		}
+			for (RowHandle Row : RowsToDelete)
+			{
+				DataStorageInterface->RemoveRow(Row);
+			}
 
-		DataStorageInterface->UnregisterQuery(Query);
-	}
-));
+			DataStorageInterface->UnregisterQuery(Query);
+		}
+	));
+} // namespace UE::Editor::DataStorage
 
 void UTest_PingPongBetweenPhaseFactory::RegisterTables(ITypedElementDataStorageInterface& DataStorage)
 {
diff --git a/Engine/Source/Runtime/TypedElementFramework/Tests/Elements/Framework/TypedElementDataStorageProcessorTests.cpp b/Engine/Source/Runtime/TypedElementFramework/Tests/Elements/Framework/TypedElementDataStorageProcessorTests.cpp
--- a/Engine/Source/Runtime/TypedElementFramework/Tests/Elements/Framework/TypedElementDataStorageProcessorTests.cpp
+++ b/Engine/Source/Runtime/TypedElementFramework/Tests/Elements/Framework/TypedElementDataStorageProcessorTests.cpp
@@ -30,8 +30,6 @@ namespace UE::Editor::DataStorage
 		{
 			ITypedElementDataStorageInterface* DataStorage = UTypedElementRegistry::GetInstance()->GetMutableDataStorage();
 			using namespace TypedElementQueryBuilder;
-			using DSI = ITypedElementDataStorageInterface;
-			namespace DS = TypedElementDataStorage;
 
 			if (PrimaryTable == InvalidTableHandle)
 			{
@@ -47,8 +45,8 @@ namespace UE::Editor::DataStorage
 			// Test creation of a row from within a query processor
 			QueryHandle PrimaryRowQuery = DataStorage->RegisterQuery(
 				Select(TEXT("TEST: Creating a row for primary reference column"),
-				FProcessor(DSI::EQueryTickPhase::PrePhysics, DataStorage->GetQueryTickGroupName(DSI::EQueryTickGroups::Default)),
-				[](DS::IQueryContext& Context, const RowHandle* Rows, FTEDSProcessorTestsReferenceColumn* ReferenceColumns)
+				FProcessor(TypedElementDataStorage::EQueryTickPhase::PrePhysics, DataStorage->GetQueryTickGroupName(TypedElementDataStorage::EQueryTickGroups::Default)),
+				[](IQueryContext& Context, const RowHandle* Rows, FTEDSProcessorTestsReferenceColumn* ReferenceColumns)
 				{
 					const int32 RowCount = Context.GetRowCount();
 					TConstArrayView<RowHandle> RowsView = MakeArrayView(Rows, RowCount);
@@ -87,8 +85,8 @@ namespace UE::Editor::DataStorage
 
 			QueryHandle SecondaryRowQuery = DataStorage->RegisterQuery(
 				Select(TEXT("TEST: Creating a row for secondary reference column"),
-				FProcessor(DSI::EQueryTickPhase::DuringPhysics, DataStorage->GetQueryTickGroupName(DSI::EQueryTickGroups::Default)),
-				[](DS::IQueryContext& Context, const RowHandle* Rows, FTEDSProcessorTestsReferenceColumn* ReferenceColumns)
+				FProcessor(TypedElementDataStorage::EQueryTickPhase::DuringPhysics, DataStorage->GetQueryTickGroupName(TypedElementDataStorage::EQueryTickGroups::Default)),
+				[](IQueryContext& Context, const RowHandle* Rows, FTEDSProcessorTestsReferenceColumn* ReferenceColumns)
 				{
 					const int32 RowCount = Context.GetRowCount();
 					TConstArrayView<RowHandle> RowsView = MakeArrayView(Rows, RowCount);
@@ -101,7 +99,7 @@ namespace UE::Editor::DataStorage
 							0,
 							ReferenceColumnsView[Index].Reference,
 							CreateSubqueryCallbackBinding(
-							[SecondaryRow](TypedElementDataStorage::ISubqueryContext& SubqueryContext, RowHandle PrimaryRow, const FTEDSProcessorTestsReferenceColumn& ReferenceColumn)
+							[SecondaryRow](ISubqueryContext& SubqueryContext, RowHandle PrimaryRow, const FTEDSProcessorTestsReferenceColumn& ReferenceColumn)
 							{
 								if (ReferenceColumn.Reference == SecondaryRow)
 								{
