[Teds] Rename and Namespace TypedElementQueryConditions

Also fixes the FColumn ambiguious-ness I added with the prior changes. FColumn in the query conditions is now TColumn to match coding standards

diff --git a/Engine/Plugins/Experimental/Animation/HierarchyTable/Source/Editor/Private/HierarchyTable/Factories/OverrideFactory.cpp b/Engine/Plugins/Experimental/Animation/HierarchyTable/Source/Editor/Private/HierarchyTable/Factories/OverrideFactory.cpp
--- a/Engine/Plugins/Experimental/Animation/HierarchyTable/Source/Editor/Private/HierarchyTable/Factories/OverrideFactory.cpp
+++ b/Engine/Plugins/Experimental/Animation/HierarchyTable/Source/Editor/Private/HierarchyTable/Factories/OverrideFactory.cpp
@@ -7,5 +7,7 @@
 
 void UHierarchyTableOverrideFactory::RegisterWidgetConstructors(ITypedElementDataStorageInterface& DataStorage, ITypedElementDataStorageUiInterface& DataStorageUi) const
 {
-	DataStorageUi.RegisterWidgetFactory<FTypedElementWidgetConstructor_Override>(FName(TEXT("General.Cell")), TypedElementDataStorage::FColumn<FTypedElementOverrideColumn>());
+	using namespace UE::Editor::DataStorage::Queries;
+
+	DataStorageUi.RegisterWidgetFactory<FTypedElementWidgetConstructor_Override>(FName(TEXT("General.Cell")), TColumn<FTypedElementOverrideColumn>());
 }
\ No newline at end of file
diff --git a/Engine/Plugins/Experimental/Animation/HierarchyTableBuiltin/Source/Editor/Private/MaskProfile/MaskProfileFactory.cpp b/Engine/Plugins/Experimental/Animation/HierarchyTableBuiltin/Source/Editor/Private/MaskProfile/MaskProfileFactory.cpp
--- a/Engine/Plugins/Experimental/Animation/HierarchyTableBuiltin/Source/Editor/Private/MaskProfile/MaskProfileFactory.cpp
+++ b/Engine/Plugins/Experimental/Animation/HierarchyTableBuiltin/Source/Editor/Private/MaskProfile/MaskProfileFactory.cpp
@@ -6,5 +6,7 @@
 
 void UHierarchyTableMaskFactory::RegisterWidgetConstructors(ITypedElementDataStorageInterface& DataStorage, ITypedElementDataStorageUiInterface& DataStorageUi) const
 {
-	DataStorageUi.RegisterWidgetFactory<FHierarchyTableMaskWidgetConstructor_Value>(FName(TEXT("General.Cell")), TypedElementDataStorage::FColumn<FHierarchyTableMaskColumn_Value>());
+	using namespace UE::Editor::DataStorage::Queries;
+
+	DataStorageUi.RegisterWidgetFactory<FHierarchyTableMaskWidgetConstructor_Value>(FName(TEXT("General.Cell")), TColumn<FHierarchyTableMaskColumn_Value>());
 }
\ No newline at end of file
diff --git a/Engine/Plugins/Experimental/Animation/HierarchyTableBuiltin/Source/Editor/Private/TimeProfile/TimeProfileFactory.cpp b/Engine/Plugins/Experimental/Animation/HierarchyTableBuiltin/Source/Editor/Private/TimeProfile/TimeProfileFactory.cpp
--- a/Engine/Plugins/Experimental/Animation/HierarchyTableBuiltin/Source/Editor/Private/TimeProfile/TimeProfileFactory.cpp
+++ b/Engine/Plugins/Experimental/Animation/HierarchyTableBuiltin/Source/Editor/Private/TimeProfile/TimeProfileFactory.cpp
@@ -6,8 +6,10 @@
 
 void UHierarchyTableTimeFactory::RegisterWidgetConstructors(ITypedElementDataStorageInterface& DataStorage, ITypedElementDataStorageUiInterface& DataStorageUi) const
 {
-	DataStorageUi.RegisterWidgetFactory<FHierarchyTableTimeWidgetConstructor_StartTime>(FName(TEXT("General.Cell")), TypedElementDataStorage::FColumn<FHierarchyTableTimeColumn_StartTime>());
-	DataStorageUi.RegisterWidgetFactory<FHierarchyTableTimeWidgetConstructor_EndTime>(FName(TEXT("General.Cell")), TypedElementDataStorage::FColumn<FHierarchyTableTimeColumn_EndTime>());
-	DataStorageUi.RegisterWidgetFactory<FHierarchyTableTimeWidgetConstructor_TimeFactor>(FName(TEXT("General.Cell")), TypedElementDataStorage::FColumn<FHierarchyTableTimeColumn_TimeFactor>());
-	DataStorageUi.RegisterWidgetFactory<FHierarchyTableTimeWidgetConstructor_Preview>(FName(TEXT("General.Cell")), TypedElementDataStorage::FColumn<FHierarchyTableTimeColumn_Preview>());
+	using namespace UE::Editor::DataStorage::Queries;
+
+	DataStorageUi.RegisterWidgetFactory<FHierarchyTableTimeWidgetConstructor_StartTime>(FName(TEXT("General.Cell")), TColumn<FHierarchyTableTimeColumn_StartTime>());
+	DataStorageUi.RegisterWidgetFactory<FHierarchyTableTimeWidgetConstructor_EndTime>(FName(TEXT("General.Cell")), TColumn<FHierarchyTableTimeColumn_EndTime>());
+	DataStorageUi.RegisterWidgetFactory<FHierarchyTableTimeWidgetConstructor_TimeFactor>(FName(TEXT("General.Cell")), TColumn<FHierarchyTableTimeColumn_TimeFactor>());
+	DataStorageUi.RegisterWidgetFactory<FHierarchyTableTimeWidgetConstructor_Preview>(FName(TEXT("General.Cell")), TColumn<FHierarchyTableTimeColumn_Preview>());
 }
\ No newline at end of file
diff --git a/Engine/Plugins/Experimental/EditorDataStorage/Source/TedsCore/Private/TypedElementDatabase.cpp b/Engine/Plugins/Experimental/EditorDataStorage/Source/TedsCore/Private/TypedElementDatabase.cpp
--- a/Engine/Plugins/Experimental/EditorDataStorage/Source/TedsCore/Private/TypedElementDatabase.cpp
+++ b/Engine/Plugins/Experimental/EditorDataStorage/Source/TedsCore/Private/TypedElementDatabase.cpp
@@ -105,6 +105,7 @@ namespace TypedElementDatabasePrivate
 void UEditorDataStorage::Initialize()
 {
 	using namespace TypedElementDataStorage;
+	using namespace UE::Editor::DataStorage;
 
 	check(GEditor);
 	UMassEntityEditorSubsystem* Mass = GEditor->GetEditorSubsystem<UMassEntityEditorSubsystem>();
@@ -116,7 +117,7 @@ void UEditorDataStorage::Initialize()
 	ActiveEditorPhaseManager = Mass->GetMutablePhaseManager();
 	if (ActiveEditorEntityManager && ActiveEditorPhaseManager)
 	{
-		Environment = MakeShared<UE::Editor::DataStorage::FEnvironment>(*this, *ActiveEditorEntityManager, *ActiveEditorPhaseManager);
+		Environment = MakeShared<FEnvironment>(*this, *ActiveEditorEntityManager, *ActiveEditorPhaseManager);
 
 		using PhaseType = std::underlying_type_t<EQueryTickPhase>;
 		for (PhaseType PhaseId = 0; PhaseId < static_cast<PhaseType>(EQueryTickPhase::Max); ++PhaseId)
@@ -913,7 +914,7 @@ void UEditorDataStorage::ListColumns(RowHandle Row, ColumnListWithDataCallbackRe
 	}
 }
 
-bool UEditorDataStorage::MatchesColumns(RowHandle Row, const TypedElementDataStorage::FQueryConditions& Conditions) const
+bool UEditorDataStorage::MatchesColumns(RowHandle Row, const UE::Editor::DataStorage::Queries::FConditions& Conditions) const
 {
 	if (ActiveEditorEntityManager)
 	{
diff --git a/Engine/Plugins/Experimental/EditorDataStorage/Source/TedsCore/Private/TypedElementDatabase.h b/Engine/Plugins/Experimental/EditorDataStorage/Source/TedsCore/Private/TypedElementDatabase.h
--- a/Engine/Plugins/Experimental/EditorDataStorage/Source/TedsCore/Private/TypedElementDatabase.h
+++ b/Engine/Plugins/Experimental/EditorDataStorage/Source/TedsCore/Private/TypedElementDatabase.h
@@ -134,7 +134,7 @@ public:
 	virtual bool HasColumns(RowHandle Row, TConstArrayView<TWeakObjectPtr<const UScriptStruct>> ColumnTypes) const override;
 	virtual void ListColumns(RowHandle Row, ColumnListCallbackRef Callback) const;
 	virtual void ListColumns(RowHandle Row, ColumnListWithDataCallbackRef Callback);
-	virtual bool MatchesColumns(RowHandle Row, const TypedElementDataStorage::FQueryConditions& Conditions) const override;
+	virtual bool MatchesColumns(RowHandle Row, const UE::Editor::DataStorage::Queries::FConditions& Conditions) const override;
 
 	const UScriptStruct* FindDynamicColumn(const UE::Editor::DataStorage::FDynamicColumnDescription& Description) const override;
 	const UScriptStruct* GenerateDynamicColumn(const UE::Editor::DataStorage::FDynamicColumnDescription& Description) override;
diff --git a/Engine/Plugins/Experimental/EditorDataStorage/Source/TedsCore/Private/TypedElementDatabaseUI.cpp b/Engine/Plugins/Experimental/EditorDataStorage/Source/TedsCore/Private/TypedElementDatabaseUI.cpp
--- a/Engine/Plugins/Experimental/EditorDataStorage/Source/TedsCore/Private/TypedElementDatabaseUI.cpp
+++ b/Engine/Plugins/Experimental/EditorDataStorage/Source/TedsCore/Private/TypedElementDatabaseUI.cpp
@@ -97,7 +97,7 @@ bool UEditorDataStorageUi::RegisterWidgetFactory(FName Purpose, const UScriptStr
 }
 
 bool UEditorDataStorageUi::RegisterWidgetFactory(
-	FName Purpose, const UScriptStruct* Constructor, TypedElementDataStorage::FQueryConditions Columns)
+	FName Purpose, const UScriptStruct* Constructor, UE::Editor::DataStorage::Queries::FConditions Columns)
 {
 	if (!Columns.IsEmpty())
 	{
@@ -202,7 +202,7 @@ bool UEditorDataStorageUi::RegisterWidgetFactory(FName Purpose, TUniquePtr<FType
 }
 
 bool UEditorDataStorageUi::RegisterWidgetFactory(FName Purpose, TUniquePtr<FTypedElementWidgetConstructor>&& Constructor, 
-	TypedElementDataStorage::FQueryConditions Columns)
+	UE::Editor::DataStorage::Queries::FConditions Columns)
 {
 	if (!Columns.IsEmpty())
 	{
@@ -362,14 +362,14 @@ bool UEditorDataStorageUi::CreateSingleWidgetConstructor(
 	const FWidgetFactory::ConstructorType& Constructor,
 	const UE::Editor::DataStorage::FMetaDataView& Arguments,
 	TArray<TWeakObjectPtr<const UScriptStruct>> MatchedColumnTypes,
-	const TypedElementDataStorage::FQueryConditions& QueryConditions,
+	const UE::Editor::DataStorage::Queries::FConditions& QueryConditions,
 	const WidgetConstructorCallback& Callback)
 {
 	struct Visitor
 	{
 		Visitor(
 			TArray<TWeakObjectPtr<const UScriptStruct>>&& InMatchedColumnTypes,
-			const TypedElementDataStorage::FQueryConditions& InQueryConditions,
+			const UE::Editor::DataStorage::Queries::FConditions& InQueryConditions,
 			const UE::Editor::DataStorage::FMetaDataView& InArguments,
 			const WidgetConstructorCallback& InCallback) 
 			: MatchedColumnTypes(MoveTemp(InMatchedColumnTypes))
@@ -379,7 +379,7 @@ bool UEditorDataStorageUi::CreateSingleWidgetConstructor(
 		{}
 
 		TArray<TWeakObjectPtr<const UScriptStruct>> MatchedColumnTypes;
-		const TypedElementDataStorage::FQueryConditions& QueryConditions;
+		const UE::Editor::DataStorage::Queries::FConditions& QueryConditions;
 		const UE::Editor::DataStorage::FMetaDataView& Arguments;
 		const WidgetConstructorCallback& Callback;
 
@@ -602,14 +602,14 @@ UEditorDataStorageUi::FWidgetFactory::FWidgetFactory(TUniquePtr<FTypedElementWid
 }
 
 UEditorDataStorageUi::FWidgetFactory::FWidgetFactory(const UScriptStruct* InConstructor, 
-	TypedElementDataStorage::FQueryConditions&& InColumns)
+	UE::Editor::DataStorage::Queries::FConditions&& InColumns)
 	: Columns(MoveTemp(InColumns))
 	, Constructor(InConstructor)
 {
 }
 
 UEditorDataStorageUi::FWidgetFactory::FWidgetFactory(TUniquePtr<FTypedElementWidgetConstructor>&& InConstructor, 
-	TypedElementDataStorage::FQueryConditions&& InColumns)
+	UE::Editor::DataStorage::Queries::FConditions&& InColumns)
 	: Columns(MoveTemp(InColumns))
 	, Constructor(MoveTemp(InConstructor))
 {
diff --git a/Engine/Plugins/Experimental/EditorDataStorage/Source/TedsCore/Private/TypedElementDatabaseUI.h b/Engine/Plugins/Experimental/EditorDataStorage/Source/TedsCore/Private/TypedElementDatabaseUI.h
--- a/Engine/Plugins/Experimental/EditorDataStorage/Source/TedsCore/Private/TypedElementDatabaseUI.h
+++ b/Engine/Plugins/Experimental/EditorDataStorage/Source/TedsCore/Private/TypedElementDatabaseUI.h
@@ -35,10 +35,10 @@ public:
 
 	bool RegisterWidgetFactory(FName Purpose, const UScriptStruct* Constructor) override;
 	bool RegisterWidgetFactory(FName Purpose, const UScriptStruct* Constructor,
-		TypedElementDataStorage::FQueryConditions Columns) override;
+		UE::Editor::DataStorage::Queries::FConditions Columns) override;
 	bool RegisterWidgetFactory(FName Purpose, TUniquePtr<FTypedElementWidgetConstructor>&& Constructor) override;
 	bool RegisterWidgetFactory(FName Purpose, TUniquePtr<FTypedElementWidgetConstructor>&& Constructor,
-		TypedElementDataStorage::FQueryConditions Columns) override;
+		UE::Editor::DataStorage::Queries::FConditions Columns) override;
 
 	void CreateWidgetConstructors(FName Purpose,
 		const UE::Editor::DataStorage::FMetaDataView& Arguments, const WidgetConstructorCallback& Callback) override;
@@ -60,14 +60,14 @@ private:
 	{
 		using ConstructorType = std::variant<const UScriptStruct*, TUniquePtr<FTypedElementWidgetConstructor>>;
 
-		TypedElementDataStorage::FQueryConditions Columns;
+		UE::Editor::DataStorage::Queries::FConditions Columns;
 		ConstructorType Constructor;
 
 		FWidgetFactory() = default;
 		explicit FWidgetFactory(const UScriptStruct* InConstructor);
 		explicit FWidgetFactory(TUniquePtr<FTypedElementWidgetConstructor>&& InConstructor);
-		FWidgetFactory(const UScriptStruct* InConstructor, TypedElementDataStorage::FQueryConditions&& InColumns);
-		FWidgetFactory(TUniquePtr<FTypedElementWidgetConstructor>&& InConstructor, TypedElementDataStorage::FQueryConditions&& InColumns);
+		FWidgetFactory(const UScriptStruct* InConstructor, UE::Editor::DataStorage::Queries::FConditions&& InColumns);
+		FWidgetFactory(TUniquePtr<FTypedElementWidgetConstructor>&& InConstructor, UE::Editor::DataStorage::Queries::FConditions&& InColumns);
 	};
 
 	struct FPurposeInfo
@@ -84,7 +84,7 @@ private:
 		const FWidgetFactory::ConstructorType& Constructor,
 		const UE::Editor::DataStorage::FMetaDataView& Arguments,
 		TArray<TWeakObjectPtr<const UScriptStruct>> MatchedColumnTypes,
-		const TypedElementDataStorage::FQueryConditions& QueryConditions,
+		const UE::Editor::DataStorage::Queries::FConditions& QueryConditions,
 		const WidgetConstructorCallback& Callback);
 
 	void CreateWidgetInstance(
diff --git a/Engine/Plugins/Experimental/EditorDataStorage/Source/TedsUI/Private/Widgets/ExportedTextWidget.cpp b/Engine/Plugins/Experimental/EditorDataStorage/Source/TedsUI/Private/Widgets/ExportedTextWidget.cpp
--- a/Engine/Plugins/Experimental/EditorDataStorage/Source/TedsUI/Private/Widgets/ExportedTextWidget.cpp
+++ b/Engine/Plugins/Experimental/EditorDataStorage/Source/TedsUI/Private/Widgets/ExportedTextWidget.cpp
@@ -140,7 +140,7 @@ TConstArrayView<const UScriptStruct*> FExportedTextWidgetConstructor::GetAdditio
 	return Columns;
 }
 
-const TypedElementDataStorage::FQueryConditions* FExportedTextWidgetConstructor::GetQueryConditions() const
+const UE::Editor::DataStorage::Queries::FConditions* FExportedTextWidgetConstructor::GetQueryConditions() const
 {
 	// For the exported text widget, the query condition we are matched against is the column we are exporting text for
 	return &MatchedColumn;
@@ -171,13 +171,12 @@ bool FExportedTextWidgetConstructor::FinalizeWidget(
 	const TSharedPtr<SWidget>& Widget)
 {
 	using namespace UE::Editor;
-
 	FTypedElementScriptStructTypeInfoColumn& TypeInfoColumn = *DataStorage->GetColumn<FTypedElementScriptStructTypeInfoColumn>(Row);
 
 	// NOTE: We are currently assuming that an instance of FExportedTextWidgetConstructor will only be used to show the same type info
 	// which isn't ideal but it's better than nothing since we need some sort of matched conditions for column based virtualization to work.
 	// TEDS UI TODO: We should work around it by refactoring this into an STedsWidget in the future so it can store the column conditions per instance
-	MatchedColumn = TypedElementDataStorage::FQueryConditions(TypedElementDataStorage::FColumn(TypeInfoColumn.TypeInfo));
+	MatchedColumn = DataStorage::Queries::FConditions(DataStorage::Queries::TColumn(TypeInfoColumn.TypeInfo));
 
 	if (TypeInfoColumn.TypeInfo->IsChildOf(DataStorage::FTag::StaticStruct()))
 	{
diff --git a/Engine/Plugins/Experimental/EditorDataStorage/Source/TedsUI/Private/Widgets/LabelWidget.cpp b/Engine/Plugins/Experimental/EditorDataStorage/Source/TedsUI/Private/Widgets/LabelWidget.cpp
--- a/Engine/Plugins/Experimental/EditorDataStorage/Source/TedsUI/Private/Widgets/LabelWidget.cpp
+++ b/Engine/Plugins/Experimental/EditorDataStorage/Source/TedsUI/Private/Widgets/LabelWidget.cpp
@@ -27,11 +27,13 @@
 void ULabelWidgetFactory::RegisterWidgetConstructors(ITypedElementDataStorageInterface& DataStorage,
 	ITypedElementDataStorageUiInterface& DataStorageUi) const
 {
+	using namespace UE::Editor::DataStorage::Queries;
+
 	DataStorageUi.RegisterWidgetFactory<FLabelWidgetConstructor>(FName(TEXT("General.Cell")), 
-		TypedElementDataStorage::FColumn<FTypedElementLabelColumn>() || (TypedElementDataStorage::FColumn<FTypedElementLabelColumn>() && TypedElementDataStorage::FColumn<FTypedElementLabelHashColumn>()));
+		TColumn<FTypedElementLabelColumn>() || (TColumn<FTypedElementLabelColumn>() && TColumn<FTypedElementLabelHashColumn>()));
 
 	DataStorageUi.RegisterWidgetFactory<FLabelWidgetConstructor>(FName(TEXT("General.RowLabel")), 
-	TypedElementDataStorage::FColumn<FTypedElementLabelColumn>() || (TypedElementDataStorage::FColumn<FTypedElementLabelColumn>() && TypedElementDataStorage::FColumn<FTypedElementLabelHashColumn>()));
+		TColumn<FTypedElementLabelColumn>() || (TColumn<FTypedElementLabelColumn>() && TColumn<FTypedElementLabelHashColumn>()));
 }
 
 void ULabelWidgetFactory::RegisterWidgetPurposes(ITypedElementDataStorageUiInterface& DataStorageUi) const
diff --git a/Engine/Plugins/Experimental/EditorDataStorage/Source/TedsUI/Private/Widgets/PackagePathWidget.cpp b/Engine/Plugins/Experimental/EditorDataStorage/Source/TedsUI/Private/Widgets/PackagePathWidget.cpp
--- a/Engine/Plugins/Experimental/EditorDataStorage/Source/TedsUI/Private/Widgets/PackagePathWidget.cpp
+++ b/Engine/Plugins/Experimental/EditorDataStorage/Source/TedsUI/Private/Widgets/PackagePathWidget.cpp
@@ -14,10 +14,12 @@
 void UPackagePathWidgetFactory::RegisterWidgetConstructors(ITypedElementDataStorageInterface& DataStorage,
 	ITypedElementDataStorageUiInterface& DataStorageUi) const
 {
+	using namespace UE::Editor::DataStorage::Queries;
+
 	DataStorageUi.RegisterWidgetFactory<FPackagePathWidgetConstructor>(FName(TEXT("General.Cell")),
-		TypedElementDataStorage::FColumn<FTypedElementPackagePathColumn>());
+		TColumn<FTypedElementPackagePathColumn>());
 	DataStorageUi.RegisterWidgetFactory<FLoadedPackagePathWidgetConstructor>(FName(TEXT("General.Cell")),
-		TypedElementDataStorage::FColumn<FTypedElementPackageLoadedPathColumn>());
+		TColumn<FTypedElementPackageLoadedPathColumn>());
 }
 
 
diff --git a/Engine/Plugins/Experimental/EditorDataStorage/Source/TedsUI/Public/Widgets/ExportedTextWidget.h b/Engine/Plugins/Experimental/EditorDataStorage/Source/TedsUI/Public/Widgets/ExportedTextWidget.h
--- a/Engine/Plugins/Experimental/EditorDataStorage/Source/TedsUI/Public/Widgets/ExportedTextWidget.h
+++ b/Engine/Plugins/Experimental/EditorDataStorage/Source/TedsUI/Public/Widgets/ExportedTextWidget.h
@@ -35,7 +35,7 @@ public:
 	virtual ~FExportedTextWidgetConstructor() override = default;
 
 	virtual TConstArrayView<const UScriptStruct*> GetAdditionalColumnsList() const override;
-	virtual const UE::Editor::DataStorage::FQueryConditions* GetQueryConditions() const override;
+	virtual const UE::Editor::DataStorage::Queries::FConditions* GetQueryConditions() const override;
 	virtual FString CreateWidgetDisplayName(
 		ITypedElementDataStorageInterface* DataStorage, UE::Editor::DataStorage::RowHandle Row) const override;
 
@@ -49,7 +49,7 @@ protected:
 
 protected:
 	// The column this exported text widget is operating on
-	UE::Editor::DataStorage::FQueryConditions MatchedColumn;
+	UE::Editor::DataStorage::Queries::FConditions MatchedColumn;
 };
 
 USTRUCT(meta = (DisplayName = "Exported text widget"))
diff --git a/Engine/Plugins/Experimental/EditorDataStorageFeatures/Source/TedsAssetData/Private/Widgets/AssetDataDiskSizeWidget.cpp b/Engine/Plugins/Experimental/EditorDataStorageFeatures/Source/TedsAssetData/Private/Widgets/AssetDataDiskSizeWidget.cpp
--- a/Engine/Plugins/Experimental/EditorDataStorageFeatures/Source/TedsAssetData/Private/Widgets/AssetDataDiskSizeWidget.cpp
+++ b/Engine/Plugins/Experimental/EditorDataStorageFeatures/Source/TedsAssetData/Private/Widgets/AssetDataDiskSizeWidget.cpp
@@ -11,7 +11,8 @@
 void UDiskSizeWidgetFactory::RegisterWidgetConstructors(ITypedElementDataStorageInterface& DataStorage,
                                                         ITypedElementDataStorageUiInterface& DataStorageUi) const
 {
-	DataStorageUi.RegisterWidgetFactory<FDiskSizeWidgetConstructor>(TEXT("General.Cell"), TypedElementDataStorage::FColumn<FDiskSizeColumn>());
+	using namespace UE::Editor::DataStorage::Queries;
+	DataStorageUi.RegisterWidgetFactory<FDiskSizeWidgetConstructor>(TEXT("General.Cell"), TColumn<FDiskSizeColumn>());
 }
 
 FDiskSizeWidgetConstructor::FDiskSizeWidgetConstructor()
diff --git a/Engine/Plugins/Experimental/EditorDataStorageFeatures/Source/TedsAssetData/Private/Widgets/AssetDataLabelWidget.cpp b/Engine/Plugins/Experimental/EditorDataStorageFeatures/Source/TedsAssetData/Private/Widgets/AssetDataLabelWidget.cpp
--- a/Engine/Plugins/Experimental/EditorDataStorageFeatures/Source/TedsAssetData/Private/Widgets/AssetDataLabelWidget.cpp
+++ b/Engine/Plugins/Experimental/EditorDataStorageFeatures/Source/TedsAssetData/Private/Widgets/AssetDataLabelWidget.cpp
@@ -14,9 +14,10 @@
 void UAssetDataLabelWidgetFactory::RegisterWidgetConstructors(ITypedElementDataStorageInterface& DataStorage,
                                                               ITypedElementDataStorageUiInterface& DataStorageUi) const
 {
+	using namespace UE::Editor::DataStorage::Queries;
 	DataStorageUi.RegisterWidgetFactory<FAssetDataLabelWidgetConstructor>(
 		TEXT("General.RowLabel"),
-		TypedElementDataStorage::FColumn<FItemNameColumn_Experimental>() && TypedElementDataStorage::FColumn<FAssetTag>());
+		TColumn<FItemNameColumn_Experimental>() && TColumn<FAssetTag>());
 }
 
 FAssetDataLabelWidgetConstructor::FAssetDataLabelWidgetConstructor()
diff --git a/Engine/Plugins/Experimental/EditorDataStorageFeatures/Source/TedsDebugger/Private/Widgets/RowReferenceWidget.cpp b/Engine/Plugins/Experimental/EditorDataStorageFeatures/Source/TedsDebugger/Private/Widgets/RowReferenceWidget.cpp
--- a/Engine/Plugins/Experimental/EditorDataStorageFeatures/Source/TedsDebugger/Private/Widgets/RowReferenceWidget.cpp
+++ b/Engine/Plugins/Experimental/EditorDataStorageFeatures/Source/TedsDebugger/Private/Widgets/RowReferenceWidget.cpp
@@ -99,9 +99,11 @@ URowReferenceWidgetFactory::~URowReferenceWidgetFactory()
 
 void URowReferenceWidgetFactory::RegisterWidgetConstructors(ITypedElementDataStorageInterface& DataStorage, ITypedElementDataStorageUiInterface& DataStorageUi) const
 {
+	using namespace UE::Editor::DataStorage::Queries;
+
 	// TEDS UI TODO: We can re-use this widget for FTypedElementParentColumn
 	DataStorageUi.RegisterWidgetFactory<FRowReferenceWidgetConstructor>(FName(TEXT("SceneOutliner.Cell")),
-	TypedElementDataStorage::FColumn<FTypedElementRowReferenceColumn>());
+		TColumn<FTypedElementRowReferenceColumn>());
 }
 
 void URowReferenceWidgetFactory::RegisterQueries(ITypedElementDataStorageInterface& DataStorage)
diff --git a/Engine/Plugins/Experimental/EditorDataStorageFeatures/Source/TedsDebugger/Private/Widgets/SlateVisualizationWidget.cpp b/Engine/Plugins/Experimental/EditorDataStorageFeatures/Source/TedsDebugger/Private/Widgets/SlateVisualizationWidget.cpp
--- a/Engine/Plugins/Experimental/EditorDataStorageFeatures/Source/TedsDebugger/Private/Widgets/SlateVisualizationWidget.cpp
+++ b/Engine/Plugins/Experimental/EditorDataStorageFeatures/Source/TedsDebugger/Private/Widgets/SlateVisualizationWidget.cpp
@@ -11,8 +11,10 @@
 
 void USlateVisualizationWidgetFactory::RegisterWidgetConstructors(ITypedElementDataStorageInterface& DataStorage, ITypedElementDataStorageUiInterface& DataStorageUi) const
 {
+	using namespace UE::Editor::DataStorage::Queries;
+
 	DataStorageUi.RegisterWidgetFactory<FSlateVisualizationWidgetConstructor>(FName(TEXT("General.Cell")),
-	TypedElementDataStorage::FColumn<FTypedElementSlateWidgetReferenceColumn>());
+		TColumn<FTypedElementSlateWidgetReferenceColumn>());
 }
 
 FSlateVisualizationWidgetConstructor::FSlateVisualizationWidgetConstructor()
diff --git a/Engine/Plugins/Experimental/EditorDataStorageFeatures/Source/TedsOutliner/Private/Widgets/AlertWidget.cpp b/Engine/Plugins/Experimental/EditorDataStorageFeatures/Source/TedsOutliner/Private/Widgets/AlertWidget.cpp
--- a/Engine/Plugins/Experimental/EditorDataStorageFeatures/Source/TedsOutliner/Private/Widgets/AlertWidget.cpp
+++ b/Engine/Plugins/Experimental/EditorDataStorageFeatures/Source/TedsOutliner/Private/Widgets/AlertWidget.cpp
@@ -151,11 +151,13 @@ void UAlertWidgetFactory::RegisterWidgetConstructors(
 	ITypedElementDataStorageInterface& DataStorage,
 	ITypedElementDataStorageUiInterface& DataStorageUi) const
 {
+	using namespace UE::Editor::DataStorage::Queries;
+
 	DataStorageUi.RegisterWidgetFactory<FAlertWidgetConstructor>(FName(TEXT("General.Cell")),
-		TypedElementDataStorage::FColumn<FTypedElementAlertColumn>() || TypedElementDataStorage::FColumn<FTypedElementChildAlertColumn>());
+		TColumn<FTypedElementAlertColumn>() || TColumn<FTypedElementChildAlertColumn>());
 	
 	DataStorageUi.RegisterWidgetFactory<FAlertHeaderWidgetConstructor>(FName(TEXT("General.Header")),
-		TypedElementDataStorage::FColumn<FTypedElementAlertColumn>() || TypedElementDataStorage::FColumn<FTypedElementChildAlertColumn>());
+		TColumn<FTypedElementAlertColumn>() || TColumn<FTypedElementChildAlertColumn>());
 }
 
 void UAlertWidgetFactory::RegisterQueries(ITypedElementDataStorageInterface& DataStorage)
diff --git a/Engine/Plugins/Experimental/EditorDataStorageFeatures/Source/TedsOutliner/Private/Widgets/OutlinerLabelWidget.cpp b/Engine/Plugins/Experimental/EditorDataStorageFeatures/Source/TedsOutliner/Private/Widgets/OutlinerLabelWidget.cpp
--- a/Engine/Plugins/Experimental/EditorDataStorageFeatures/Source/TedsOutliner/Private/Widgets/OutlinerLabelWidget.cpp
+++ b/Engine/Plugins/Experimental/EditorDataStorageFeatures/Source/TedsOutliner/Private/Widgets/OutlinerLabelWidget.cpp
@@ -106,9 +106,10 @@ namespace UE::OutlinerLabelWidget::Local
 void UOutlinerLabelWidgetFactory::RegisterWidgetConstructors(ITypedElementDataStorageInterface& DataStorage,
 	ITypedElementDataStorageUiInterface& DataStorageUi) const
 {
+	using namespace UE::Editor::DataStorage::Queries;
 	DataStorageUi.RegisterWidgetFactory<FOutlinerLabelWidgetConstructor>(
 		TEXT("SceneOutliner.RowLabel"),
-		TypedElementDataStorage::FColumn<FTypedElementLabelColumn>() && TypedElementDataStorage::FColumn<FTypedElementClassTypeInfoColumn>());
+		TColumn<FTypedElementLabelColumn>() && TColumn<FTypedElementClassTypeInfoColumn>());
 }
 
 FOutlinerLabelWidgetConstructor::FOutlinerLabelWidgetConstructor()
diff --git a/Engine/Plugins/Experimental/EditorDataStorageFeatures/Source/TedsOutliner/Private/Widgets/TransformHeadsUpWidget.cpp b/Engine/Plugins/Experimental/EditorDataStorageFeatures/Source/TedsOutliner/Private/Widgets/TransformHeadsUpWidget.cpp
--- a/Engine/Plugins/Experimental/EditorDataStorageFeatures/Source/TedsOutliner/Private/Widgets/TransformHeadsUpWidget.cpp
+++ b/Engine/Plugins/Experimental/EditorDataStorageFeatures/Source/TedsOutliner/Private/Widgets/TransformHeadsUpWidget.cpp
@@ -190,8 +190,10 @@ void UTransformHeadsUpWidgetFactory::RegisterQueries(ITypedElementDataStorageInt
 void UTransformHeadsUpWidgetFactory::RegisterWidgetConstructors(ITypedElementDataStorageInterface& DataStorage,
 	ITypedElementDataStorageUiInterface& DataStorageUi) const
 {
+	using namespace UE::Editor::DataStorage::Queries;
+
 	DataStorageUi.RegisterWidgetFactory<FTransformHeadsUpWidgetConstructor>(FName(TEXT("SceneOutliner.Cell")),
-		TypedElementDataStorage::FColumn<FTypedElementLocalTransformColumn>());
+		TColumn<FTypedElementLocalTransformColumn>());
 }
 
 
diff --git a/Engine/Plugins/Experimental/EditorDataStorageFeatures/Source/TedsOutliner/Private/Widgets/TypeInfoWidget.cpp b/Engine/Plugins/Experimental/EditorDataStorageFeatures/Source/TedsOutliner/Private/Widgets/TypeInfoWidget.cpp
--- a/Engine/Plugins/Experimental/EditorDataStorageFeatures/Source/TedsOutliner/Private/Widgets/TypeInfoWidget.cpp
+++ b/Engine/Plugins/Experimental/EditorDataStorageFeatures/Source/TedsOutliner/Private/Widgets/TypeInfoWidget.cpp
@@ -16,8 +16,10 @@
 void UTypeInfoWidgetFactory::RegisterWidgetConstructors(ITypedElementDataStorageInterface& DataStorage,
 	ITypedElementDataStorageUiInterface& DataStorageUi) const
 {
+	using namespace UE::Editor::DataStorage::Queries;
+
 	DataStorageUi.RegisterWidgetFactory<FTypeInfoWidgetConstructor>(FName(TEXT("General.Cell")),
-		TypedElementDataStorage::FColumn<FTypedElementClassTypeInfoColumn>());
+		TColumn<FTypedElementClassTypeInfoColumn>());
 
 }
 
diff --git a/Engine/Plugins/Experimental/EditorDataStorageFeatures/Source/TedsTableViewer/Private/Widgets/UObjectLabelWidget.cpp b/Engine/Plugins/Experimental/EditorDataStorageFeatures/Source/TedsTableViewer/Private/Widgets/UObjectLabelWidget.cpp
--- a/Engine/Plugins/Experimental/EditorDataStorageFeatures/Source/TedsTableViewer/Private/Widgets/UObjectLabelWidget.cpp
+++ b/Engine/Plugins/Experimental/EditorDataStorageFeatures/Source/TedsTableViewer/Private/Widgets/UObjectLabelWidget.cpp
@@ -17,9 +17,10 @@
 void UUObjectLabelWidgetFactory::RegisterWidgetConstructors(ITypedElementDataStorageInterface& DataStorage,
 	ITypedElementDataStorageUiInterface& DataStorageUi) const
 {
+	using namespace UE::Editor::DataStorage::Queries;
 	DataStorageUi.RegisterWidgetFactory<FUObjectLabelWidgetConstructor>(
 		TEXT("General.RowLabel"),
-		TypedElementDataStorage::FColumn<FTypedElementLabelColumn>() && TypedElementDataStorage::FColumn<FTypedElementUObjectColumn>());
+		TColumn<FTypedElementLabelColumn>() && TColumn<FTypedElementUObjectColumn>());
 }
 
 FUObjectLabelWidgetConstructor::FUObjectLabelWidgetConstructor()
diff --git a/Engine/Source/Runtime/TypedElementFramework/Private/Elements/Common/TypedElementQueryConditions.cpp b/Engine/Source/Runtime/TypedElementFramework/Private/Elements/Common/TypedElementQueryConditions.cpp
--- a/Engine/Source/Runtime/TypedElementFramework/Private/Elements/Common/TypedElementQueryConditions.cpp
+++ b/Engine/Source/Runtime/TypedElementFramework/Private/Elements/Common/TypedElementQueryConditions.cpp
@@ -4,15 +4,15 @@
 
 #include "UObject/Class.h"
 
-namespace TypedElementDataStorage
+namespace UE::Editor::DataStorage::Queries
 {
-	FQueryConditions::FQueryConditions(FColumnBase Column)
+	FConditions::FConditions(FColumnBase Column)
 		: ColumnCount(1)
 	{
 		Columns[0] = Column.TypeInfo;
 	}
 
-	void FQueryConditions::AppendToString(FString& Output) const
+	void FConditions::AppendToString(FString& Output) const
 	{
 		if (TokenCount > 0)
 		{
@@ -60,7 +60,7 @@ namespace TypedElementDataStorage
 		}
 	}
 
-	bool FQueryConditions::Verify(TConstArrayView<FColumnBase> AvailableColumns) const
+	bool FConditions::Verify(TConstArrayView<FColumnBase> AvailableColumns) const
 	{
 		return VerifyBootstrap(
 			[&AvailableColumns](uint8_t ColumnIndex, TWeakObjectPtr<const UScriptStruct> Column)
@@ -76,7 +76,7 @@ namespace TypedElementDataStorage
 			});
 	}
 
-	bool FQueryConditions::Verify(TArray<TWeakObjectPtr<const UScriptStruct>>& MatchedColumns, TConstArrayView<FColumnBase> AvailableColumns,
+	bool FConditions::Verify(TArray<TWeakObjectPtr<const UScriptStruct>>& MatchedColumns, TConstArrayView<FColumnBase> AvailableColumns,
 		bool AvailableColumnsAreSorted) const
 	{
 		static_assert(MaxColumnCount < 64, "Query conditions use a bit mask to locate matches. As a result MaxColumnCount can be larger than 64.");
@@ -116,7 +116,7 @@ namespace TypedElementDataStorage
 		return Result;
 	}
 
-	bool FQueryConditions::Verify(TConstArrayView<TWeakObjectPtr<const UScriptStruct>> AvailableColumns, bool AvailableColumnsAreSorted) const
+	bool FConditions::Verify(TConstArrayView<TWeakObjectPtr<const UScriptStruct>> AvailableColumns, bool AvailableColumnsAreSorted) const
 	{
 		return AvailableColumnsAreSorted
 			? VerifyBootstrap(
@@ -135,7 +135,7 @@ namespace TypedElementDataStorage
 				});
 	}
 
-	bool FQueryConditions::Verify(TArray<TWeakObjectPtr<const UScriptStruct>>& MatchedColumns,
+	bool FConditions::Verify(TArray<TWeakObjectPtr<const UScriptStruct>>& MatchedColumns,
 		TConstArrayView<TWeakObjectPtr<const UScriptStruct>> AvailableColumns, bool AvailableColumnsAreSorted) const
 	{
 		static_assert(MaxColumnCount < 64, "Query conditions use a bit mask to locate matches. As a result MaxColumnCount cannot be larger than 64.");
@@ -167,7 +167,7 @@ namespace TypedElementDataStorage
 		return Result;
 	}
 
-	bool FQueryConditions::Verify(TSet<TWeakObjectPtr<const UScriptStruct>> AvailableColumns) const
+	bool FConditions::Verify(TSet<TWeakObjectPtr<const UScriptStruct>> AvailableColumns) const
 	{
 		return VerifyBootstrap(
 			[&AvailableColumns](uint8_t ColumnIndex, TWeakObjectPtr<const UScriptStruct> Column)
@@ -176,12 +176,12 @@ namespace TypedElementDataStorage
 			});
 	}
 
-	bool FQueryConditions::Verify(ContainsCallback Callback) const
+	bool FConditions::Verify(ContainsCallback Callback) const
 	{
 		return VerifyBootstrap(Callback);
 	}
 
-	uint8_t FQueryConditions::MinimumColumnMatchRequired() const
+	uint8_t FConditions::MinimumColumnMatchRequired() const
 	{
 		if (TokenCount > 0)
 		{
@@ -198,17 +198,17 @@ namespace TypedElementDataStorage
 		}
 	}
 
-	TConstArrayView<TWeakObjectPtr<const UScriptStruct>> FQueryConditions::GetColumns() const
+	TConstArrayView<TWeakObjectPtr<const UScriptStruct>> FConditions::GetColumns() const
 	{
 		return TConstArrayView<TWeakObjectPtr<const UScriptStruct>>(Columns, ColumnCount);
 	}
 
-	bool FQueryConditions::IsEmpty() const
+	bool FConditions::IsEmpty() const
 	{
 		return ColumnCount == 0;
 	}
 
-	void FQueryConditions::AppendName(FString& Output, TWeakObjectPtr<const UScriptStruct> TypeInfo) const
+	void FConditions::AppendName(FString& Output, TWeakObjectPtr<const UScriptStruct> TypeInfo) const
 	{
 #if WITH_EDITORONLY_DATA
 		static FName DisplayNameName(TEXT("DisplayName"));
@@ -221,17 +221,17 @@ namespace TypedElementDataStorage
 #endif
 	}
 
-	bool FQueryConditions::EntersScopeNext(uint8_t Index) const
+	bool FConditions::EntersScopeNext(uint8_t Index) const
 	{
 		return Index < (TokenCount - 1) && Tokens[Index + 1] == Token::ScopeOpen;
 	}
 
-	bool FQueryConditions::EntersScope(uint8_t Index) const
+	bool FConditions::EntersScope(uint8_t Index) const
 	{
 		return Tokens[Index] == Token::ScopeOpen;
 	}
 
-	bool FQueryConditions::Contains(TWeakObjectPtr<const UScriptStruct> ColumnType, const TArray<FColumnBase>& AvailableColumns) const
+	bool FConditions::Contains(TWeakObjectPtr<const UScriptStruct> ColumnType, const TArray<FColumnBase>& AvailableColumns) const
 	{
 		for (const FColumnBase& Column : AvailableColumns)
 		{
@@ -243,7 +243,7 @@ namespace TypedElementDataStorage
 		return false;
 	}
 
-	bool FQueryConditions::VerifyBootstrap(ContainsCallback Contains) const
+	bool FConditions::VerifyBootstrap(ContainsCallback Contains) const
 	{
 		if (TokenCount > 0)
 		{
@@ -262,7 +262,7 @@ namespace TypedElementDataStorage
 		}
 	}
 
-	bool FQueryConditions::VerifyRange(uint8_t& TokenIndex, uint8_t& ColumnIndex, ContainsCallback Contains) const
+	bool FConditions::VerifyRange(uint8_t& TokenIndex, uint8_t& ColumnIndex, ContainsCallback Contains) const
 	{
 		auto Init = [&]()
 		{
@@ -320,7 +320,7 @@ namespace TypedElementDataStorage
 		return Result;
 	}
 
-	void FQueryConditions::ConvertColumnBitToArray(TArray<TWeakObjectPtr<const UScriptStruct>>& MatchedColumns, uint64 ColumnBits) const
+	void FConditions::ConvertColumnBitToArray(TArray<TWeakObjectPtr<const UScriptStruct>>& MatchedColumns, uint64 ColumnBits) const
 	{
 		uint64 Index = 0;
 		while (ColumnBits)
@@ -334,7 +334,7 @@ namespace TypedElementDataStorage
 		}
 	}
 
-	uint8_t FQueryConditions::MinimumColumnMatchRequiredRange(uint8_t& Front) const
+	uint8_t FConditions::MinimumColumnMatchRequiredRange(uint8_t& Front) const
 	{
 		uint8_t Result = EntersScope(Front) ? MinimumColumnMatchRequiredRange(++Front) : 1;
 		for (; Front < TokenCount; ++Front)
@@ -375,7 +375,7 @@ namespace TypedElementDataStorage
 		return Result;
 	}
 
-	void FQueryConditions::AppendQuery(FQueryConditions& Target, const FQueryConditions& Source)
+	void FConditions::AppendQuery(FConditions& Target, const FConditions& Source)
 	{
 		checkf(Target.ColumnCount + Source.ColumnCount < MaxColumnCount, TEXT("Too many columns in the query."));
 		for (uint8_t Index = 0; Index < Source.ColumnCount; ++Index)
@@ -393,87 +393,87 @@ namespace TypedElementDataStorage
 		Target.TokenCount += Source.TokenCount;
 	}
 
-	FQueryConditions operator&&(const FQueryConditions& Lhs, FColumnBase Rhs)
+	FConditions operator&&(const FConditions& Lhs, FColumnBase Rhs)
 	{
-		FQueryConditions Result = Lhs;
+		FConditions Result = Lhs;
 		Result.Columns[Result.ColumnCount++] = Rhs.TypeInfo;
-		Result.Tokens[Result.TokenCount++] = FQueryConditions::Token::And;
+		Result.Tokens[Result.TokenCount++] = FConditions::Token::And;
 		return Result;
 	}
 
-	FQueryConditions operator&&(const FQueryConditions& Lhs, const FQueryConditions& Rhs)
+	FConditions operator&&(const FConditions& Lhs, const FConditions& Rhs)
 	{
-		FQueryConditions Result;
-		Result.Tokens[Result.TokenCount++] = FQueryConditions::Token::ScopeOpen;
-		FQueryConditions::AppendQuery(Result, Lhs);
-		Result.Tokens[Result.TokenCount++] = FQueryConditions::Token::ScopeClose;
+		FConditions Result;
+		Result.Tokens[Result.TokenCount++] = FConditions::Token::ScopeOpen;
+		FConditions::AppendQuery(Result, Lhs);
+		Result.Tokens[Result.TokenCount++] = FConditions::Token::ScopeClose;
 
-		Result.Tokens[Result.TokenCount++] = FQueryConditions::Token::And;
+		Result.Tokens[Result.TokenCount++] = FConditions::Token::And;
 
-		Result.Tokens[Result.TokenCount++] = FQueryConditions::Token::ScopeOpen;
-		FQueryConditions::AppendQuery(Result, Rhs);
-		Result.Tokens[Result.TokenCount++] = FQueryConditions::Token::ScopeClose;
+		Result.Tokens[Result.TokenCount++] = FConditions::Token::ScopeOpen;
+		FConditions::AppendQuery(Result, Rhs);
+		Result.Tokens[Result.TokenCount++] = FConditions::Token::ScopeClose;
 
 		return Result;
 	}
 
-	FQueryConditions operator&&(FColumnBase Lhs, FColumnBase Rhs)
+	FConditions operator&&(FColumnBase Lhs, FColumnBase Rhs)
 	{
-		FQueryConditions Result(Lhs);
+		FConditions Result(Lhs);
 		Result.Columns[Result.ColumnCount++] = Rhs.TypeInfo;
-		Result.Tokens[Result.TokenCount++] = FQueryConditions::Token::And;
+		Result.Tokens[Result.TokenCount++] = FConditions::Token::And;
 		return Result;
 	}
 
-	FQueryConditions operator&&(FColumnBase Lhs, const FQueryConditions& Rhs)
+	FConditions operator&&(FColumnBase Lhs, const FConditions& Rhs)
 	{
-		FQueryConditions Result(Lhs);
-		Result.Tokens[Result.TokenCount++] = FQueryConditions::Token::And;
-		Result.Tokens[Result.TokenCount++] = FQueryConditions::Token::ScopeOpen;
-		FQueryConditions::AppendQuery(Result, Rhs);
-		Result.Tokens[Result.TokenCount++] = FQueryConditions::Token::ScopeClose;
+		FConditions Result(Lhs);
+		Result.Tokens[Result.TokenCount++] = FConditions::Token::And;
+		Result.Tokens[Result.TokenCount++] = FConditions::Token::ScopeOpen;
+		FConditions::AppendQuery(Result, Rhs);
+		Result.Tokens[Result.TokenCount++] = FConditions::Token::ScopeClose;
 		return Result;
 	}
 
-	FQueryConditions operator||(const FQueryConditions& Lhs, FColumnBase Rhs)
+	FConditions operator||(const FConditions& Lhs, FColumnBase Rhs)
 	{
-		FQueryConditions Result = Lhs;
+		FConditions Result = Lhs;
 		Result.Columns[Result.ColumnCount++] = Rhs.TypeInfo;
-		Result.Tokens[Result.TokenCount++] = FQueryConditions::Token::Or;
+		Result.Tokens[Result.TokenCount++] = FConditions::Token::Or;
 		return Result;
 	}
 
-	FQueryConditions operator||(const FQueryConditions& Lhs, const FQueryConditions& Rhs)
+	FConditions operator||(const FConditions& Lhs, const FConditions& Rhs)
 	{
-		FQueryConditions Result;
-		Result.Tokens[Result.TokenCount++] = FQueryConditions::Token::ScopeOpen;
-		FQueryConditions::AppendQuery(Result, Lhs);
-		Result.Tokens[Result.TokenCount++] = FQueryConditions::Token::ScopeClose;
+		FConditions Result;
+		Result.Tokens[Result.TokenCount++] = FConditions::Token::ScopeOpen;
+		FConditions::AppendQuery(Result, Lhs);
+		Result.Tokens[Result.TokenCount++] = FConditions::Token::ScopeClose;
 
-		Result.Tokens[Result.TokenCount++] = FQueryConditions::Token::Or;
+		Result.Tokens[Result.TokenCount++] = FConditions::Token::Or;
 
-		Result.Tokens[Result.TokenCount++] = FQueryConditions::Token::ScopeOpen;
-		FQueryConditions::AppendQuery(Result, Rhs);
-		Result.Tokens[Result.TokenCount++] = FQueryConditions::Token::ScopeClose;
+		Result.Tokens[Result.TokenCount++] = FConditions::Token::ScopeOpen;
+		FConditions::AppendQuery(Result, Rhs);
+		Result.Tokens[Result.TokenCount++] = FConditions::Token::ScopeClose;
 
 		return Result;
 	}
 
-	FQueryConditions operator||(FColumnBase Lhs, FColumnBase Rhs)
+	FConditions operator||(FColumnBase Lhs, FColumnBase Rhs)
 	{
-		FQueryConditions Result(Lhs);
+		FConditions Result(Lhs);
 		Result.Columns[Result.ColumnCount++] = Rhs.TypeInfo;
-		Result.Tokens[Result.TokenCount++] = FQueryConditions::Token::Or;
+		Result.Tokens[Result.TokenCount++] = FConditions::Token::Or;
 		return Result;
 	}
 
-	FQueryConditions operator||(FColumnBase Lhs, const FQueryConditions& Rhs)
+	FConditions operator||(FColumnBase Lhs, const FConditions& Rhs)
 	{
-		FQueryConditions Result(Lhs);
-		Result.Tokens[Result.TokenCount++] = FQueryConditions::Token::Or;
-		Result.Tokens[Result.TokenCount++] = FQueryConditions::Token::ScopeOpen;
-		FQueryConditions::AppendQuery(Result, Rhs);
-		Result.Tokens[Result.TokenCount++] = FQueryConditions::Token::ScopeClose;
+		FConditions Result(Lhs);
+		Result.Tokens[Result.TokenCount++] = FConditions::Token::Or;
+		Result.Tokens[Result.TokenCount++] = FConditions::Token::ScopeOpen;
+		FConditions::AppendQuery(Result, Rhs);
+		Result.Tokens[Result.TokenCount++] = FConditions::Token::ScopeClose;
 		return Result;
 	}
-} // namespace TypedElementDataStorage
\ No newline at end of file
+} // namespace UE::Editor::DataStorage::Queries
\ No newline at end of file
diff --git a/Engine/Source/Runtime/TypedElementFramework/Private/Elements/Interfaces/TypedElementDataStorageUiInterface.cpp b/Engine/Source/Runtime/TypedElementFramework/Private/Elements/Interfaces/TypedElementDataStorageUiInterface.cpp
--- a/Engine/Source/Runtime/TypedElementFramework/Private/Elements/Interfaces/TypedElementDataStorageUiInterface.cpp
+++ b/Engine/Source/Runtime/TypedElementFramework/Private/Elements/Interfaces/TypedElementDataStorageUiInterface.cpp
@@ -14,7 +14,7 @@ FTypedElementWidgetConstructor::FTypedElementWidgetConstructor(const UScriptStru
 }
 
 bool FTypedElementWidgetConstructor::Initialize(const UE::Editor::DataStorage::FMetaDataView& InArguments,
-	TArray<TWeakObjectPtr<const UScriptStruct>> InMatchedColumnTypes, const TypedElementDataStorage::FQueryConditions& InQueryConditions)
+	TArray<TWeakObjectPtr<const UScriptStruct>> InMatchedColumnTypes, const UE::Editor::DataStorage::Queries::FConditions& InQueryConditions)
 {
 	MatchedColumnTypes = MoveTemp(InMatchedColumnTypes);
 	QueryConditions = &InQueryConditions;
@@ -31,7 +31,7 @@ const TArray<TWeakObjectPtr<const UScriptStruct>>& FTypedElementWidgetConstructo
 	return MatchedColumnTypes;
 }
 
-const TypedElementDataStorage::FQueryConditions* FTypedElementWidgetConstructor::GetQueryConditions() const
+const UE::Editor::DataStorage::Queries::FConditions* FTypedElementWidgetConstructor::GetQueryConditions() const
 {
 	return QueryConditions;
 }
diff --git a/Engine/Source/Runtime/TypedElementFramework/Public/Elements/Common/TypedElementQueryConditions.h b/Engine/Source/Runtime/TypedElementFramework/Public/Elements/Common/TypedElementQueryConditions.h
--- a/Engine/Source/Runtime/TypedElementFramework/Public/Elements/Common/TypedElementQueryConditions.h
+++ b/Engine/Source/Runtime/TypedElementFramework/Public/Elements/Common/TypedElementQueryConditions.h
@@ -2,13 +2,13 @@
 
 #pragma once
 
-#include "CoreMinimal.h"
+#include "Elements/Common/TypedElementCommonTypes.h"
 #include "Templates/SubclassOf.h"
 #include "UObject/WeakObjectPtrTemplates.h"
 
 class UScriptStruct;
 
-namespace TypedElementDataStorage
+namespace UE::Editor::DataStorage::Queries
 {
 	struct FColumnBase
 	{
@@ -20,37 +20,37 @@ namespace TypedElementDataStorage
 	};
 
 	template<typename T = void>
-	struct FColumn final : public FColumnBase
+	struct TColumn final : public FColumnBase
 	{
 		template <typename U = T> requires (!std::is_same_v<U, void>)
-			constexpr FColumn() : FColumnBase(T::StaticStruct()) {}
+			constexpr TColumn() : FColumnBase(T::StaticStruct()) {}
 		template <typename U = T> requires (std::is_same_v<U, void>)
-			constexpr explicit FColumn(TWeakObjectPtr<const UScriptStruct> ColumnTypeInfo) : FColumnBase(ColumnTypeInfo) {};
+			constexpr explicit TColumn(TWeakObjectPtr<const UScriptStruct> ColumnTypeInfo) : FColumnBase(ColumnTypeInfo) {};
 	};
 
 	/**
 	 * Product of boolean combination of multiple columns. This can be used to verify if a collection of columns match
 	 * the stored columns.
 	 */
-	class FQueryConditions final
+	class FConditions final
 	{
 	public:
 		using ContainsCallback = TFunctionRef<bool(uint8_t ColumnIndex, TWeakObjectPtr<const UScriptStruct> Column)>;
 
-		TYPEDELEMENTFRAMEWORK_API friend FQueryConditions operator&&(const FQueryConditions& Lhs, FColumnBase Rhs);
-		TYPEDELEMENTFRAMEWORK_API friend FQueryConditions operator&&(const FQueryConditions& Lhs, const FQueryConditions& Rhs);
-		TYPEDELEMENTFRAMEWORK_API friend FQueryConditions operator&&(FColumnBase Lhs, FColumnBase Rhs);
-		TYPEDELEMENTFRAMEWORK_API friend FQueryConditions operator&&(FColumnBase Lhs, const FQueryConditions& Rhs);
+		TYPEDELEMENTFRAMEWORK_API friend FConditions operator&&(const FConditions& Lhs, FColumnBase Rhs);
+		TYPEDELEMENTFRAMEWORK_API friend FConditions operator&&(const FConditions& Lhs, const FConditions& Rhs);
+		TYPEDELEMENTFRAMEWORK_API friend FConditions operator&&(FColumnBase Lhs, FColumnBase Rhs);
+		TYPEDELEMENTFRAMEWORK_API friend FConditions operator&&(FColumnBase Lhs, const FConditions& Rhs);
 
-		TYPEDELEMENTFRAMEWORK_API friend FQueryConditions operator||(const FQueryConditions& Lhs, FColumnBase Rhs);
-		TYPEDELEMENTFRAMEWORK_API friend FQueryConditions operator||(const FQueryConditions& Lhs, const FQueryConditions& Rhs);
-		TYPEDELEMENTFRAMEWORK_API friend FQueryConditions operator||(FColumnBase Lhs, FColumnBase Rhs);
-		TYPEDELEMENTFRAMEWORK_API friend FQueryConditions operator||(FColumnBase Lhs, const FQueryConditions& Rhs);
+		TYPEDELEMENTFRAMEWORK_API friend FConditions operator||(const FConditions& Lhs, FColumnBase Rhs);
+		TYPEDELEMENTFRAMEWORK_API friend FConditions operator||(const FConditions& Lhs, const FConditions& Rhs);
+		TYPEDELEMENTFRAMEWORK_API friend FConditions operator||(FColumnBase Lhs, FColumnBase Rhs);
+		TYPEDELEMENTFRAMEWORK_API friend FConditions operator||(FColumnBase Lhs, const FConditions& Rhs);
 
-		FQueryConditions() = default;
+		FConditions() = default;
 		// Not marked as "explicit" to allow conversion from a column. This means that conditions with a single
 		// argument can be written in the same way as ones that use combinations.
-		TYPEDELEMENTFRAMEWORK_API FQueryConditions(FColumnBase Column);
+		TYPEDELEMENTFRAMEWORK_API FConditions(FColumnBase Column);
 
 		/** Convert the conditions into a string and append them to the provided string. */
 		TYPEDELEMENTFRAMEWORK_API void AppendToString(FString& Output) const;
@@ -102,7 +102,7 @@ namespace TypedElementDataStorage
 
 		uint8_t MinimumColumnMatchRequiredRange(uint8_t& Front) const;
 
-		static void AppendQuery(FQueryConditions& Target, const FQueryConditions& Source);
+		static void AppendQuery(FConditions& Target, const FConditions& Source);
 
 		static constexpr SIZE_T MaxColumnCount = 32;
 		static constexpr SIZE_T MaxTokenCount = 64;
@@ -125,13 +125,13 @@ namespace TypedElementDataStorage
 
 	};
 
-	TYPEDELEMENTFRAMEWORK_API FQueryConditions operator&&(const FQueryConditions& Lhs, FColumnBase Rhs);
-	TYPEDELEMENTFRAMEWORK_API FQueryConditions operator&&(const FQueryConditions& Lhs, const FQueryConditions& Rhs);
-	TYPEDELEMENTFRAMEWORK_API FQueryConditions operator&&(FColumnBase Lhs, FColumnBase Rhs);
-	TYPEDELEMENTFRAMEWORK_API FQueryConditions operator&&(FColumnBase Lhs, const FQueryConditions& Rhs);
+	TYPEDELEMENTFRAMEWORK_API FConditions operator&&(const FConditions& Lhs, FColumnBase Rhs);
+	TYPEDELEMENTFRAMEWORK_API FConditions operator&&(const FConditions& Lhs, const FConditions& Rhs);
+	TYPEDELEMENTFRAMEWORK_API FConditions operator&&(FColumnBase Lhs, FColumnBase Rhs);
+	TYPEDELEMENTFRAMEWORK_API FConditions operator&&(FColumnBase Lhs, const FConditions& Rhs);
 
-	TYPEDELEMENTFRAMEWORK_API FQueryConditions operator||(const FQueryConditions& Lhs, FColumnBase Rhs);
-	TYPEDELEMENTFRAMEWORK_API FQueryConditions operator||(const FQueryConditions& Lhs, const FQueryConditions& Rhs);
-	TYPEDELEMENTFRAMEWORK_API FQueryConditions operator||(FColumnBase Lhs, FColumnBase Rhs);
-	TYPEDELEMENTFRAMEWORK_API FQueryConditions operator||(FColumnBase Lhs, const FQueryConditions& Rhs);
-} // namespace TypedElementDataStorage
\ No newline at end of file
+	TYPEDELEMENTFRAMEWORK_API FConditions operator||(const FConditions& Lhs, FColumnBase Rhs);
+	TYPEDELEMENTFRAMEWORK_API FConditions operator||(const FConditions& Lhs, const FConditions& Rhs);
+	TYPEDELEMENTFRAMEWORK_API FConditions operator||(FColumnBase Lhs, FColumnBase Rhs);
+	TYPEDELEMENTFRAMEWORK_API FConditions operator||(FColumnBase Lhs, const FConditions& Rhs);
+} // namespace UE::Editor::DataStorage::Queries
diff --git a/Engine/Source/Runtime/TypedElementFramework/Public/Elements/Interfaces/TypedElementDataStorageInterface.h b/Engine/Source/Runtime/TypedElementFramework/Public/Elements/Interfaces/TypedElementDataStorageInterface.h
--- a/Engine/Source/Runtime/TypedElementFramework/Public/Elements/Interfaces/TypedElementDataStorageInterface.h
+++ b/Engine/Source/Runtime/TypedElementFramework/Public/Elements/Interfaces/TypedElementDataStorageInterface.h
@@ -303,7 +303,7 @@ public:
 	virtual void ListColumns(RowHandle Row, ColumnListWithDataCallbackRef Callback) = 0;
 
 	/** Determines if the columns in the row match the query conditions. */
-	virtual bool MatchesColumns(RowHandle Row, const TypedElementDataStorage::FQueryConditions& Conditions) const = 0;
+	virtual bool MatchesColumns(RowHandle Row, const UE::Editor::DataStorage::Queries::FConditions& Conditions) const = 0;
 
 	/**
 	 * Finds the type information for a dynamic column.
diff --git a/Engine/Source/Runtime/TypedElementFramework/Public/Elements/Interfaces/TypedElementDataStorageUiInterface.h b/Engine/Source/Runtime/TypedElementFramework/Public/Elements/Interfaces/TypedElementDataStorageUiInterface.h
--- a/Engine/Source/Runtime/TypedElementFramework/Public/Elements/Interfaces/TypedElementDataStorageUiInterface.h
+++ b/Engine/Source/Runtime/TypedElementFramework/Public/Elements/Interfaces/TypedElementDataStorageUiInterface.h
@@ -36,14 +36,14 @@ public:
 
 	/** Initializes a new constructor based on the provided arguments.. */
 	TYPEDELEMENTFRAMEWORK_API virtual bool Initialize(const UE::Editor::DataStorage::FMetaDataView& InArguments,
-		TArray<TWeakObjectPtr<const UScriptStruct>> InMatchedColumnTypes, const UE::Editor::DataStorage::FQueryConditions& InQueryConditions);
+		TArray<TWeakObjectPtr<const UScriptStruct>> InMatchedColumnTypes, const UE::Editor::DataStorage::Queries::FConditions& InQueryConditions);
 
 	/** Retrieves the type information for the constructor type. */
 	TYPEDELEMENTFRAMEWORK_API virtual const UScriptStruct* GetTypeInfo() const;
 	/** Retrieves the columns, if any, that were matched to this constructor when it was created. */
 	TYPEDELEMENTFRAMEWORK_API virtual const TArray<TWeakObjectPtr<const UScriptStruct>>& GetMatchedColumns() const;
 	/** Retrieves the query conditions that need to match for this widget constructor to produce a widget. */
-	TYPEDELEMENTFRAMEWORK_API virtual const UE::Editor::DataStorage::FQueryConditions* GetQueryConditions() const;
+	TYPEDELEMENTFRAMEWORK_API virtual const UE::Editor::DataStorage::Queries::FConditions* GetQueryConditions() const;
 
 	/** Returns a list of additional columns the widget requires to be added to its rows. */
 	TYPEDELEMENTFRAMEWORK_API virtual TConstArrayView<const UScriptStruct*> GetAdditionalColumnsList() const;
@@ -122,7 +122,7 @@ protected:
 protected:
 
 	TArray<TWeakObjectPtr<const UScriptStruct>> MatchedColumnTypes;
-	const UE::Editor::DataStorage::FQueryConditions* QueryConditions = nullptr;
+	const UE::Editor::DataStorage::Queries::FConditions* QueryConditions = nullptr;
 	const UScriptStruct* TypeInfo = nullptr;
 };
 
@@ -284,7 +284,7 @@ public:
 	 * If registration is successful true will be returned otherwise false.
 	 */
 	virtual bool RegisterWidgetFactory(FName Purpose, const UScriptStruct* Constructor,
-		UE::Editor::DataStorage::FQueryConditions Columns) = 0;
+		UE::Editor::DataStorage::Queries::FConditions Columns) = 0;
 	/**
 	 * Registers a widget factory that will be called when the purpose it's registered under is requested.
 	 * This version registers a generic type. Construction using these are typically cheaper as they can avoid
@@ -294,7 +294,7 @@ public:
 	 * If registration is successful true will be returned otherwise false.
 	 */
 	template<typename ConstructorType>
-	bool RegisterWidgetFactory(FName Purpose, UE::Editor::DataStorage::FQueryConditions Columns);
+	bool RegisterWidgetFactory(FName Purpose, UE::Editor::DataStorage::Queries::FConditions Columns);
 	/**
 	 * Registers a widget factory that will be called when the purpose it's registered under is requested.
 	 * This version uses a previously created instance of the Constructor. The benefit of this is that it store
@@ -312,7 +312,7 @@ public:
 	 * If registration is successful true will be returned otherwise false.
 	 */
 	virtual bool RegisterWidgetFactory(FName Purpose, TUniquePtr<FTypedElementWidgetConstructor>&& Constructor, 
-		UE::Editor::DataStorage::FQueryConditions Columns) = 0;
+		UE::Editor::DataStorage::Queries::FConditions Columns) = 0;
 	
 	/** 
 	 * Creates widget constructors for the requested purpose.
@@ -367,7 +367,7 @@ bool ITypedElementDataStorageUiInterface::RegisterWidgetFactory(FName Purpose)
 }
 
 template<typename ConstructorType>
-bool ITypedElementDataStorageUiInterface::RegisterWidgetFactory(FName Purpose, TypedElementDataStorage::FQueryConditions Columns)
+bool ITypedElementDataStorageUiInterface::RegisterWidgetFactory(FName Purpose, UE::Editor::DataStorage::Queries::FConditions Columns)
 {
 	return this->RegisterWidgetFactory(Purpose, ConstructorType::StaticStruct(), Columns);
 }
diff --git a/Engine/Source/Runtime/TypedElementFramework/Tests/Elements/Framework/TypedElementQueryBuilderTests.cpp b/Engine/Source/Runtime/TypedElementFramework/Tests/Elements/Framework/TypedElementQueryBuilderTests.cpp
--- a/Engine/Source/Runtime/TypedElementFramework/Tests/Elements/Framework/TypedElementQueryBuilderTests.cpp
+++ b/Engine/Source/Runtime/TypedElementFramework/Tests/Elements/Framework/TypedElementQueryBuilderTests.cpp
@@ -11,9 +11,11 @@
 #include "Tests/TestHarnessAdapter.h"
 #include "UObject/Class.h"
 
-using namespace TypedElementDataStorage;
-
-
+namespace UE::Editor::DataStorage
+{
+	using namespace Queries;
+namespace Tests
+{
 static void AppendColumnName(FString& Output, TWeakObjectPtr<const UScriptStruct> TypeInfo)
 {
 #if WITH_EDITORONLY_DATA
@@ -27,7 +29,7 @@ static void AppendColumnName(FString& Output, TWeakObjectPtr<const UScriptStruct
 #endif
 }
 
-static bool TestMatching(const FQueryConditions& TestQuery, const TArray<FColumnBase>& RequestedColumns, bool Expected, bool Sort = false)
+static bool TestMatching(const FConditions& TestQuery, const TArray<FColumnBase>& RequestedColumns, bool Expected, bool Sort = false)
 {
 	if (Sort)
 	{
@@ -101,173 +103,175 @@ static bool TestMatching(const FQueryConditions& TestQuery, const TArray<FColumn
 
 TEST_CASE_NAMED(FTypedElementQueryConditions_NoColumn, "TypedElementQueryBuilder::FTypedElementQueryConditions_NoColumn", "[ApplicationContextMask][EngineFilter]")
 {
-	FQueryConditions Example;
+	FConditions Example;
 	
 	CHECK(Example.MinimumColumnMatchRequired() == 0);
 	// Since there are no restrictions provided in the query, all input passes.
-	CHECK(TestMatching(Example, { TypedElementDataStorage::FColumn<FTestColumnA>() }, true));
+	CHECK(TestMatching(Example, { TColumn<FTestColumnA>() }, true));
 }
 
 TEST_CASE_NAMED(FTypedElementQueryConditions_OneColumn, "TypedElementQueryBuilder::FTypedElementQueryConditions_OneColumn", "[ApplicationContextMask][EngineFilter]")
 {
-	FQueryConditions Example{ TypedElementDataStorage::FColumn<FTestColumnA>() };
+	FConditions Example{ TColumn<FTestColumnA>() };
 
 	CHECK(Example.MinimumColumnMatchRequired() == 1);
-	CHECK(TestMatching(Example, { TypedElementDataStorage::FColumn<FTestColumnA>() }, true));
+	CHECK(TestMatching(Example, { TColumn<FTestColumnA>() }, true));
 }
 
 TEST_CASE_NAMED(FTypedElementQueryConditions1, "TypedElementQueryBuilder::FTypedElementQueryConditions A && B && C", "[ApplicationContextMask][EngineFilter]")
 {
-	TypedElementDataStorage::FColumn TestA(FTestColumnA::StaticStruct());
+	TColumn<FTestColumnA> TestA;
 	
-	FQueryConditions Example = TypedElementDataStorage::FColumn<FTestColumnA>() && TypedElementDataStorage::FColumn<FTestColumnB>() && TypedElementDataStorage::FColumn<FTestColumnC>();
+	FConditions Example = TColumn<FTestColumnA>() && TColumn<FTestColumnB>() && TColumn<FTestColumnC>();
 	
 	CHECK(Example.MinimumColumnMatchRequired() == 3);
-	CHECK(TestMatching(Example, { TypedElementDataStorage::FColumn<FTestColumnA>(), TypedElementDataStorage::FColumn<FTestColumnB>(), TypedElementDataStorage::FColumn<FTestColumnC>() }, true));
-	CHECK(TestMatching(Example, { TypedElementDataStorage::FColumn<FTestColumnA>(), TypedElementDataStorage::FColumn<FTestColumnB>(), TypedElementDataStorage::FColumn<FTestColumnD>() }, false));
+	CHECK(TestMatching(Example, { TColumn<FTestColumnA>(), TColumn<FTestColumnB>(), TColumn<FTestColumnC>() }, true));
+	CHECK(TestMatching(Example, { TColumn<FTestColumnA>(), TColumn<FTestColumnB>(), TColumn<FTestColumnD>() }, false));
 }
 
 TEST_CASE_NAMED(FTypedElementQueryConditions2, "TypedElementQueryBuilder::FTypedElementQueryConditions A || B || C", "[ApplicationContextMask][EngineFilter]")
 {
-	FQueryConditions Example = TypedElementDataStorage::FColumn<FTestColumnA>() || TypedElementDataStorage::FColumn<FTestColumnB>() || TypedElementDataStorage::FColumn<FTestColumnC>();
+	FConditions Example = TColumn<FTestColumnA>() || TColumn<FTestColumnB>() || TColumn<FTestColumnC>();
 	
 	CHECK(Example.MinimumColumnMatchRequired() == 1);
-	CHECK(TestMatching(Example, { TypedElementDataStorage::FColumn<FTestColumnB>() }, true));
-	CHECK(TestMatching(Example, { TypedElementDataStorage::FColumn<FTestColumnB>(), TypedElementDataStorage::FColumn<FTestColumnC>() }, true));
-	CHECK(TestMatching(Example, { TypedElementDataStorage::FColumn<FTestColumnD>() }, false));
+	CHECK(TestMatching(Example, { TColumn<FTestColumnB>() }, true));
+	CHECK(TestMatching(Example, { TColumn<FTestColumnB>(), TColumn<FTestColumnC>() }, true));
+	CHECK(TestMatching(Example, { TColumn<FTestColumnD>() }, false));
 }
 
 TEST_CASE_NAMED(FTypedElementQueryConditions3, "TypedElementQueryBuilder::FTypedElementQueryConditions A && (B || C)", "[ApplicationContextMask][EngineFilter]")
 {
-	FQueryConditions Example = TypedElementDataStorage::FColumn<FTestColumnA>() && (TypedElementDataStorage::FColumn<FTestColumnB>() || TypedElementDataStorage::FColumn<FTestColumnC>());
+	FConditions Example = TColumn<FTestColumnA>() && (TColumn<FTestColumnB>() || TColumn<FTestColumnC>());
 	
 	CHECK(Example.MinimumColumnMatchRequired() == 2);
-	CHECK(TestMatching(Example, { TypedElementDataStorage::FColumn<FTestColumnA>(), TypedElementDataStorage::FColumn<FTestColumnB>() }, true));
-	CHECK(TestMatching(Example, { TypedElementDataStorage::FColumn<FTestColumnA>(), TypedElementDataStorage::FColumn<FTestColumnC>() }, true));
-	CHECK(TestMatching(Example, { TypedElementDataStorage::FColumn<FTestColumnA>(), TypedElementDataStorage::FColumn<FTestColumnD>() }, false));
-	CHECK(TestMatching(Example, { TypedElementDataStorage::FColumn<FTestColumnD>(), TypedElementDataStorage::FColumn<FTestColumnB>() }, false));
+	CHECK(TestMatching(Example, { TColumn<FTestColumnA>(), TColumn<FTestColumnB>() }, true));
+	CHECK(TestMatching(Example, { TColumn<FTestColumnA>(), TColumn<FTestColumnC>() }, true));
+	CHECK(TestMatching(Example, { TColumn<FTestColumnA>(), TColumn<FTestColumnD>() }, false));
+	CHECK(TestMatching(Example, { TColumn<FTestColumnD>(), TColumn<FTestColumnB>() }, false));
 }
 
 TEST_CASE_NAMED(FTypedElementQueryConditions4, "TypedElementQueryBuilder::FTypedElementQueryConditions A && (B || C) && (D || E)", "[ApplicationContextMask][EngineFilter]")
 {
-	FQueryConditions Example = 
-		TypedElementDataStorage::FColumn<FTestColumnA>() && 
-		(TypedElementDataStorage::FColumn<FTestColumnB>() || TypedElementDataStorage::FColumn<FTestColumnC>()) &&
-		(TypedElementDataStorage::FColumn<FTestColumnD>() || TypedElementDataStorage::FColumn<FTestColumnE>());
+	FConditions Example = 
+		TColumn<FTestColumnA>() && 
+		(TColumn<FTestColumnB>() || TColumn<FTestColumnC>()) &&
+		(TColumn<FTestColumnD>() || TColumn<FTestColumnE>());
 	
 	CHECK(Example.MinimumColumnMatchRequired() == 3);
 	
-	CHECK(TestMatching(Example, { TypedElementDataStorage::FColumn<FTestColumnA>() }, false));
-	CHECK(TestMatching(Example, { TypedElementDataStorage::FColumn<FTestColumnA>(), TypedElementDataStorage::FColumn<FTestColumnB>() }, false));
+	CHECK(TestMatching(Example, { TColumn<FTestColumnA>() }, false));
+	CHECK(TestMatching(Example, { TColumn<FTestColumnA>(), TColumn<FTestColumnB>() }, false));
 
-	CHECK(TestMatching(Example, { TypedElementDataStorage::FColumn<FTestColumnA>(), TypedElementDataStorage::FColumn<FTestColumnB>(), TypedElementDataStorage::FColumn<FTestColumnD>() }, true));
-	CHECK(TestMatching(Example, { TypedElementDataStorage::FColumn<FTestColumnA>(), TypedElementDataStorage::FColumn<FTestColumnB>(), TypedElementDataStorage::FColumn<FTestColumnE>() }, true));
-	CHECK(TestMatching(Example, { TypedElementDataStorage::FColumn<FTestColumnA>(), TypedElementDataStorage::FColumn<FTestColumnC>(), TypedElementDataStorage::FColumn<FTestColumnD>() }, true));
+	CHECK(TestMatching(Example, { TColumn<FTestColumnA>(), TColumn<FTestColumnB>(), TColumn<FTestColumnD>() }, true));
+	CHECK(TestMatching(Example, { TColumn<FTestColumnA>(), TColumn<FTestColumnB>(), TColumn<FTestColumnE>() }, true));
+	CHECK(TestMatching(Example, { TColumn<FTestColumnA>(), TColumn<FTestColumnC>(), TColumn<FTestColumnD>() }, true));
 
-	CHECK(TestMatching(Example, { TypedElementDataStorage::FColumn<FTestColumnA>(), TypedElementDataStorage::FColumn<FTestColumnC>(), TypedElementDataStorage::FColumn<FTestColumnF>() }, false));
-	CHECK(TestMatching(Example, { TypedElementDataStorage::FColumn<FTestColumnA>(), TypedElementDataStorage::FColumn<FTestColumnF>(), TypedElementDataStorage::FColumn<FTestColumnD>() }, false));
-	CHECK(TestMatching(Example, { TypedElementDataStorage::FColumn<FTestColumnB>(), TypedElementDataStorage::FColumn<FTestColumnC>(), TypedElementDataStorage::FColumn<FTestColumnD>() }, false));
+	CHECK(TestMatching(Example, { TColumn<FTestColumnA>(), TColumn<FTestColumnC>(), TColumn<FTestColumnF>() }, false));
+	CHECK(TestMatching(Example, { TColumn<FTestColumnA>(), TColumn<FTestColumnF>(), TColumn<FTestColumnD>() }, false));
+	CHECK(TestMatching(Example, { TColumn<FTestColumnB>(), TColumn<FTestColumnC>(), TColumn<FTestColumnD>() }, false));
 }
 
 TEST_CASE_NAMED(FTypedElementQueryConditions5, "TypedElementQueryBuilder::FTypedElementQueryConditions (A || B) && (C || D) && (E || F)", "[ApplicationContextMask][EngineFilter]")
 {
-	FQueryConditions Example =
-		(TypedElementDataStorage::FColumn<FTestColumnA>() || TypedElementDataStorage::FColumn<FTestColumnB>()) &&
-		(TypedElementDataStorage::FColumn<FTestColumnC>() || TypedElementDataStorage::FColumn<FTestColumnD>()) &&
-		(TypedElementDataStorage::FColumn<FTestColumnE>() || TypedElementDataStorage::FColumn<FTestColumnF>());
+	FConditions Example =
+		(TColumn<FTestColumnA>() || TColumn<FTestColumnB>()) &&
+		(TColumn<FTestColumnC>() || TColumn<FTestColumnD>()) &&
+		(TColumn<FTestColumnE>() || TColumn<FTestColumnF>());
 	
 	CHECK(Example.MinimumColumnMatchRequired() == 3); 
 	
-	CHECK(TestMatching(Example, { TypedElementDataStorage::FColumn<FTestColumnA>() }, false));
-	CHECK(TestMatching(Example, { TypedElementDataStorage::FColumn<FTestColumnA>(), TypedElementDataStorage::FColumn<FTestColumnC>() }, false));
+	CHECK(TestMatching(Example, { TColumn<FTestColumnA>() }, false));
+	CHECK(TestMatching(Example, { TColumn<FTestColumnA>(), TColumn<FTestColumnC>() }, false));
 	
-	CHECK(TestMatching(Example, { TypedElementDataStorage::FColumn<FTestColumnA>(), TypedElementDataStorage::FColumn<FTestColumnC>(), TypedElementDataStorage::FColumn<FTestColumnE>() }, true));
-	CHECK(TestMatching(Example, { TypedElementDataStorage::FColumn<FTestColumnB>(), TypedElementDataStorage::FColumn<FTestColumnC>(), TypedElementDataStorage::FColumn<FTestColumnE>() }, true));
+	CHECK(TestMatching(Example, { TColumn<FTestColumnA>(), TColumn<FTestColumnC>(), TColumn<FTestColumnE>() }, true));
+	CHECK(TestMatching(Example, { TColumn<FTestColumnB>(), TColumn<FTestColumnC>(), TColumn<FTestColumnE>() }, true));
 	
-	CHECK(TestMatching(Example, { TypedElementDataStorage::FColumn<FTestColumnA>(), TypedElementDataStorage::FColumn<FTestColumnC>(), TypedElementDataStorage::FColumn<FTestColumnG>() }, false));
-	CHECK(TestMatching(Example, { TypedElementDataStorage::FColumn<FTestColumnA>(), TypedElementDataStorage::FColumn<FTestColumnG>(), TypedElementDataStorage::FColumn<FTestColumnD>() }, false));
-	CHECK(TestMatching(Example, { TypedElementDataStorage::FColumn<FTestColumnG>(), TypedElementDataStorage::FColumn<FTestColumnC>(), TypedElementDataStorage::FColumn<FTestColumnD>() }, false));
+	CHECK(TestMatching(Example, { TColumn<FTestColumnA>(), TColumn<FTestColumnC>(), TColumn<FTestColumnG>() }, false));
+	CHECK(TestMatching(Example, { TColumn<FTestColumnA>(), TColumn<FTestColumnG>(), TColumn<FTestColumnD>() }, false));
+	CHECK(TestMatching(Example, { TColumn<FTestColumnG>(), TColumn<FTestColumnC>(), TColumn<FTestColumnD>() }, false));
 }
 
 TEST_CASE_NAMED(FTypedElementQueryConditions6, "TypedElementQueryBuilder::FTypedElementQueryConditions ((A || B) && (C || D)) || (E && F)", "[ApplicationContextMask][EngineFilter]")
 {
-	FQueryConditions Example =
+	FConditions Example =
 		(
-			(TypedElementDataStorage::FColumn<FTestColumnA>() || TypedElementDataStorage::FColumn<FTestColumnB>()) &&
-			(TypedElementDataStorage::FColumn<FTestColumnC>() || TypedElementDataStorage::FColumn<FTestColumnD>())
+			(TColumn<FTestColumnA>() || TColumn<FTestColumnB>()) &&
+			(TColumn<FTestColumnC>() || TColumn<FTestColumnD>())
 		) ||
-		(TypedElementDataStorage::FColumn<FTestColumnE>() && TypedElementDataStorage::FColumn<FTestColumnF>());
+		(TColumn<FTestColumnE>() && TColumn<FTestColumnF>());
 	
 	CHECK(Example.MinimumColumnMatchRequired() == 2);
 
-	CHECK(TestMatching(Example, { TypedElementDataStorage::FColumn<FTestColumnA>() }, false));
-	CHECK(TestMatching(Example, { TypedElementDataStorage::FColumn<FTestColumnA>(), TypedElementDataStorage::FColumn<FTestColumnC>() }, true));
+	CHECK(TestMatching(Example, { TColumn<FTestColumnA>() }, false));
+	CHECK(TestMatching(Example, { TColumn<FTestColumnA>(), TColumn<FTestColumnC>() }, true));
 
-	CHECK(TestMatching(Example, { TypedElementDataStorage::FColumn<FTestColumnE>(), TypedElementDataStorage::FColumn<FTestColumnF>() }, true));
-	CHECK(TestMatching(Example, { TypedElementDataStorage::FColumn<FTestColumnG>() }, false));
+	CHECK(TestMatching(Example, { TColumn<FTestColumnE>(), TColumn<FTestColumnF>() }, true));
+	CHECK(TestMatching(Example, { TColumn<FTestColumnG>() }, false));
 }
 
 TEST_CASE_NAMED(FTypedElementQueryConditions7, "TypedElementQueryBuilder::FTypedElementQueryConditions (A && B) || (C && D) || (E && F)", "[ApplicationContextMask][EngineFilter]")
 {
-	FQueryConditions Example =
-		(TypedElementDataStorage::FColumn<FTestColumnA>() && TypedElementDataStorage::FColumn<FTestColumnB>()) ||
-		(TypedElementDataStorage::FColumn<FTestColumnC>() && TypedElementDataStorage::FColumn<FTestColumnD>()) ||
-		(TypedElementDataStorage::FColumn<FTestColumnE>() && TypedElementDataStorage::FColumn<FTestColumnF>());
+	FConditions Example =
+		(TColumn<FTestColumnA>() && TColumn<FTestColumnB>()) ||
+		(TColumn<FTestColumnC>() && TColumn<FTestColumnD>()) ||
+		(TColumn<FTestColumnE>() && TColumn<FTestColumnF>());
 
 	CHECK(Example.MinimumColumnMatchRequired() == 2);
 
-	CHECK(TestMatching(Example, { TypedElementDataStorage::FColumn<FTestColumnA>() }, false));
-	CHECK(TestMatching(Example, { TypedElementDataStorage::FColumn<FTestColumnA>(), TypedElementDataStorage::FColumn<FTestColumnB>() }, true));
-	CHECK(TestMatching(Example, { TypedElementDataStorage::FColumn<FTestColumnC>(), TypedElementDataStorage::FColumn<FTestColumnD>() }, true));
-	CHECK(TestMatching(Example, { TypedElementDataStorage::FColumn<FTestColumnE>(), TypedElementDataStorage::FColumn<FTestColumnF>() }, true));
+	CHECK(TestMatching(Example, { TColumn<FTestColumnA>() }, false));
+	CHECK(TestMatching(Example, { TColumn<FTestColumnA>(), TColumn<FTestColumnB>() }, true));
+	CHECK(TestMatching(Example, { TColumn<FTestColumnC>(), TColumn<FTestColumnD>() }, true));
+	CHECK(TestMatching(Example, { TColumn<FTestColumnE>(), TColumn<FTestColumnF>() }, true));
 
-	CHECK(TestMatching(Example, { TypedElementDataStorage::FColumn<FTestColumnC>(), TypedElementDataStorage::FColumn<FTestColumnD>(), TypedElementDataStorage::FColumn<FTestColumnE>(), TypedElementDataStorage::FColumn<FTestColumnF>() }, true));
+	CHECK(TestMatching(Example, { TColumn<FTestColumnC>(), TColumn<FTestColumnD>(), TColumn<FTestColumnE>(), TColumn<FTestColumnF>() }, true));
 
 
-	CHECK(TestMatching(Example, { TypedElementDataStorage::FColumn<FTestColumnE>() }, false));
-	CHECK(TestMatching(Example, { TypedElementDataStorage::FColumn<FTestColumnG>() }, false));
+	CHECK(TestMatching(Example, { TColumn<FTestColumnE>() }, false));
+	CHECK(TestMatching(Example, { TColumn<FTestColumnG>() }, false));
 }
 
 TEST_CASE_NAMED(FTypedElementQueryConditions_MultiMatch, "TypedElementQueryBuilder::FTypedElementQueryConditions_MultiMatch", "[ApplicationContextMask][EngineFilter]")
 {
-	FQueryConditions Example =
-		(TypedElementDataStorage::FColumn<FTestColumnA>() || TypedElementDataStorage::FColumn<FTestColumnB>()) &&
-		(TypedElementDataStorage::FColumn<FTestColumnC>() || TypedElementDataStorage::FColumn<FTestColumnD>()) &&
-		(TypedElementDataStorage::FColumn<FTestColumnE>() || TypedElementDataStorage::FColumn<FTestColumnF>());
+	FConditions Example =
+		(TColumn<FTestColumnA>() || TColumn<FTestColumnB>()) &&
+		(TColumn<FTestColumnC>() || TColumn<FTestColumnD>()) &&
+		(TColumn<FTestColumnE>() || TColumn<FTestColumnF>());
 
 	CHECK(TestMatching(Example, 
 		{ 
-			TypedElementDataStorage::FColumn<FTestColumnA>(), 
-			TypedElementDataStorage::FColumn<FTestColumnB>(), 
-			TypedElementDataStorage::FColumn<FTestColumnC>(), 
-			TypedElementDataStorage::FColumn<FTestColumnD>(),
-			TypedElementDataStorage::FColumn<FTestColumnE>(),
-			TypedElementDataStorage::FColumn<FTestColumnF>()
+			TColumn<FTestColumnA>(), 
+			TColumn<FTestColumnB>(), 
+			TColumn<FTestColumnC>(), 
+			TColumn<FTestColumnD>(),
+			TColumn<FTestColumnE>(),
+			TColumn<FTestColumnF>()
 		}, true));
 
 	CHECK(TestMatching(Example,
 		{
-			TypedElementDataStorage::FColumn<FTestColumnA>(),
-			TypedElementDataStorage::FColumn<FTestColumnC>(),
-			TypedElementDataStorage::FColumn<FTestColumnE>(),
-			TypedElementDataStorage::FColumn<FTestColumnG>()
+			TColumn<FTestColumnA>(),
+			TColumn<FTestColumnC>(),
+			TColumn<FTestColumnE>(),
+			TColumn<FTestColumnG>()
 		}, true));
 }
 
 TEST_CASE_NAMED(FTypedElementQueryConditions_Sorted, "TypedElementQueryBuilder::FTypedElementQueryConditions_Sorted", "[ApplicationContextMask][EngineFilter]")
 {
-	FQueryConditions Example =
-		(TypedElementDataStorage::FColumn<FTestColumnA>() && TypedElementDataStorage::FColumn<FTestColumnB>()) ||
-		(TypedElementDataStorage::FColumn<FTestColumnC>() && TypedElementDataStorage::FColumn<FTestColumnD>()) ||
-		(TypedElementDataStorage::FColumn<FTestColumnE>() && TypedElementDataStorage::FColumn<FTestColumnF>());
+	FConditions Example =
+		(TColumn<FTestColumnA>() && TColumn<FTestColumnB>()) ||
+		(TColumn<FTestColumnC>() && TColumn<FTestColumnD>()) ||
+		(TColumn<FTestColumnE>() && TColumn<FTestColumnF>());
 
 	CHECK(Example.MinimumColumnMatchRequired() == 2);
 
-	CHECK(TestMatching(Example, { TypedElementDataStorage::FColumn<FTestColumnA>(), TypedElementDataStorage::FColumn<FTestColumnB>() }, true, true));
-	CHECK(TestMatching(Example, { TypedElementDataStorage::FColumn<FTestColumnC>(), TypedElementDataStorage::FColumn<FTestColumnD>() }, true, true));
-	CHECK(TestMatching(Example, { TypedElementDataStorage::FColumn<FTestColumnE>(), TypedElementDataStorage::FColumn<FTestColumnF>() }, true, true));
+	CHECK(TestMatching(Example, { TColumn<FTestColumnA>(), TColumn<FTestColumnB>() }, true, true));
+	CHECK(TestMatching(Example, { TColumn<FTestColumnC>(), TColumn<FTestColumnD>() }, true, true));
+	CHECK(TestMatching(Example, { TColumn<FTestColumnE>(), TColumn<FTestColumnF>() }, true, true));
 
-	CHECK(TestMatching(Example, { TypedElementDataStorage::FColumn<FTestColumnC>(), TypedElementDataStorage::FColumn<FTestColumnD>(), TypedElementDataStorage::FColumn<FTestColumnE>(), TypedElementDataStorage::FColumn<FTestColumnF>() }, true, true));
+	CHECK(TestMatching(Example, { TColumn<FTestColumnC>(), TColumn<FTestColumnD>(), TColumn<FTestColumnE>(), TColumn<FTestColumnF>() }, true, true));
 }
+} // namespace Tests
+} // namespace UE::Editor::DataStorage
 
 #endif // WITH_TESTS
