[TypedElementFramework]
* Moved dll export to be per function instead of through specialization. This should solve linker errors on non-win/mac/linux platforms when building non-monolithic

diff --git a/Engine/Source/Runtime/TypedElementFramework/Private/Elements/Framework/TypedElementList.cpp b/Engine/Source/Runtime/TypedElementFramework/Private/Elements/Framework/TypedElementList.cpp
--- a/Engine/Source/Runtime/TypedElementFramework/Private/Elements/Framework/TypedElementList.cpp
+++ b/Engine/Source/Runtime/TypedElementFramework/Private/Elements/Framework/TypedElementList.cpp
@@ -513,8 +513,8 @@ void TTypedElementList<HandleType>::NoteListChanged(const EChangeType InChangeTy
 	}
 }
 
-template class TYPEDELEMENTFRAMEWORK_API TTypedElementList<FTypedElementHandle>;
-template class TYPEDELEMENTFRAMEWORK_API TTypedElementList<FScriptTypedElementHandle>;
+template class TTypedElementList<FTypedElementHandle>;
+template class TTypedElementList<FScriptTypedElementHandle>;
 
 
 namespace UE::TypedElementFramework
diff --git a/Engine/Source/Runtime/TypedElementFramework/Public/Elements/Framework/TypedElementList.h b/Engine/Source/Runtime/TypedElementFramework/Public/Elements/Framework/TypedElementList.h
--- a/Engine/Source/Runtime/TypedElementFramework/Public/Elements/Framework/TypedElementList.h
+++ b/Engine/Source/Runtime/TypedElementFramework/Public/Elements/Framework/TypedElementList.h
@@ -67,13 +67,13 @@ public:
 	using TTypedElementListConstRef = TSharedRef<const TTypedElementList>;
 
 
-	~TTypedElementList();
+	TYPEDELEMENTFRAMEWORK_API ~TTypedElementList();
 
 	/**
 	 * Clone this list instance.
 	 * @note Only copies elements; does not copy any bindings!
 	 */
-	TTypedElementListRef Clone() const;
+	TYPEDELEMENTFRAMEWORK_API TTypedElementListRef Clone() const;
 
 	/**
 	 * Get the element handle at the given index.
@@ -224,34 +224,34 @@ public:
 	/**
 	 * Get the element interface from the given handle.
 	 */
-	UObject* GetElementInterface(const HandleType& InElementHandle, const TSubclassOf<UInterface>& InBaseInterfaceType) const;
+	TYPEDELEMENTFRAMEWORK_API UObject* GetElementInterface(const HandleType& InElementHandle, const TSubclassOf<UInterface>& InBaseInterfaceType) const;
 
 	/**
 	 * Test whether there are elements in this list, optionally filtering to elements that implement the given interface.
 	 */
-	bool HasElements(const TSubclassOf<UInterface>& InBaseInterfaceType = nullptr) const;
+	TYPEDELEMENTFRAMEWORK_API bool HasElements(const TSubclassOf<UInterface>& InBaseInterfaceType = nullptr) const;
 
 	/**
 	 * Count the number of elements in this list, optionally filtering to elements that implement the given interface.
 	 */
-	int32 CountElements(const TSubclassOf<UInterface>& InBaseInterfaceType = nullptr) const;
+	TYPEDELEMENTFRAMEWORK_API int32 CountElements(const TSubclassOf<UInterface>& InBaseInterfaceType = nullptr) const;
 
 	/**
 	 * Test whether there are elements in this list of the given type.
 	 */
-	bool HasElementsOfType(const FName InElementTypeName) const;
-	bool HasElementsOfType(const FTypedHandleTypeId InElementTypeId) const;
+	TYPEDELEMENTFRAMEWORK_API bool HasElementsOfType(const FName InElementTypeName) const;
+	TYPEDELEMENTFRAMEWORK_API bool HasElementsOfType(const FTypedHandleTypeId InElementTypeId) const;
 
 	/**
 	 * Count the number of elements in this list of the given type.
 	 */
-	int32 CountElementsOfType(const FName InElementTypeName) const;
-	int32 CountElementsOfType(const FTypedHandleTypeId InElementTypeId) const;
+	TYPEDELEMENTFRAMEWORK_API int32 CountElementsOfType(const FName InElementTypeName) const;
+	TYPEDELEMENTFRAMEWORK_API int32 CountElementsOfType(const FTypedHandleTypeId InElementTypeId) const;
 
 	/**
 	 * Get the handle of every element in this list, optionally filtering to elements that implement the given interface.
 	 */
-	TArray<HandleType> GetElementHandles(const TSubclassOf<UInterface>& InBaseInterfaceType = nullptr) const;
+	TYPEDELEMENTFRAMEWORK_API TArray<HandleType> GetElementHandles(const TSubclassOf<UInterface>& InBaseInterfaceType = nullptr) const;
 
 	/**
 	 * Get the handle of every element in this list, optionally filtering to elements that implement the given interface.
@@ -272,7 +272,7 @@ public:
 	 * Enumerate the handle of every element in this list, optionally filtering to elements that implement the given interface.
 	 * @note Return true from the callback to continue enumeration.
 	 */
-	void ForEachElementHandle(TFunctionRef<bool(const HandleType&)> InCallback, const TSubclassOf<UInterface>& InBaseInterfaceType = nullptr) const;
+	TYPEDELEMENTFRAMEWORK_API void ForEachElementHandle(TFunctionRef<bool(const HandleType&)> InCallback, const TSubclassOf<UInterface>& InBaseInterfaceType = nullptr) const;
 
 	/**
 	 * Enumerate the elements in this list that implement the given interface.
@@ -558,7 +558,7 @@ public:
 	 * Invoke the delegate called whenever this element list has been changed.
 	 * @return true if a change notification was emitted
 	 */
-	bool NotifyPendingChanges();
+	TYPEDELEMENTFRAMEWORK_API bool NotifyPendingChanges();
 
 	/**
 	 * A utility struct that help to cancel any new pending notification that happened in a scope.
@@ -672,24 +672,24 @@ public:
 	/**
 	 * Get a scoped object that when destroyed it clear a pending change notification without emitting the notification if it happened during its lifecycle.
 	 */
-	FScopedClearNewPendingChange GetScopedClearNewPendingChange();
+	TYPEDELEMENTFRAMEWORK_API FScopedClearNewPendingChange GetScopedClearNewPendingChange();
 
 	/**
 	 * Access the interface to allow external systems (such as USelection) to receive immediate sync notifications as an element list is changed.
 	 * This exists purely as a bridging mechanism and shouldn't be relied on for new code. It is lazily created as needed.
 	 */
-	FLegacySync& Legacy_GetSync();
+	TYPEDELEMENTFRAMEWORK_API FLegacySync& Legacy_GetSync();
 
 	/**
 	 * Access the interface to allow external systems (such as USelection) to receive immediate sync notifications as an element list is changed.
 	 * This exists purely as a bridging mechanism and shouldn't be relied on for new code. This will return null if no legacy sync has been created for this instance.
 	 */
-	FLegacySync* Legacy_GetSyncPtr() const;
+	TYPEDELEMENTFRAMEWORK_API FLegacySync* Legacy_GetSyncPtr() const;
 
 	/**
 	 * Internal function used by the element registry to create an element list instance.
 	 */
-	static TTypedElementListRef Private_CreateElementList(UTypedElementRegistry* InRegistry);
+	static TYPEDELEMENTFRAMEWORK_API TTypedElementListRef Private_CreateElementList(UTypedElementRegistry* InRegistry);
 
 	UTypedElementRegistry* GetRegistry() const
 	{
@@ -718,15 +718,15 @@ private:
 		Cleared,
 	};
 
-	explicit TTypedElementList(UTypedElementRegistry* InRegistry);
+	TYPEDELEMENTFRAMEWORK_API explicit TTypedElementList(UTypedElementRegistry* InRegistry);
 
-	bool AddElementImpl(HandleType&& InElementHandle);
-	bool RemoveElementImpl(const FTypedElementId& InElementId);
-	int32 RemoveAllElementsImpl(TFunctionRef<bool(const HandleType&)> InPredicate);
-	bool ContainsElementImpl(const FTypedElementId& InElementId) const;
+	TYPEDELEMENTFRAMEWORK_API bool AddElementImpl(HandleType&& InElementHandle);
+	TYPEDELEMENTFRAMEWORK_API bool RemoveElementImpl(const FTypedElementId& InElementId);
+	TYPEDELEMENTFRAMEWORK_API int32 RemoveAllElementsImpl(TFunctionRef<bool(const HandleType&)> InPredicate);
+	TYPEDELEMENTFRAMEWORK_API bool ContainsElementImpl(const FTypedElementId& InElementId) const;
 
-	void NoteListMayChange();
-	void NoteListChanged(const EChangeType InChangeType, const HandleType& InElementHandle = HandleType());
+	TYPEDELEMENTFRAMEWORK_API void NoteListMayChange();
+	TYPEDELEMENTFRAMEWORK_API void NoteListChanged(const EChangeType InChangeType, const HandleType& InElementHandle = HandleType());
 
 	/**
 	 * Element registry this element list is associated with.
