[Teds] Removing the final TypedElementDataStorage namespace and DSI using declarations

Also cleaning up some of the nested namespaces per last CR, and some CoreMinimal.h includes switched to ICWYU

#rnx

diff --git a/Engine/Plugins/Developer/PerforceSourceControl/Source/PerforceSourceControl/Private/PerforceSourceControlOperations.cpp b/Engine/Plugins/Developer/PerforceSourceControl/Source/PerforceSourceControl/Private/PerforceSourceControlOperations.cpp
--- a/Engine/Plugins/Developer/PerforceSourceControl/Source/PerforceSourceControl/Private/PerforceSourceControlOperations.cpp
+++ b/Engine/Plugins/Developer/PerforceSourceControl/Source/PerforceSourceControl/Private/PerforceSourceControlOperations.cpp
@@ -2021,8 +2021,6 @@ bool FPerforceUpdateStatusWorker::UpdateStates() const
 	auto UpdateDataStorage = [](const FPerforceSourceControlState& State)
 	{
 		using namespace TypedElementQueryBuilder;
-		using namespace TypedElementDataStorage;
-		using DSI = ITypedElementDataStorageInterface;
 		using namespace UE::Editor;
 
 		UTypedElementRegistry* Registry = UTypedElementRegistry::GetInstance();
@@ -2031,7 +2029,7 @@ bool FPerforceUpdateStatusWorker::UpdateStates() const
 			return;
 		}
 
-		DSI* DataStorage = Registry->GetMutableDataStorage();
+		ITypedElementDataStorageInterface* DataStorage = Registry->GetMutableDataStorage();
 		if (!DataStorage)
 		{
 			return;
diff --git a/Engine/Plugins/Experimental/EditorDataStorage/Source/TedsCore/Private/Commands/EditorDataStorageCompatibilityCommands.cpp b/Engine/Plugins/Experimental/EditorDataStorage/Source/TedsCore/Private/Commands/EditorDataStorageCompatibilityCommands.cpp
--- a/Engine/Plugins/Experimental/EditorDataStorage/Source/TedsCore/Private/Commands/EditorDataStorageCompatibilityCommands.cpp
+++ b/Engine/Plugins/Experimental/EditorDataStorage/Source/TedsCore/Private/Commands/EditorDataStorageCompatibilityCommands.cpp
@@ -156,8 +156,6 @@ namespace UE::Editor::DataStorage
 	void FPatchData::RunPatch(CompatibilityCommandBuffer::FCollection& Commands, UEditorDataStorageCompatibility& StorageCompat,
 		FScratchBuffer& ScratchBuffer)
 	{
-		using namespace TypedElementDataStorage;
-
 		uint32 ReinstanceCount = Commands.GetCommandCount<FTypeInfoReinstanced>();
 		FTypeInfoReinstanced* ReinstanceArray = ScratchBuffer.EmplaceArray<FTypeInfoReinstanced>(ReinstanceCount);
 		TArrayView<FTypeInfoReinstanced> ReinstanceView(ReinstanceArray, ReinstanceCount);
@@ -236,7 +234,6 @@ namespace UE::Editor::DataStorage
 
 	void FPrepareCommands::operator()(FAddCompatibleUObject& Object)
 	{
-		using namespace TypedElementDataStorage;
 
 		UObject* ObjectPtr = Object.Object.Get();
 		if (ObjectPtr && Storage.IsRowAvailable(Object.Row))
@@ -255,8 +252,6 @@ namespace UE::Editor::DataStorage
 
 	void FPrepareCommands::operator()(FAddCompatibleExternalObject& Object)
 	{
-		using namespace TypedElementDataStorage;
-
 		if (Storage.IsRowAvailable(Object.Row) && Object.Object != nullptr)
 		{
 			Object.Table = StorageCompat.FindBestMatchingTable(Object.TypeInfo.Get());
@@ -270,8 +265,6 @@ namespace UE::Editor::DataStorage
 
 	void FPrepareCommands::operator()(FRemoveCompatibleUObject& Command)
 	{
-		using namespace TypedElementDataStorage;
-
 		if (Command.ObjectRow == InvalidRowHandle)
 		{
 			IndexHash Hash = GenerateIndexHash(Command.Object);
@@ -285,8 +278,6 @@ namespace UE::Editor::DataStorage
 
 	void FPrepareCommands::operator()(FRemoveCompatibleExternalObject& Command)
 	{
-		using namespace TypedElementDataStorage;
-
 		IndexHash Hash = GenerateIndexHash(Command.Object);
 		Command.ObjectRow = Storage.FindIndexedRow(Hash);
 		if (!Storage.IsRowAvailable(Command.ObjectRow))
@@ -297,8 +288,6 @@ namespace UE::Editor::DataStorage
 
 	void FPrepareCommands::operator()(FAddSyncFromWorldTag& Command)
 	{
-		using namespace TypedElementDataStorage;
-
 		Command.Row = StorageCompat.FindRowWithCompatibleObject(Command.Target);
 		if (!Storage.IsRowAvailable(Command.Row))
 		{
@@ -308,8 +297,6 @@ namespace UE::Editor::DataStorage
 
 	void FPrepareCommands::operator()(FRemoveInteractiveSyncFromWorldTag& Command)
 	{
-		using namespace TypedElementDataStorage;
-
 		Command.Row = StorageCompat.FindRowWithCompatibleObject(Command.Target);
 		if (!Storage.IsRowAvailable(Command.Row))
 		{
@@ -331,8 +318,6 @@ namespace UE::Editor::DataStorage
 
 	void FPrepareCommands::operator()(FAddInteractiveSyncFromWorldTag& Command)
 	{
-		using namespace TypedElementDataStorage;
-
 		Command.Row = StorageCompat.FindRowWithCompatibleObject(Command.Target);
 		if (!Storage.IsRowAvailable(Command.Row))
 		{
@@ -473,8 +458,6 @@ namespace UE::Editor::DataStorage
 
 	void FSorter::SortCommands(CompatibilityCommandBuffer::FCollection& Commands)
 	{
-		using namespace TypedElementDataStorage;
-
 		// First sort commands by the row they operate on, using stable sort so commands remain in the same range as they get processed.
 		// Within a row, commands that can't be reordered relative to each other are grouped together while the remaining commands are
 		// ordered in the order they were declared in the command buffer.
@@ -542,7 +525,6 @@ namespace UE::Editor::DataStorage
 
 	void FCommandProcessor::operator()(FTypeBatchInfoReinstanced& Command)
 	{
-		using namespace TypedElementDataStorage;
 		using namespace TypedElementQueryBuilder;
 
 		StorageCompatibility.Storage->RunQuery(StorageCompatibility.ClassTypeInfoQuery, CreateDirectQueryCallbackBinding(
@@ -595,8 +577,6 @@ namespace UE::Editor::DataStorage
 
 	void FCommandProcessor::operator()(FBatchAddCompatibleUObject& Batch)
 	{
-		using namespace TypedElementDataStorage;
-
 		TConstArrayView<RowHandle> Rows(Batch.RowArray, Batch.Count);
 		RowHandle* RowIterator = Batch.RowArray;
 		TWeakObjectPtr<UObject>* ObjectIterator = Batch.ObjectArray;
@@ -622,8 +602,6 @@ namespace UE::Editor::DataStorage
 
 	void FCommandProcessor::operator()(FBatchAddCompatibleExternalObject& Batch)
 	{
-		using namespace TypedElementDataStorage;
-
 		TConstArrayView<RowHandle> Rows(Batch.RowArray, Batch.Count);
 		RowHandle* RowIterator = Batch.RowArray;
 		void** ObjectIterator = Batch.ObjectArray;
@@ -844,8 +822,6 @@ namespace UE::Editor::DataStorage
 
 	void FCommandOptimizer::operator()(const FAddCompatibleUObject& Command)
 	{
-		using namespace TypedElementDataStorage;
-
 		int32 Count = FoldCommandsForAdd(Command);
 
 		// If there are more than 1 adds, batch them together into a batch call and nop the additional adds out.
@@ -883,8 +859,6 @@ namespace UE::Editor::DataStorage
 	
 	void FCommandOptimizer::operator()(const FAddCompatibleExternalObject& Command)
 	{
-		using namespace TypedElementDataStorage;
-
 		int32 Count = FoldCommandsForAdd(Command);
 
 		// If there are more than 1 adds, batch them together into a batch call and nop them out.
@@ -923,8 +897,6 @@ namespace UE::Editor::DataStorage
 
 	void FCommandOptimizer::operator()(const FRemoveCompatibleUObject& Command)
 	{
-		using namespace TypedElementDataStorage;
-
 		// Remove everything after the delete because no operation on a row is going to succeed after
 		// the row has been deleted.
 		while(Optimizer.IsValid())
@@ -944,9 +916,7 @@ namespace UE::Editor::DataStorage
 	}
 
 	void FCommandOptimizer::operator()(const FRemoveCompatibleExternalObject& Command)
-	{
-		using namespace TypedElementDataStorage;
-		
+	{		
 		// Remove everything after the delete because no operation on a row is going to succeed after
 		// the row has been deleted.
 		while (Optimizer.IsValid())
@@ -968,8 +938,6 @@ namespace UE::Editor::DataStorage
 	template<typename AddCommandType>
 	int32 FCommandOptimizer::FoldCommandsForAdd(const AddCommandType& Command)
 	{
-		using namespace TypedElementDataStorage;
-
 		// Discover the longest chain of adds for the same table and remove any async tag additions as it'll be included
 		// in the table for UObjects.
 		int32 Count = 1;
diff --git a/Engine/Plugins/Experimental/EditorDataStorage/Source/TedsCore/Private/Compatibility/TedsTypedElementActorHandleFactory.cpp b/Engine/Plugins/Experimental/EditorDataStorage/Source/TedsCore/Private/Compatibility/TedsTypedElementActorHandleFactory.cpp
--- a/Engine/Plugins/Experimental/EditorDataStorage/Source/TedsCore/Private/Compatibility/TedsTypedElementActorHandleFactory.cpp
+++ b/Engine/Plugins/Experimental/EditorDataStorage/Source/TedsCore/Private/Compatibility/TedsTypedElementActorHandleFactory.cpp
@@ -34,7 +34,6 @@ void UTypedElementActorHandleDataStorageFactory::RegisterQueries(ITypedElementDa
 	}
 
 	using namespace TypedElementQueryBuilder;
-	using namespace TypedElementDataStorage;
 	GetAllActorsQuery = DataStorage.RegisterQuery(
 	Select()
 		.ReadOnly<FTypedElementUObjectColumn>()
@@ -46,9 +45,9 @@ void UTypedElementActorHandleDataStorageFactory::RegisterQueries(ITypedElementDa
 void UTypedElementActorHandleDataStorageFactory::RegisterQuery_ActorHandlePopulate(ITypedElementDataStorageInterface& DataStorage)
 {
 	using namespace TypedElementQueryBuilder;
-	using namespace TypedElementDataStorage;
-	using DSI = ITypedElementDataStorageInterface;
-	if (!ensureMsgf(ActorHandlePopulateQuery == UE::Editor::DataStorage::InvalidQueryHandle, TEXT("Already registered query")))
+	using namespace UE::Editor::DataStorage;
+
+	if (!ensureMsgf(ActorHandlePopulateQuery == InvalidQueryHandle, TEXT("Already registered query")))
 	{
 		return;
 	}
@@ -56,13 +55,13 @@ void UTypedElementActorHandleDataStorageFactory::RegisterQuery_ActorHandlePopula
 	ActorHandlePopulateQuery = DataStorage.RegisterQuery(
 	Select(TEXT("Populate actor typed element handles"),
 		FObserver::OnAdd<FTypedElementUObjectColumn>(),
-		[](DSI::IQueryContext& Context, UE::Editor::DataStorage::RowHandle Row, const FTypedElementUObjectColumn& ObjectColumn)
+		[](IQueryContext& Context, RowHandle Row, const FTypedElementUObjectColumn& ObjectColumn)
 		{
 			if (UObject* Object = ObjectColumn.Object.Get())
 			{
 				checkSlow(Cast<AActor>(Object));
 				FTypedElementHandle Handle = UEngineElementsLibrary::AcquireEditorActorElementHandle(static_cast<AActor*>(Object));
-				Context.AddColumn(Row, UE::Editor::DataStorage::Compatibility::FTypedElementColumn
+				Context.AddColumn(Row, Compatibility::FTypedElementColumn
 				{
 					.Handle = Handle
 				});
@@ -83,8 +82,6 @@ void UTypedElementActorHandleDataStorageFactory::HandleBridgeEnabled(bool bEnabl
 	if (bEnabled)
 	{
 		using namespace TypedElementQueryBuilder;
-		using namespace TypedElementDataStorage;
-		using DSI = ITypedElementDataStorageInterface;
 		
 		// Populate all the rows
 		TArray<RowHandle> CollatedRowHandles;
diff --git a/Engine/Plugins/Experimental/EditorDataStorage/Source/TedsCore/Private/Compatibility/TedsTypedElementBridgeQueries.cpp b/Engine/Plugins/Experimental/EditorDataStorage/Source/TedsCore/Private/Compatibility/TedsTypedElementBridgeQueries.cpp
--- a/Engine/Plugins/Experimental/EditorDataStorage/Source/TedsCore/Private/Compatibility/TedsTypedElementBridgeQueries.cpp
+++ b/Engine/Plugins/Experimental/EditorDataStorage/Source/TedsCore/Private/Compatibility/TedsTypedElementBridgeQueries.cpp
@@ -58,7 +58,6 @@ bool UTypedElementBridgeDataStorageFactory::IsEnabled()
 void UTypedElementBridgeDataStorageFactory::RegisterQuery_NewUObject(ITypedElementDataStorageInterface& DataStorage)
 {
 	using namespace TypedElementQueryBuilder;
-	using DSI = ITypedElementDataStorageInterface;
 	
 	RemoveTypedElementRowHandleQuery = DataStorage.RegisterQuery(
 		Select()
@@ -73,21 +72,21 @@ void UTypedElementBridgeDataStorageFactory::UnregisterQuery_NewUObject(ITypedEle
 void UTypedElementBridgeDataStorageFactory::CleanupTypedElementColumns(ITypedElementDataStorageInterface& DataStorage)
 {
 	using namespace TypedElementQueryBuilder;
-	using namespace TypedElementDataStorage;
+	using namespace UE::Editor::DataStorage;
 
 	// Remove any TEv1 handles
 	{
-		TArray<UE::Editor::DataStorage::RowHandle> Handles;
+		TArray<RowHandle> Handles;
 		using namespace TypedElementQueryBuilder;
 		DataStorage.RunQuery(
 			RemoveTypedElementRowHandleQuery,
 			CreateDirectQueryCallbackBinding(
-				[&Handles](ITypedElementDataStorageInterface::IDirectQueryContext& Context)
+				[&Handles](IDirectQueryContext& Context)
 				{
 					Handles.Append(Context.GetRowHandles());
 				}));
 		
-		DataStorage.BatchAddRemoveColumns(TConstArrayView<UE::Editor::DataStorage::RowHandle>(Handles), {}, {UE::Editor::DataStorage::Compatibility::FTypedElementColumn::StaticStruct()});
+		DataStorage.BatchAddRemoveColumns(TConstArrayView<RowHandle>(Handles), {}, {Compatibility::FTypedElementColumn::StaticStruct()});
 	}
 }
 
diff --git a/Engine/Plugins/Experimental/EditorDataStorage/Source/TedsCore/Private/Compatibility/TypedElementActorIconOverrideQueries.cpp b/Engine/Plugins/Experimental/EditorDataStorage/Source/TedsCore/Private/Compatibility/TypedElementActorIconOverrideQueries.cpp
--- a/Engine/Plugins/Experimental/EditorDataStorage/Source/TedsCore/Private/Compatibility/TypedElementActorIconOverrideQueries.cpp
+++ b/Engine/Plugins/Experimental/EditorDataStorage/Source/TedsCore/Private/Compatibility/TypedElementActorIconOverrideQueries.cpp
@@ -10,7 +10,7 @@
 
 void UActorIconOverrideDataStorageFactory::RegisterQueries(ITypedElementDataStorageInterface& DataStorage)
 {
-	using namespace TypedElementDataStorage;
+	using namespace UE::Editor::DataStorage;
 	using namespace TypedElementQueryBuilder;
 
 	DataStorage.RegisterQuery(
diff --git a/Engine/Plugins/Experimental/EditorDataStorage/Source/TedsCore/Private/Compatibility/TypedElementActorLabelQueries.cpp b/Engine/Plugins/Experimental/EditorDataStorage/Source/TedsCore/Private/Compatibility/TypedElementActorLabelQueries.cpp
--- a/Engine/Plugins/Experimental/EditorDataStorage/Source/TedsCore/Private/Compatibility/TypedElementActorLabelQueries.cpp
+++ b/Engine/Plugins/Experimental/EditorDataStorage/Source/TedsCore/Private/Compatibility/TypedElementActorLabelQueries.cpp
@@ -22,7 +22,7 @@ void UActorLabelDataStorageFactory::RegisterQueries(ITypedElementDataStorageInte
 void UActorLabelDataStorageFactory::RegisterActorLabelToColumnQuery(ITypedElementDataStorageInterface& DataStorage) const
 {
 	using namespace TypedElementQueryBuilder;
-	using namespace TypedElementDataStorage;
+	using namespace UE::Editor::DataStorage;
 	
 	DataStorage.RegisterQuery(
 		Select(
@@ -52,7 +52,7 @@ void UActorLabelDataStorageFactory::RegisterActorLabelToColumnQuery(ITypedElemen
 void UActorLabelDataStorageFactory::RegisterLabelColumnToActorQuery(ITypedElementDataStorageInterface& DataStorage) const
 {
 	using namespace TypedElementQueryBuilder;
-	using namespace TypedElementDataStorage;
+	using namespace UE::Editor::DataStorage;
 	
 	DataStorage.RegisterQuery(
 		Select(
diff --git a/Engine/Plugins/Experimental/EditorDataStorage/Source/TedsCore/Private/Compatibility/TypedElementActorParentQueries.cpp b/Engine/Plugins/Experimental/EditorDataStorage/Source/TedsCore/Private/Compatibility/TypedElementActorParentQueries.cpp
--- a/Engine/Plugins/Experimental/EditorDataStorage/Source/TedsCore/Private/Compatibility/TypedElementActorParentQueries.cpp
+++ b/Engine/Plugins/Experimental/EditorDataStorage/Source/TedsCore/Private/Compatibility/TypedElementActorParentQueries.cpp
@@ -31,7 +31,6 @@ void UActorParentDataStorageFactory::RegisterQueries(ITypedElementDataStorageInt
 
 void UActorParentDataStorageFactory::RegisterAddParentColumn(ITypedElementDataStorageInterface& DataStorage) const
 {
-	using namespace TypedElementDataStorage;
 	using namespace TypedElementQueryBuilder;
 	using namespace UE::Editor::DataStorage;
 	
@@ -69,7 +68,6 @@ void UActorParentDataStorageFactory::RegisterAddParentColumn(ITypedElementDataSt
 
 void UActorParentDataStorageFactory::RegisterUpdateOrRemoveParentColumn(ITypedElementDataStorageInterface& DataStorage) const
 {
-	using namespace TypedElementDataStorage;
 	using namespace TypedElementQueryBuilder;
 	using namespace UE::Editor::DataStorage;
 	
diff --git a/Engine/Plugins/Experimental/EditorDataStorage/Source/TedsCore/Private/Compatibility/TypedElementActorTransformQueries.cpp b/Engine/Plugins/Experimental/EditorDataStorage/Source/TedsCore/Private/Compatibility/TypedElementActorTransformQueries.cpp
--- a/Engine/Plugins/Experimental/EditorDataStorage/Source/TedsCore/Private/Compatibility/TypedElementActorTransformQueries.cpp
+++ b/Engine/Plugins/Experimental/EditorDataStorage/Source/TedsCore/Private/Compatibility/TypedElementActorTransformQueries.cpp
@@ -18,7 +18,6 @@ void UActorTransformDataStorageFactory::RegisterQueries(ITypedElementDataStorage
 void UActorTransformDataStorageFactory::RegisterActorAddTransformColumn(ITypedElementDataStorageInterface& DataStorage) const
 {
 	using namespace TypedElementQueryBuilder;
-	using namespace TypedElementDataStorage;
 	using namespace UE::Editor::DataStorage;
 	
 	DataStorage.RegisterQuery(
@@ -45,7 +44,6 @@ void UActorTransformDataStorageFactory::RegisterActorAddTransformColumn(ITypedEl
 void UActorTransformDataStorageFactory::RegisterActorLocalTransformToColumn(ITypedElementDataStorageInterface& DataStorage) const
 {
 	using namespace TypedElementQueryBuilder;
-	using namespace TypedElementDataStorage;
 	using namespace UE::Editor::DataStorage;
 	
 	DataStorage.RegisterQuery(
@@ -75,7 +73,7 @@ void UActorTransformDataStorageFactory::RegisterActorLocalTransformToColumn(ITyp
 void UActorTransformDataStorageFactory::RegisterLocalTransformColumnToActor(ITypedElementDataStorageInterface& DataStorage) const
 {
 	using namespace TypedElementQueryBuilder;
-	using namespace TypedElementDataStorage;
+	using namespace UE::Editor::DataStorage;
 	
 	DataStorage.RegisterQuery(
 		Select(
diff --git a/Engine/Plugins/Experimental/EditorDataStorage/Source/TedsCore/Private/Compatibility/TypedElementActorViewportProcessors.cpp b/Engine/Plugins/Experimental/EditorDataStorage/Source/TedsCore/Private/Compatibility/TypedElementActorViewportProcessors.cpp
--- a/Engine/Plugins/Experimental/EditorDataStorage/Source/TedsCore/Private/Compatibility/TypedElementActorViewportProcessors.cpp
+++ b/Engine/Plugins/Experimental/EditorDataStorage/Source/TedsCore/Private/Compatibility/TypedElementActorViewportProcessors.cpp
@@ -20,7 +20,6 @@ namespace UE::Editor::DataStorage::Private
 		FConsoleCommandWithArgsAndOutputDeviceDelegate::CreateLambda([](const TArray<FString>& Args, FOutputDevice& Output)
 			{
 				using namespace TypedElementQueryBuilder;
-				using DSI = ITypedElementDataStorageInterface;
 				using namespace UE::Editor::DataStorage;
 
 				TRACE_CPUPROFILER_EVENT_SCOPE(TEDS.Debug.AddOverlayColorToSelectionCommand);
@@ -60,7 +59,7 @@ namespace UE::Editor::DataStorage::Private
 
 					TArray<RowHandle> RowHandles;
 				
-					DataStorage->RunQuery(OverlayQuery, [ColorIndex, &RowHandles](const DSI::FQueryDescription&, DSI::IDirectQueryContext& Context)
+					DataStorage->RunQuery(OverlayQuery, [ColorIndex, &RowHandles](const FQueryDescription&, IDirectQueryContext& Context)
 					{
 						RowHandles = Context.GetRowHandles();
 					});
@@ -80,7 +79,6 @@ namespace UE::Editor::DataStorage::Private
 		FConsoleCommandWithArgsAndOutputDeviceDelegate::CreateLambda([](const TArray<FString>& Args, FOutputDevice& Output)
 			{
 				using namespace TypedElementQueryBuilder;
-				using DSI = ITypedElementDataStorageInterface;
 				using namespace UE::Editor::DataStorage;
 
 				TRACE_CPUPROFILER_EVENT_SCOPE(TEDS.Debug.AddOverlayColorToSelectionCommand);
@@ -114,7 +112,7 @@ namespace UE::Editor::DataStorage::Private
 
 					TArray<RowHandle> RowHandles;
 				
-					DataStorage->RunQuery(OverlayQuery, [Color, &RowHandles](const DSI::FQueryDescription&, DSI::IDirectQueryContext& Context)
+					DataStorage->RunQuery(OverlayQuery, [Color, &RowHandles](const FQueryDescription&, IDirectQueryContext& Context)
 					{
 						RowHandles = Context.GetRowHandles();
 					});
@@ -133,7 +131,6 @@ namespace UE::Editor::DataStorage::Private
 		FConsoleCommandWithArgsAndOutputDeviceDelegate::CreateLambda([](const TArray<FString>& Args, FOutputDevice& Output)
 			{
 				using namespace TypedElementQueryBuilder;
-				using DSI = ITypedElementDataStorageInterface;
 				using namespace UE::Editor::DataStorage;
 
 				TRACE_CPUPROFILER_EVENT_SCOPE(TEDS.Debug.AddOverlayColorToSelectionCommand);
@@ -157,7 +154,7 @@ namespace UE::Editor::DataStorage::Private
 
 					TArray<RowHandle> RowHandles;
 				
-					DataStorage->RunQuery(OverlayQuery, [&RowHandles](const DSI::FQueryDescription&, DSI::IDirectQueryContext& Context)
+					DataStorage->RunQuery(OverlayQuery, [&RowHandles](const FQueryDescription&, IDirectQueryContext& Context)
 					{
 						RowHandles = Context.GetRowHandles();
 					});
@@ -179,7 +176,7 @@ void UActorViewportDataStorageFactory::RegisterQueries(ITypedElementDataStorageI
 void UActorViewportDataStorageFactory::RegisterOutlineColorColumnToActor(ITypedElementDataStorageInterface& DataStorage)
 {
 	using namespace TypedElementQueryBuilder;
-	using namespace TypedElementDataStorage;
+	using namespace UE::Editor::DataStorage;
 
 	DataStorage.RegisterQuery(
 		Select(
@@ -207,7 +204,7 @@ void UActorViewportDataStorageFactory::RegisterOutlineColorColumnToActor(ITypedE
 void UActorViewportDataStorageFactory::RegisterOverlayColorColumnToActor(ITypedElementDataStorageInterface& DataStorage)
 {
 	using namespace TypedElementQueryBuilder;
-	using namespace TypedElementDataStorage;
+	using namespace UE::Editor::DataStorage;
 	
 	DataStorage.RegisterQuery(
 		Select(
diff --git a/Engine/Plugins/Experimental/EditorDataStorage/Source/TedsCore/Private/Compatibility/TypedElementObjectReinstancingManager.cpp b/Engine/Plugins/Experimental/EditorDataStorage/Source/TedsCore/Private/Compatibility/TypedElementObjectReinstancingManager.cpp
--- a/Engine/Plugins/Experimental/EditorDataStorage/Source/TedsCore/Private/Compatibility/TypedElementObjectReinstancingManager.cpp
+++ b/Engine/Plugins/Experimental/EditorDataStorage/Source/TedsCore/Private/Compatibility/TypedElementObjectReinstancingManager.cpp
@@ -18,7 +18,6 @@ UTedsObjectReinstancingManager::UTedsObjectReinstancingManager()
 
 void UTedsObjectReinstancingManager::Initialize(UEditorDataStorage& InDataStorage, UEditorDataStorageCompatibility& InDataStorageCompatibility)
 {
-	using namespace TypedElementDataStorage;
 	using namespace UE::Editor::DataStorage;
 
 	DataStorage = &InDataStorage;
diff --git a/Engine/Plugins/Experimental/EditorDataStorage/Source/TedsCore/Private/Compatibility/TypedElementUObjectWorldQueries.cpp b/Engine/Plugins/Experimental/EditorDataStorage/Source/TedsCore/Private/Compatibility/TypedElementUObjectWorldQueries.cpp
--- a/Engine/Plugins/Experimental/EditorDataStorage/Source/TedsCore/Private/Compatibility/TypedElementUObjectWorldQueries.cpp
+++ b/Engine/Plugins/Experimental/EditorDataStorage/Source/TedsCore/Private/Compatibility/TypedElementUObjectWorldQueries.cpp
@@ -15,7 +15,6 @@ void UObjectWorldDataStorageFactory::RegisterQueries(ITypedElementDataStorageInt
 
 void UObjectWorldDataStorageFactory::RegisterAddWorldColumn(ITypedElementDataStorageInterface& DataStorage) const
 {
-	using namespace TypedElementDataStorage;
 	using namespace TypedElementQueryBuilder;
 	using namespace UE::Editor::DataStorage;
 	
@@ -46,7 +45,6 @@ void UObjectWorldDataStorageFactory::RegisterAddWorldColumn(ITypedElementDataSto
 
 void UObjectWorldDataStorageFactory::RegisterUpdateOrRemoveWorldColumn(ITypedElementDataStorageInterface& DataStorage) const
 {
-	using namespace TypedElementDataStorage;
 	using namespace TypedElementQueryBuilder;
 	using namespace UE::Editor::DataStorage;
 	
diff --git a/Engine/Plugins/Experimental/EditorDataStorage/Source/TedsCore/Private/Debug/TypedElementDatabaseDebugCommands.cpp b/Engine/Plugins/Experimental/EditorDataStorage/Source/TedsCore/Private/Debug/TypedElementDatabaseDebugCommands.cpp
--- a/Engine/Plugins/Experimental/EditorDataStorage/Source/TedsCore/Private/Debug/TypedElementDatabaseDebugCommands.cpp
+++ b/Engine/Plugins/Experimental/EditorDataStorage/Source/TedsCore/Private/Debug/TypedElementDatabaseDebugCommands.cpp
@@ -25,7 +25,6 @@ namespace UE::Editor::DataStorage::Private
 	template<typename TypeInfoType>
 	void PrintObjectTypeInformation(ITypedElementDataStorageInterface* DataStorage, FString Message, FOutputDevice& Output)
 	{
-		using namespace TypedElementDataStorage;
 		using namespace TypedElementQueryBuilder;
 
 		static QueryHandle Query = [DataStorage]
@@ -68,8 +67,6 @@ namespace UE::Editor::DataStorage::Private
 	void PrintObjectLabels(FOutputDevice& Output)
 	{
 		using namespace TypedElementQueryBuilder;
-		using namespace UE::Editor::DataStorage;
-		using DSI = ITypedElementDataStorageInterface;
 
 		if (ITypedElementDataStorageInterface* DataStorage = UTypedElementRegistry::GetInstance()->GetMutableDataStorage())
 		{
@@ -97,7 +94,7 @@ namespace UE::Editor::DataStorage::Private
 			{
 				FString Message;
 				DataStorage->RunQuery(LabelQuery, CreateDirectQueryCallbackBinding(
-					[&Output, &Message](DSI::IDirectQueryContext& Context, const FTypedElementUObjectColumn* Objects, const FTypedElementLabelColumn* Labels)
+					[&Output, &Message](IDirectQueryContext& Context, const FTypedElementUObjectColumn* Objects, const FTypedElementLabelColumn* Labels)
 					{
 						const uint32 Count = Context.GetRowCount();
 
@@ -131,7 +128,7 @@ namespace UE::Editor::DataStorage::Private
 			}
 		}
 	}
-}
+} // namespace UE::Editor::DataStorage::Private
 
 FAutoConsoleCommandWithOutputDevice PrintObjectTypeInformationConsoleCommand(
 	TEXT("TEDS.Debug.PrintObjectTypeInfo"),
@@ -411,7 +408,6 @@ static FAutoConsoleCommand CVarPrintDynamicColumnWithQuery(
 	FConsoleCommandWithArgsDelegate::CreateLambda([](const TArray<FString>& Args)
 	{
 		using namespace TypedElementQueryBuilder;
-		using DSI = ITypedElementDataStorageInterface;
 		using namespace UE::Editor::DataStorage;
 		
 		// Print column using query
@@ -528,8 +524,6 @@ static FAutoConsoleCommand CVarRegisterListDynamicColumnQuery(
 				
 		const FName Identifier(*Args[0]);
 		const FName ActivationGroup(*Args[1]);
-		
-		using namespace TypedElementDataStorage;
 
 		// Lists the rows processed that have 
 		DataStorage->RegisterQuery(
@@ -820,7 +814,7 @@ static FAutoConsoleCommand CVarMatchValueTagFromEnum(
 		
 		uint64 Count = 0;
 		
-		const TypedElementDataStorage::FQueryResult Result = DataStorage->RunQuery(Query, CreateDirectQueryCallbackBinding(
+		const FQueryResult Result = DataStorage->RunQuery(Query, CreateDirectQueryCallbackBinding(
 			[&Count](const IDirectQueryContext& Context, const RowHandle*)
 			{
 				Count += Context.GetRowCount();
diff --git a/Engine/Plugins/Experimental/EditorDataStorage/Source/TedsCore/Private/Processors/TypedElementProcessorAdaptors.cpp b/Engine/Plugins/Experimental/EditorDataStorage/Source/TedsCore/Private/Processors/TypedElementProcessorAdaptors.cpp
--- a/Engine/Plugins/Experimental/EditorDataStorage/Source/TedsCore/Private/Processors/TypedElementProcessorAdaptors.cpp
+++ b/Engine/Plugins/Experimental/EditorDataStorage/Source/TedsCore/Private/Processors/TypedElementProcessorAdaptors.cpp
@@ -51,7 +51,7 @@ namespace UE::Editor::DataStorage
 
 			void GetColumns(TArrayView<char*> RetrievedAddresses,
 				TConstArrayView<TWeakObjectPtr<const UScriptStruct>> ColumnTypes,
-				TConstArrayView<TypedElementDataStorage::EQueryAccessType> AccessTypes)
+				TConstArrayView<EQueryAccessType> AccessTypes)
 			{
 				checkf(RetrievedAddresses.Num() == ColumnTypes.Num(), TEXT("Unable to retrieve a batch of columns as the number of addresses "
 					"doesn't match the number of requested column."));
@@ -65,7 +65,7 @@ namespace UE::Editor::DataStorage
 				int32 TypeCount,
 				char** RetrievedAddresses,
 				const TWeakObjectPtr<const UScriptStruct>* ColumnTypes,
-				const TypedElementDataStorage::EQueryAccessType* AccessTypes)
+				const EQueryAccessType* AccessTypes)
 			{
 				for (int32 Index = 0; Index < TypeCount; ++Index)
 				{
@@ -624,8 +624,8 @@ namespace UE::Editor::DataStorage
 			TConstArrayView<RowHandle> GetRowHandles() const override { return Implementation.GetRowHandles(); }
 			const void* GetColumn(const UScriptStruct* ColumnType) const override { return Implementation.GetColumn(ColumnType); }
 			void* GetMutableColumn(const UScriptStruct* ColumnType) override { return Implementation.GetMutableColumn(ColumnType); }
-			void GetColumns(TArrayView<char*> RetrievedAddresses, TConstArrayView<TWeakObjectPtr<const UScriptStruct>> ColumnTypes, TConstArrayView<TypedElementDataStorage::EQueryAccessType> AccessTypes) override { return Implementation.GetColumns(RetrievedAddresses, ColumnTypes, AccessTypes); }
-			void GetColumnsUnguarded(int32 TypeCount, char** RetrievedAddresses, const TWeakObjectPtr<const UScriptStruct>* ColumnTypes, const TypedElementDataStorage::EQueryAccessType* AccessTypes) override { return Implementation.GetColumnsUnguarded(TypeCount, RetrievedAddresses, ColumnTypes, AccessTypes); }
+			void GetColumns(TArrayView<char*> RetrievedAddresses, TConstArrayView<TWeakObjectPtr<const UScriptStruct>> ColumnTypes, TConstArrayView<EQueryAccessType> AccessTypes) override { return Implementation.GetColumns(RetrievedAddresses, ColumnTypes, AccessTypes); }
+			void GetColumnsUnguarded(int32 TypeCount, char** RetrievedAddresses, const TWeakObjectPtr<const UScriptStruct>* ColumnTypes, const EQueryAccessType* AccessTypes) override { return Implementation.GetColumnsUnguarded(TypeCount, RetrievedAddresses, ColumnTypes, AccessTypes); }
 			bool HasColumn(const UScriptStruct* ColumnType) const override { return Implementation.HasColumn(ColumnType); }
 			bool HasColumn(RowHandle Row, const UScriptStruct* ColumnType) const override { return Implementation.HasColumn(Row, ColumnType); }
 			virtual const UScriptStruct* FindDynamicColumnType(const UE::Editor::DataStorage::FDynamicColumnDescription& Description) const override { return Implementation.FindDynamicColumnType(Description); }
@@ -644,8 +644,8 @@ namespace UE::Editor::DataStorage
 			TConstArrayView<RowHandle> GetRowHandles() const override { return Implementation.GetRowHandles(); }
 			const void* GetColumn(const UScriptStruct* ColumnType) const override { return Implementation.GetColumn(ColumnType); }
 			void* GetMutableColumn(const UScriptStruct* ColumnType) override { return Implementation.GetMutableColumn(ColumnType); }
-			void GetColumns(TArrayView<char*> RetrievedAddresses, TConstArrayView<TWeakObjectPtr<const UScriptStruct>> ColumnTypes, TConstArrayView<TypedElementDataStorage::EQueryAccessType> AccessTypes) override { return Implementation.GetColumns(RetrievedAddresses, ColumnTypes, AccessTypes); }
-			void GetColumnsUnguarded(int32 TypeCount, char** RetrievedAddresses, const TWeakObjectPtr<const UScriptStruct>* ColumnTypes, const TypedElementDataStorage::EQueryAccessType* AccessTypes) override { return Implementation.GetColumnsUnguarded(TypeCount, RetrievedAddresses, ColumnTypes, AccessTypes); }
+			void GetColumns(TArrayView<char*> RetrievedAddresses, TConstArrayView<TWeakObjectPtr<const UScriptStruct>> ColumnTypes, TConstArrayView<EQueryAccessType> AccessTypes) override { return Implementation.GetColumns(RetrievedAddresses, ColumnTypes, AccessTypes); }
+			void GetColumnsUnguarded(int32 TypeCount, char** RetrievedAddresses, const TWeakObjectPtr<const UScriptStruct>* ColumnTypes, const EQueryAccessType* AccessTypes) override { return Implementation.GetColumnsUnguarded(TypeCount, RetrievedAddresses, ColumnTypes, AccessTypes); }
 			bool HasColumn(const UScriptStruct* ColumnType) const override { return Implementation.HasColumn(ColumnType); }
 			bool HasColumn(RowHandle Row, const UScriptStruct* ColumnType) const override { return Implementation.HasColumn(Row, ColumnType); }
 			uint64 GetUpdateCycleId() const override { return Implementation.GetUpdateCycleId(); }
@@ -721,7 +721,7 @@ namespace UE::Editor::DataStorage
 				}
 			}
 	
-			RowHandle FindIndexedRow(TypedElementDataStorage::IndexHash Index) const
+			RowHandle FindIndexedRow(IndexHash Index) const
 			{
 				EGlobalLockScope Scope = FGlobalLock::GetLockStatus(EGlobalLockScope::Internal) == EGlobalLockStatus::Unlocked
 					? EGlobalLockScope::Public // There's no internal lock so use a public lock instead.
@@ -790,8 +790,8 @@ namespace UE::Editor::DataStorage
 			TConstArrayView<RowHandle> GetRowHandles() const override { return Implementation.GetRowHandles(); }
 			const void* GetColumn(const UScriptStruct* ColumnType) const override { return Implementation.GetColumn(ColumnType); }
 			void* GetMutableColumn(const UScriptStruct* ColumnType) override { return Implementation.GetMutableColumn(ColumnType); }
-			void GetColumns(TArrayView<char*> RetrievedAddresses, TConstArrayView<TWeakObjectPtr<const UScriptStruct>> ColumnTypes, TConstArrayView<TypedElementDataStorage::EQueryAccessType> AccessTypes) override { return Implementation.GetColumns(RetrievedAddresses, ColumnTypes, AccessTypes); }
-			void GetColumnsUnguarded(int32 TypeCount, char** RetrievedAddresses, const TWeakObjectPtr<const UScriptStruct>* ColumnTypes, const TypedElementDataStorage::EQueryAccessType* AccessTypes) override { return Implementation.GetColumnsUnguarded(TypeCount, RetrievedAddresses, ColumnTypes, AccessTypes); }
+			void GetColumns(TArrayView<char*> RetrievedAddresses, TConstArrayView<TWeakObjectPtr<const UScriptStruct>> ColumnTypes, TConstArrayView<EQueryAccessType> AccessTypes) override { return Implementation.GetColumns(RetrievedAddresses, ColumnTypes, AccessTypes); }
+			void GetColumnsUnguarded(int32 TypeCount, char** RetrievedAddresses, const TWeakObjectPtr<const UScriptStruct>* ColumnTypes, const EQueryAccessType* AccessTypes) override { return Implementation.GetColumnsUnguarded(TypeCount, RetrievedAddresses, ColumnTypes, AccessTypes); }
 			bool HasColumn(const UScriptStruct* ColumnType) const override { return Implementation.HasColumn(ColumnType); }
 			bool HasColumn(RowHandle Row, const UScriptStruct* ColumnType) const override { return Implementation.HasColumn(Row, ColumnType); }
 			uint64 GetUpdateCycleId() const override { return Implementation.GetUpdateCycleId(); }
@@ -814,12 +814,12 @@ namespace UE::Editor::DataStorage
 
 			const UObject* GetDependency(const UClass* DependencyClass) override { return Implementation.GetDependency(DependencyClass); }
 			UObject* GetMutableDependency(const UClass* DependencyClass) override { return Implementation.GetMutableDependency(DependencyClass); }
-			void GetDependencies(TArrayView<UObject*> RetrievedAddresses, TConstArrayView<TWeakObjectPtr<const UClass>> DependencyTypes, TConstArrayView<TypedElementDataStorage::EQueryAccessType> AccessTypes) override { return Implementation.GetDependencies(RetrievedAddresses, DependencyTypes, AccessTypes); }
-			RowHandle FindIndexedRow(TypedElementDataStorage::IndexHash Index) const override { return Implementation.FindIndexedRow(Index); }
-			TypedElementDataStorage::FQueryResult RunQuery(QueryHandle Query) override { return Implementation.RunQuery(Query); }
-			TypedElementDataStorage::FQueryResult RunSubquery(int32 SubqueryIndex) override { return Implementation.RunSubquery(SubqueryIndex); }
-			TypedElementDataStorage::FQueryResult RunSubquery(int32 SubqueryIndex, UE::Editor::DataStorage::SubqueryCallbackRef Callback) override { return Implementation.RunSubquery(SubqueryIndex, Callback); }
-			TypedElementDataStorage::FQueryResult RunSubquery(int32 SubqueryIndex, RowHandle Row, UE::Editor::DataStorage::SubqueryCallbackRef Callback) override { return Implementation.RunSubquery(SubqueryIndex, Row, Callback); }
+			void GetDependencies(TArrayView<UObject*> RetrievedAddresses, TConstArrayView<TWeakObjectPtr<const UClass>> DependencyTypes, TConstArrayView<EQueryAccessType> AccessTypes) override { return Implementation.GetDependencies(RetrievedAddresses, DependencyTypes, AccessTypes); }
+			RowHandle FindIndexedRow(IndexHash Index) const override { return Implementation.FindIndexedRow(Index); }
+			FQueryResult RunQuery(QueryHandle Query) override { return Implementation.RunQuery(Query); }
+			FQueryResult RunSubquery(int32 SubqueryIndex) override { return Implementation.RunSubquery(SubqueryIndex); }
+			FQueryResult RunSubquery(int32 SubqueryIndex, UE::Editor::DataStorage::SubqueryCallbackRef Callback) override { return Implementation.RunSubquery(SubqueryIndex, Callback); }
+			FQueryResult RunSubquery(int32 SubqueryIndex, RowHandle Row, UE::Editor::DataStorage::SubqueryCallbackRef Callback) override { return Implementation.RunSubquery(SubqueryIndex, Row, Callback); }
 	
 			FMassQueryContextImplementation Implementation;
 		};
@@ -950,6 +950,8 @@ FString FTypedElementQueryProcessorData::GetProcessorName() const
 void FTypedElementQueryProcessorData::DebugOutputDescription(FOutputDevice& Ar, int32 Indent) const
 {
 #if WITH_MASSENTITY_DEBUG
+	using namespace UE::Editor::DataStorage;
+
 	if (const FExtendedQuery* StoredQuery = QueryStore ? QueryStore->Get(ParentQuery) : nullptr)
 	{
 		const ITypedElementDataStorageInterface::FQueryDescription& Description = StoredQuery->Description;
@@ -990,16 +992,16 @@ void FTypedElementQueryProcessorData::DebugOutputDescription(FOutputDevice& Ar,
 
 		switch (Callback.ExecutionMode)
 		{
-		case TypedElementDataStorage::EExecutionMode::Default:
+		case EExecutionMode::Default:
 			Ar.Logf(TEXT("\n%*sExecution mode: Default"), Indent, TEXT(""));
 			break;
-		case TypedElementDataStorage::EExecutionMode::GameThread:
+		case EExecutionMode::GameThread:
 			Ar.Logf(TEXT("\n%*sExecution mode: Game Thread"), Indent, TEXT(""));
 			break;
-		case TypedElementDataStorage::EExecutionMode::Threaded:
+		case EExecutionMode::Threaded:
 			Ar.Logf(TEXT("\n%*sExecution mode: Threaded"), Indent, TEXT(""));
 			break;
-		case TypedElementDataStorage::EExecutionMode::ThreadedChunks:
+		case EExecutionMode::ThreadedChunks:
 			Ar.Logf(TEXT("\n%*sExecution mode: Threaded Chunks"), Indent, TEXT(""));
 			break;
 		default:
@@ -1234,9 +1236,11 @@ bool UTypedElementQueryProcessorCallbackAdapterProcessorBase::ConfigureQueryCall
 	FEnvironment& Environment,
 	TArrayView<FMassEntityQuery> Subqueries)
 {
+	using namespace UE::Editor::DataStorage;
+
 	bool Result = Data.CommonQueryConfiguration(*this, Query, QueryHandle, QueryStore, Environment, Subqueries);
 
-	bRequiresGameThreadExecution = Query.Description.Callback.ExecutionMode == TypedElementDataStorage::EExecutionMode::GameThread;
+	bRequiresGameThreadExecution = Query.Description.Callback.ExecutionMode == EExecutionMode::GameThread;
 	ExecutionFlags = static_cast<int32>(EProcessorExecutionFlags::Editor); 
 	ExecutionOrder.ExecuteInGroup = Query.Description.Callback.Group;
 	ExecutionOrder.ExecuteBefore = Query.Description.Callback.BeforeGroups;
@@ -1363,9 +1367,11 @@ bool UTypedElementQueryObserverCallbackAdapterProcessorBase::ConfigureQueryCallb
 	FExtendedQueryStore& QueryStore,
 	FEnvironment& Environment, TArrayView<FMassEntityQuery> Subqueries)
 {
+	using namespace UE::Editor::DataStorage;
+
 	bool Result = Data.CommonQueryConfiguration(*this, Query, QueryHandle, QueryStore, Environment, Subqueries);
 
-	bRequiresGameThreadExecution = Query.Description.Callback.ExecutionMode == TypedElementDataStorage::EExecutionMode::GameThread;
+	bRequiresGameThreadExecution = Query.Description.Callback.ExecutionMode == EExecutionMode::GameThread;
 	ExecutionFlags = static_cast<int32>(EProcessorExecutionFlags::Editor);
 	
 	ObservedType = const_cast<UScriptStruct*>(Query.Description.Callback.MonitoredType);
diff --git a/Engine/Plugins/Experimental/EditorDataStorage/Source/TedsCore/Private/Queries/TypedElementAlertQueries.cpp b/Engine/Plugins/Experimental/EditorDataStorage/Source/TedsCore/Private/Queries/TypedElementAlertQueries.cpp
--- a/Engine/Plugins/Experimental/EditorDataStorage/Source/TedsCore/Private/Queries/TypedElementAlertQueries.cpp
+++ b/Engine/Plugins/Experimental/EditorDataStorage/Source/TedsCore/Private/Queries/TypedElementAlertQueries.cpp
@@ -17,7 +17,6 @@ FAutoConsoleCommand AddRandomAlertToRowConsoleCommand(
 	FConsoleCommandDelegate::CreateLambda(
 		[]()
 		{
-			using namespace TypedElementDataStorage;
 			using namespace TypedElementQueryBuilder;
 			using namespace UE::Editor::DataStorage;
 
@@ -77,7 +76,6 @@ FAutoConsoleCommand ClearAllAlertsConsoleCommand(
 	FConsoleCommandDelegate::CreateLambda(
 		[]()
 		{
-			using namespace TypedElementDataStorage;
 			using namespace TypedElementQueryBuilder;
 			using namespace UE::Editor::DataStorage;
 
@@ -118,7 +116,6 @@ FAutoConsoleCommand ClearSelectedAlertsConsoleCommand(
 	FConsoleCommandDelegate::CreateLambda(
 		[]()
 		{
-			using namespace TypedElementDataStorage;
 			using namespace TypedElementQueryBuilder;
 			using namespace UE::Editor::DataStorage;
 
@@ -177,7 +174,7 @@ void UTypedElementAlertQueriesFactory::RegisterSubQueries(ITypedElementDataStora
 
 void UTypedElementAlertQueriesFactory::RegisterParentUpdatesQueries(ITypedElementDataStorageInterface& DataStorage)
 {
-	using namespace TypedElementDataStorage;
+	using namespace UE::Editor::DataStorage;
 	using namespace TypedElementQueryBuilder;
 
 	DataStorage.RegisterQuery(
@@ -215,7 +212,7 @@ void UTypedElementAlertQueriesFactory::RegisterParentUpdatesQueries(ITypedElemen
 
 void UTypedElementAlertQueriesFactory::RegisterChildAlertUpdatesQueries(ITypedElementDataStorageInterface& DataStorage)
 {
-	using namespace TypedElementDataStorage;
+	using namespace UE::Editor::DataStorage;
 	using namespace TypedElementQueryBuilder;
 
 	DataStorage.RegisterQuery(
@@ -276,7 +273,6 @@ void UTypedElementAlertQueriesFactory::RegisterChildAlertUpdatesQueries(ITypedEl
 
 void UTypedElementAlertQueriesFactory::RegisterOnAddQueries(ITypedElementDataStorageInterface& DataStorage)
 {
-	using namespace TypedElementDataStorage;
 	using namespace TypedElementQueryBuilder;
 	using namespace UE::Editor::DataStorage;
 	
@@ -307,7 +303,6 @@ void UTypedElementAlertQueriesFactory::RegisterOnAddQueries(ITypedElementDataSto
 
 void UTypedElementAlertQueriesFactory::RegisterOnRemoveQueries(ITypedElementDataStorageInterface& DataStorage)
 {
-	using namespace TypedElementDataStorage;
 	using namespace TypedElementQueryBuilder;
 	using namespace UE::Editor::DataStorage;
 	
@@ -361,7 +356,6 @@ void UTypedElementAlertQueriesFactory::IncrementParents(
 	UE::Editor::DataStorage::IQueryContext& Context, UE::Editor::DataStorage::RowHandle Row, FTypedElementAlertColumnType AlertType,
 	int32 ChildAlertQueryIndex)
 {
-	using namespace TypedElementDataStorage;
 	using namespace TypedElementQueryBuilder;
 	using namespace UE::Editor::DataStorage;
 	
@@ -390,7 +384,7 @@ void UTypedElementAlertQueriesFactory::ResetChildAlertCounters(FTypedElementChil
 bool UTypedElementAlertQueriesFactory::MoveToNextParent(
 	UE::Editor::DataStorage::RowHandle& Parent, UE::Editor::DataStorage::IQueryContext& Context, int32 SubQueryIndex)
 {
-	using namespace TypedElementDataStorage;
+	using namespace UE::Editor::DataStorage;
 	using namespace TypedElementQueryBuilder;
 	
 	FQueryResult Result = Context.RunSubquery(SubQueryIndex, Parent, CreateSubqueryCallbackBinding(
diff --git a/Engine/Plugins/Experimental/EditorDataStorage/Source/TedsCore/Private/Queries/TypedElementExtendedQueryStore.cpp b/Engine/Plugins/Experimental/EditorDataStorage/Source/TedsCore/Private/Queries/TypedElementExtendedQueryStore.cpp
--- a/Engine/Plugins/Experimental/EditorDataStorage/Source/TedsCore/Private/Queries/TypedElementExtendedQueryStore.cpp
+++ b/Engine/Plugins/Experimental/EditorDataStorage/Source/TedsCore/Private/Queries/TypedElementExtendedQueryStore.cpp
@@ -11,10 +11,10 @@
 
 namespace UE::Editor::DataStorage
 {
-	const ITypedElementDataStorageInterface::FQueryDescription FExtendedQueryStore::EmptyDescription{};
+	const FQueryDescription FExtendedQueryStore::EmptyDescription{};
 
 	FExtendedQueryStore::Handle FExtendedQueryStore::RegisterQuery(
-		ITypedElementDataStorageInterface::FQueryDescription Query,
+		FQueryDescription Query,
 		FEnvironment& Environment,
 		FMassEntityManager& EntityManager,
 		FMassProcessingPhaseManager& PhaseManager)
@@ -66,8 +66,8 @@ namespace UE::Editor::DataStorage
 			});
 	}
 
-	void FExtendedQueryStore::RegisterTickGroup(FName GroupName, ITypedElementDataStorageInterface::EQueryTickPhase Phase,
-		FName BeforeGroup, FName AfterGroup, TypedElementDataStorage::EExecutionMode ExecutionMode)
+	void FExtendedQueryStore::RegisterTickGroup(FName GroupName, EQueryTickPhase Phase,
+		FName BeforeGroup, FName AfterGroup, EExecutionMode ExecutionMode)
 	{
 		FTickGroupDescription& Group = TickGroupDescriptions.FindOrAdd({ GroupName, Phase });
 
@@ -84,7 +84,7 @@ namespace UE::Editor::DataStorage
 		Group.ExecutionMode = ExecutionMode;
 	}
 
-	void FExtendedQueryStore::UnregisterTickGroup(FName GroupName, ITypedElementDataStorageInterface::EQueryTickPhase Phase)
+	void FExtendedQueryStore::UnregisterTickGroup(FName GroupName, EQueryTickPhase Phase)
 	{
 		TickGroupDescriptions.Remove({ GroupName, Phase });
 	}
@@ -119,7 +119,7 @@ namespace UE::Editor::DataStorage
 		return Queries.Get(Entry);
 	}
 
-	const ITypedElementDataStorageInterface::FQueryDescription& FExtendedQueryStore::GetQueryDescription(Handle Query) const
+	const FQueryDescription& FExtendedQueryStore::GetQueryDescription(Handle Query) const
 	{
 		const FExtendedQuery* QueryData = Get(Query);
 		return QueryData ? QueryData->Description : EmptyDescription;
@@ -178,12 +178,12 @@ namespace UE::Editor::DataStorage
 	}
 
 
-	TypedElementDataStorage::FQueryResult FExtendedQueryStore::RunQuery(FMassEntityManager& EntityManager, Handle Query)
+	FQueryResult FExtendedQueryStore::RunQuery(FMassEntityManager& EntityManager, Handle Query)
 	{
-		using ActionType = ITypedElementDataStorageInterface::FQueryDescription::EActionType;
-		using CompletionType = ITypedElementDataStorageInterface::FQueryResult::ECompletion;
+		using ActionType = FQueryDescription::EActionType;
+		using CompletionType = FQueryResult::ECompletion;
 
-		ITypedElementDataStorageInterface::FQueryResult Result;
+		FQueryResult Result;
 
 		if (FExtendedQuery* QueryData = Get(Query))
 		{
@@ -220,8 +220,6 @@ namespace UE::Editor::DataStorage
 		EDirectQueryExecutionFlags DirectExecutionFlags,
 		CallbackReference Callback)
 	{
-		using namespace UE::Editor::DataStorage;
-
 		FQueryResult Result;
 		if (FExtendedQuery* QueryData = Get(Query))
 		{
@@ -286,10 +284,10 @@ namespace UE::Editor::DataStorage
 		FEnvironment& Environment, FMassExecutionContext& ParentContext, Handle Query, RowHandle Row,
 		SubqueryCallbackRef Callback)
 	{
-		using ActionType = ITypedElementDataStorageInterface::FQueryDescription::EActionType;
-		using CompletionType = ITypedElementDataStorageInterface::FQueryResult::ECompletion;
+		using ActionType = FQueryDescription::EActionType;
+		using CompletionType = FQueryResult::ECompletion;
 
-		ITypedElementDataStorageInterface::FQueryResult Result;
+		FQueryResult Result;
 
 		if (FExtendedQuery* QueryData = Get(Query))
 		{
@@ -327,14 +325,14 @@ namespace UE::Editor::DataStorage
 		return Result;
 	}
 	void FExtendedQueryStore::RunPhasePreambleQueries(FMassEntityManager& EntityManager,
-		FEnvironment& Environment, ITypedElementDataStorageInterface::EQueryTickPhase Phase, float DeltaTime)
+		FEnvironment& Environment, EQueryTickPhase Phase, float DeltaTime)
 	{
 		RunPhasePreOrPostAmbleQueries(EntityManager, Environment, Phase, DeltaTime,
 			PhasePreparationQueries[static_cast<QueryTickPhaseType>(Phase)]);
 	}
 
 	void FExtendedQueryStore::RunPhasePostambleQueries(FMassEntityManager& EntityManager,
-		FEnvironment& Environment, ITypedElementDataStorageInterface::EQueryTickPhase Phase, float DeltaTime)
+		FEnvironment& Environment, EQueryTickPhase Phase, float DeltaTime)
 	{
 		RunPhasePreOrPostAmbleQueries(EntityManager, Environment, Phase, DeltaTime,
 			PhaseFinalizationQueries[static_cast<QueryTickPhaseType>(Phase)]);
@@ -371,10 +369,8 @@ namespace UE::Editor::DataStorage
 		Output.Log(TEXT("End of Typed Elements Data Storage query callback list."));
 	}
 
-	bool FExtendedQueryStore::SetupDynamicColumns(ITypedElementDataStorageInterface::FQueryDescription& Query, FEnvironment& Environment)
-	{
-		using namespace TypedElementDataStorage;
-	
+	bool FExtendedQueryStore::SetupDynamicColumns(FQueryDescription& Query, FEnvironment& Environment)
+	{	
 		const int32 SelectionCount = Query.DynamicSelectionTypes.Num();
 
 		for (int32 Index = 0, End = SelectionCount; Index < End; ++Index)
@@ -417,7 +413,7 @@ namespace UE::Editor::DataStorage
 	}
 
 	FMassEntityQuery& FExtendedQueryStore::SetupNativeQuery(
-		ITypedElementDataStorageInterface::FQueryDescription& Query, FExtendedQuery& StoredQuery)
+		FQueryDescription& Query, FExtendedQuery& StoredQuery)
 	{
 		/**
 		 * Mass verifies that queries that are used by processors are on the processor themselves. It does this by taking the address of the query
@@ -427,15 +423,13 @@ namespace UE::Editor::DataStorage
 		 * UHT doesn't allow for templates so it had to be done in an explicit way.
 		 */
 
-		using DSI = ITypedElementDataStorageInterface;
-
-		if (Query.Action == DSI::FQueryDescription::EActionType::Select)
+		if (Query.Action == FQueryDescription::EActionType::Select)
 		{
 			switch (Query.Callback.Type)
 			{
-			case DSI::EQueryCallbackType::None:
+			case EQueryCallbackType::None:
 				break;
-			case DSI::EQueryCallbackType::Processor:
+			case EQueryCallbackType::Processor:
 			{
 				UTypedElementQueryProcessorCallbackAdapterProcessorBase* Processor;
 				switch (Query.Subqueries.Num())
@@ -475,9 +469,9 @@ namespace UE::Editor::DataStorage
 				StoredQuery.Processor.Reset(Processor);
 				return Processor->GetQuery();
 			}
-			case DSI::EQueryCallbackType::ObserveAdd:
+			case EQueryCallbackType::ObserveAdd:
 				// Fall-through
-			case DSI::EQueryCallbackType::ObserveRemove:
+			case EQueryCallbackType::ObserveRemove:
 			{
 				UTypedElementQueryObserverCallbackAdapterProcessorBase* Observer;
 				switch (Query.Subqueries.Num())
@@ -517,9 +511,9 @@ namespace UE::Editor::DataStorage
 				StoredQuery.Processor.Reset(Observer);
 				return Observer->GetQuery();
 			}
-			case DSI::EQueryCallbackType::PhasePreparation:
+			case EQueryCallbackType::PhasePreparation:
 				break;
-			case DSI::EQueryCallbackType::PhaseFinalization:
+			case EQueryCallbackType::PhaseFinalization:
 				break;
 			default:
 				checkf(false, TEXT("Unsupported query callback type %i."), static_cast<int>(Query.Callback.Type));
@@ -529,16 +523,13 @@ namespace UE::Editor::DataStorage
 		return StoredQuery.NativeQuery;
 	}
 
-	bool FExtendedQueryStore::SetupSelectedColumns(
-		ITypedElementDataStorageInterface::FQueryDescription& Query, FMassEntityQuery& NativeQuery)
+	bool FExtendedQueryStore::SetupSelectedColumns(FQueryDescription& Query, FMassEntityQuery& NativeQuery)
 	{
-		using DSI = ITypedElementDataStorageInterface;
-
 		switch (Query.Action)
 		{
-		case DSI::FQueryDescription::EActionType::None:
+		case FQueryDescription::EActionType::None:
 			return true;
-		case DSI::FQueryDescription::EActionType::Select:
+		case FQueryDescription::EActionType::Select:
 		{
 			const int32 SelectionCount = Query.SelectionTypes.Num();
 			if (ensureMsgf(SelectionCount == Query.SelectionAccessTypes.Num(),
@@ -548,7 +539,7 @@ namespace UE::Editor::DataStorage
 				for (int SelectionIndex = 0; SelectionIndex < SelectionCount; ++SelectionIndex)
 				{
 					TWeakObjectPtr<const UScriptStruct>& Type = Query.SelectionTypes[SelectionIndex];
-					DSI::EQueryAccessType AccessType = Query.SelectionAccessTypes[SelectionIndex];
+					EQueryAccessType AccessType = Query.SelectionAccessTypes[SelectionIndex];
 					if (ensureMsgf(Type.IsValid(), TEXT("Provided query selection type can not be null.")) &&
 						ensureMsgf(
 							Type->IsChildOf(FColumn::StaticStruct()) ||
@@ -567,7 +558,7 @@ namespace UE::Editor::DataStorage
 			}
 			return false;
 		}
-		case DSI::FQueryDescription::EActionType::Count:
+		case FQueryDescription::EActionType::Count:
 		{
 			bool bIsSelectionEmpty = Query.SelectionTypes.IsEmpty();
 			bool bIsAccessTypesEmpty = Query.SelectionAccessTypes.IsEmpty();
@@ -581,11 +572,8 @@ namespace UE::Editor::DataStorage
 		}
 	}
 
-	bool FExtendedQueryStore::SetupConditions(
-		ITypedElementDataStorageInterface::FQueryDescription& Query, FMassEntityQuery& NativeQuery)
+	bool FExtendedQueryStore::SetupConditions(FQueryDescription& Query, FMassEntityQuery& NativeQuery)
 	{
-		using DSI = ITypedElementDataStorageInterface;
-
 		if (Query.ConditionTypes.IsEmpty())
 		{
 			return true;
@@ -594,19 +582,19 @@ namespace UE::Editor::DataStorage
 		if (ensureMsgf(Query.ConditionTypes.Num() == Query.ConditionOperators.Num(),
 			TEXT("The types and operators for a typed element query have gone out of sync.")))
 		{
-			const DSI::FQueryDescription::FOperator* Operand = Query.ConditionOperators.GetData();
-			for (DSI::FQueryDescription::EOperatorType Type : Query.ConditionTypes)
+			const FQueryDescription::FOperator* Operand = Query.ConditionOperators.GetData();
+			for (FQueryDescription::EOperatorType Type : Query.ConditionTypes)
 			{
 				EMassFragmentPresence Presence;
 				switch (Type)
 				{
-				case DSI::FQueryDescription::EOperatorType::SimpleAll:
+				case FQueryDescription::EOperatorType::SimpleAll:
 					Presence = EMassFragmentPresence::All;
 					break;
-				case DSI::FQueryDescription::EOperatorType::SimpleAny:
+				case FQueryDescription::EOperatorType::SimpleAny:
 					Presence = EMassFragmentPresence::Any;
 					break;
-				case DSI::FQueryDescription::EOperatorType::SimpleNone:
+				case FQueryDescription::EOperatorType::SimpleNone:
 					Presence = EMassFragmentPresence::None;
 					break;
 				default:
@@ -631,7 +619,7 @@ namespace UE::Editor::DataStorage
 
 	bool FExtendedQueryStore::SetupChunkFilters(
 		Handle QueryHandle,
-		ITypedElementDataStorageInterface::FQueryDescription& Query,
+		FQueryDescription& Query,
 		FEnvironment& Environment,
 		FMassEntityQuery& NativeQuery)
 	{
@@ -710,10 +698,8 @@ namespace UE::Editor::DataStorage
 	}
 
 	bool FExtendedQueryStore::SetupDependencies(
-		ITypedElementDataStorageInterface::FQueryDescription& Query, FMassEntityQuery& NativeQuery)
+		FQueryDescription& Query, FMassEntityQuery& NativeQuery)
 	{
-		using DSI = ITypedElementDataStorageInterface;
-
 		const int32 DependencyCount = Query.DependencyTypes.Num();
 		if (ensureMsgf(DependencyCount == Query.DependencyFlags.Num() && DependencyCount == Query.CachedDependencies.Num(),
 			TEXT("The number of query dependencies (%i) doesn't match the number of dependency access types (%i) and/or cached dependencies count (%i)."),
@@ -726,11 +712,11 @@ namespace UE::Editor::DataStorage
 					ensureMsgf(Type->IsChildOf<USubsystem>(), TEXT("Provided query dependency type '%s' is not based on USubSystem."),
 						*Type->GetStructPathName().ToString()))
 				{
-					DSI::EQueryDependencyFlags Flags = Query.DependencyFlags[DependencyIndex];
+					EQueryDependencyFlags Flags = Query.DependencyFlags[DependencyIndex];
 					NativeQuery.AddSubsystemRequirement(
 						const_cast<UClass*>(Type.Get()),
-						EnumHasAllFlags(Flags, DSI::EQueryDependencyFlags::ReadOnly) ? EMassFragmentAccess::ReadOnly : EMassFragmentAccess::ReadWrite,
-						EnumHasAllFlags(Flags, DSI::EQueryDependencyFlags::GameThreadBound));
+						EnumHasAllFlags(Flags, EQueryDependencyFlags::ReadOnly) ? EMassFragmentAccess::ReadOnly : EMassFragmentAccess::ReadWrite,
+						EnumHasAllFlags(Flags, EQueryDependencyFlags::GameThreadBound));
 				}
 				else
 				{
@@ -742,10 +728,8 @@ namespace UE::Editor::DataStorage
 		return false;
 	}
 
-	bool FExtendedQueryStore::SetupTickGroupDefaults(ITypedElementDataStorageInterface::FQueryDescription& Query)
+	bool FExtendedQueryStore::SetupTickGroupDefaults(FQueryDescription& Query)
 	{
-	using namespace TypedElementDataStorage;
-
 		const FTickGroupDescription* TickGroup = TickGroupDescriptions.Find({ Query.Callback.Group, Query.Callback.Phase });
 		if (TickGroup)
 		{
@@ -778,15 +762,13 @@ namespace UE::Editor::DataStorage
 	bool FExtendedQueryStore::SetupProcessors(Handle QueryHandle, FExtendedQuery& StoredQuery,
 		FEnvironment& Environment, FMassEntityManager& EntityManager, FMassProcessingPhaseManager& PhaseManager)
 	{
-		using DSI = ITypedElementDataStorageInterface;
-
 		// Register Phase processors locally.
 		switch (StoredQuery.Description.Callback.Type)
 		{
-		case DSI::EQueryCallbackType::PhasePreparation:
+		case EQueryCallbackType::PhasePreparation:
 			RegisterPreambleQuery(StoredQuery.Description.Callback.Phase, QueryHandle);
 			break;
-		case DSI::EQueryCallbackType::PhaseFinalization:
+		case EQueryCallbackType::PhaseFinalization:
 			RegisterPostambleQuery(StoredQuery.Description.Callback.Phase, QueryHandle);
 			break;
 		}
@@ -829,7 +811,7 @@ namespace UE::Editor::DataStorage
 		return true;
 	}
 
-	bool FExtendedQueryStore::SetupActivatable(Handle QueryHandle, ITypedElementDataStorageInterface::FQueryDescription& Query)
+	bool FExtendedQueryStore::SetupActivatable(Handle QueryHandle, FQueryDescription& Query)
 	{
 		if (!Query.Callback.ActivationName.IsNone())
 		{
@@ -838,15 +820,15 @@ namespace UE::Editor::DataStorage
 		return true;
 	}
 
-	EMassFragmentAccess FExtendedQueryStore::ConvertToNativeAccessType(ITypedElementDataStorageInterface::EQueryAccessType AccessType)
+	EMassFragmentAccess FExtendedQueryStore::ConvertToNativeAccessType(EQueryAccessType AccessType)
 	{
 		switch (AccessType)
 		{
-		case ITypedElementDataStorageInterface::EQueryAccessType::ReadOnly:
+		case EQueryAccessType::ReadOnly:
 			// Fall through
-		case ITypedElementDataStorageInterface::EQueryAccessType::OptionalReadOnly:
+		case EQueryAccessType::OptionalReadOnly:
 			return EMassFragmentAccess::ReadOnly;
-		case ITypedElementDataStorageInterface::EQueryAccessType::ReadWrite:
+		case EQueryAccessType::ReadWrite:
 			return EMassFragmentAccess::ReadWrite;
 		default:
 			checkf(false, TEXT("Invalid query access type: %i."), static_cast<uint32>(AccessType));
@@ -854,15 +836,15 @@ namespace UE::Editor::DataStorage
 		}
 	}
 
-	EMassFragmentPresence FExtendedQueryStore::ConvertToNativePresenceType(ITypedElementDataStorageInterface::EQueryAccessType AccessType)
+	EMassFragmentPresence FExtendedQueryStore::ConvertToNativePresenceType(EQueryAccessType AccessType)
 	{
 		switch (AccessType)
 		{
-		case ITypedElementDataStorageInterface::EQueryAccessType::ReadOnly:
+		case EQueryAccessType::ReadOnly:
 			return EMassFragmentPresence::All;
-		case ITypedElementDataStorageInterface::EQueryAccessType::OptionalReadOnly:
+		case EQueryAccessType::OptionalReadOnly:
 			return EMassFragmentPresence::Optional;
-		case ITypedElementDataStorageInterface::EQueryAccessType::ReadWrite:
+		case EQueryAccessType::ReadWrite:
 			return EMassFragmentPresence::All;
 		default:
 			checkf(false, TEXT("Invalid query access type: %i."), static_cast<uint32>(AccessType));
@@ -870,17 +852,17 @@ namespace UE::Editor::DataStorage
 		}
 	}
 
-	void FExtendedQueryStore::RegisterPreambleQuery(ITypedElementDataStorageInterface::EQueryTickPhase Phase, Handle Query)
+	void FExtendedQueryStore::RegisterPreambleQuery(EQueryTickPhase Phase, Handle Query)
 	{
 		PhasePreparationQueries[static_cast<QueryTickPhaseType>(Phase)].Add(Query);
 	}
 
-	void FExtendedQueryStore::RegisterPostambleQuery(ITypedElementDataStorageInterface::EQueryTickPhase Phase, Handle Query)
+	void FExtendedQueryStore::RegisterPostambleQuery(EQueryTickPhase Phase, Handle Query)
 	{
 		PhaseFinalizationQueries[static_cast<QueryTickPhaseType>(Phase)].Add(Query);
 	}
 
-	void FExtendedQueryStore::UnregisterPreambleQuery(ITypedElementDataStorageInterface::EQueryTickPhase Phase, Handle Query)
+	void FExtendedQueryStore::UnregisterPreambleQuery(EQueryTickPhase Phase, Handle Query)
 	{
 		int32 Index;
 		if (PhasePreparationQueries[static_cast<QueryTickPhaseType>(Phase)].Find(Query, Index))
@@ -889,7 +871,7 @@ namespace UE::Editor::DataStorage
 		}
 	}
 
-	void FExtendedQueryStore::UnregisterPostambleQuery(ITypedElementDataStorageInterface::EQueryTickPhase Phase, Handle Query)
+	void FExtendedQueryStore::UnregisterPostambleQuery(EQueryTickPhase Phase, Handle Query)
 	{
 		int32 Index;
 		if (PhaseFinalizationQueries[static_cast<QueryTickPhaseType>(Phase)].Find(Query, Index))
@@ -899,7 +881,7 @@ namespace UE::Editor::DataStorage
 	}
 
 	void FExtendedQueryStore::RunPhasePreOrPostAmbleQueries(FMassEntityManager& EntityManager,
-		FEnvironment& Environment, ITypedElementDataStorageInterface::EQueryTickPhase Phase,
+		FEnvironment& Environment, EQueryTickPhase Phase,
 		float DeltaTime, TArray<Handle>& QueryHandles)
 	{
 		if (!QueryHandles.IsEmpty())
@@ -943,11 +925,11 @@ namespace UE::Editor::DataStorage
 					*QueryData.Description.Callback.Name.ToString(), *QueryData.Processor->GetSparseClassDataStruct()->GetName());
 			}
 		}
-		else if (QueryData.Description.Callback.Type == ITypedElementDataStorageInterface::EQueryCallbackType::PhasePreparation)
+		else if (QueryData.Description.Callback.Type == EQueryCallbackType::PhasePreparation)
 		{
 			UnregisterPreambleQuery(QueryData.Description.Callback.Phase, Query);
 		}
-		else if (QueryData.Description.Callback.Type == ITypedElementDataStorageInterface::EQueryCallbackType::PhaseFinalization)
+		else if (QueryData.Description.Callback.Type == EQueryCallbackType::PhaseFinalization)
 		{
 			UnregisterPostambleQuery(QueryData.Description.Callback.Phase, Query);
 		}
diff --git a/Engine/Plugins/Experimental/EditorDataStorage/Source/TedsCore/Private/Queries/TypedElementExtendedQueryStore.h b/Engine/Plugins/Experimental/EditorDataStorage/Source/TedsCore/Private/Queries/TypedElementExtendedQueryStore.h
--- a/Engine/Plugins/Experimental/EditorDataStorage/Source/TedsCore/Private/Queries/TypedElementExtendedQueryStore.h
+++ b/Engine/Plugins/Experimental/EditorDataStorage/Source/TedsCore/Private/Queries/TypedElementExtendedQueryStore.h
@@ -54,7 +54,7 @@ namespace UE::Editor::DataStorage
 
 		/** Register the defaults for a tick group. These will be applied on top of any settings provided with a query registration. */
 		void RegisterTickGroup(FName GroupName, ITypedElementDataStorageInterface::EQueryTickPhase Phase,
-			FName BeforeGroup, FName AfterGroup, TypedElementDataStorage::EExecutionMode ExecutionMode);
+			FName BeforeGroup, FName AfterGroup, EExecutionMode ExecutionMode);
 		/** Removes a previously registered set of tick group defaults. */
 		void UnregisterTickGroup(FName GroupName, ITypedElementDataStorageInterface::EQueryTickPhase Phase);
 
@@ -152,7 +152,7 @@ namespace UE::Editor::DataStorage
 		{
 			TArray<FName> BeforeGroups;
 			TArray<FName> AfterGroups;
-			TypedElementDataStorage::EExecutionMode ExecutionMode = TypedElementDataStorage::EExecutionMode::Default;
+			EExecutionMode ExecutionMode = EExecutionMode::Default;
 		};
 
 		template<typename CallbackReference>
diff --git a/Engine/Plugins/Experimental/EditorDataStorage/Source/TedsCore/Private/Queries/TypedElementHierarchyQueries.cpp b/Engine/Plugins/Experimental/EditorDataStorage/Source/TedsCore/Private/Queries/TypedElementHierarchyQueries.cpp
--- a/Engine/Plugins/Experimental/EditorDataStorage/Source/TedsCore/Private/Queries/TypedElementHierarchyQueries.cpp
+++ b/Engine/Plugins/Experimental/EditorDataStorage/Source/TedsCore/Private/Queries/TypedElementHierarchyQueries.cpp
@@ -7,15 +7,14 @@
 
 void UTypedElementHiearchyQueriesFactory::RegisterQueries(ITypedElementDataStorageInterface& DataStorage)
 {
-	using namespace TypedElementDataStorage;
+	using namespace UE::Editor::DataStorage;
 	using namespace TypedElementQueryBuilder;
-	using DSI = ITypedElementDataStorageInterface;
 
 	DataStorage.RegisterQuery(
 		Select(
 			TEXT("Resolve hierarchy rows"),
-			FProcessor(DSI::EQueryTickPhase::PrePhysics, DataStorage.GetQueryTickGroupName(DSI::EQueryTickGroups::Default)),
-			[](DSI::IQueryContext& Context, RowHandle Row, const FUnresolvedTableRowParentColumn& UnresolvedParent)
+			FProcessor(EQueryTickPhase::PrePhysics, DataStorage.GetQueryTickGroupName(EQueryTickGroups::Default)),
+			[](IQueryContext& Context, RowHandle Row, const FUnresolvedTableRowParentColumn& UnresolvedParent)
 			{
 				RowHandle ParentRow = Context.FindIndexedRow(UnresolvedParent.ParentIdHash);
 				if (Context.IsRowAvailable(ParentRow))
diff --git a/Engine/Plugins/Experimental/EditorDataStorage/Source/TedsCore/Private/Queries/TypedElementMiscQueries.cpp b/Engine/Plugins/Experimental/EditorDataStorage/Source/TedsCore/Private/Queries/TypedElementMiscQueries.cpp
--- a/Engine/Plugins/Experimental/EditorDataStorage/Source/TedsCore/Private/Queries/TypedElementMiscQueries.cpp
+++ b/Engine/Plugins/Experimental/EditorDataStorage/Source/TedsCore/Private/Queries/TypedElementMiscQueries.cpp
@@ -9,13 +9,12 @@ void UTypedElementRemoveSyncToWorldTagFactory::RegisterQueries(ITypedElementData
 {
 	using namespace TypedElementQueryBuilder;
 	using namespace UE::Editor::DataStorage;
-	using DSI = ITypedElementDataStorageInterface;
 
 	DataStorage.RegisterQuery(
 		Select(
 			TEXT("Remove 'sync to world' tag"),
-			FPhaseAmble(FPhaseAmble::ELocation::Postamble, DSI::EQueryTickPhase::FrameEnd),
-			[](DSI::IQueryContext& Context, const RowHandle* Rows)
+			FPhaseAmble(FPhaseAmble::ELocation::Postamble, EQueryTickPhase::FrameEnd),
+			[](IQueryContext& Context, const RowHandle* Rows)
 			{
 				Context.RemoveColumns<FTypedElementSyncBackToWorldTag>(TConstArrayView<RowHandle>(Rows, Context.GetRowCount()));
 			}
@@ -28,8 +27,8 @@ void UTypedElementRemoveSyncToWorldTagFactory::RegisterQueries(ITypedElementData
 	DataStorage.RegisterQuery(
 	Select(
 		TEXT("Remove 'sync from world' tag"),
-		FPhaseAmble(FPhaseAmble::ELocation::Postamble, DSI::EQueryTickPhase::FrameEnd),
-		[](DSI::IQueryContext& Context, const RowHandle* Rows)
+		FPhaseAmble(FPhaseAmble::ELocation::Postamble, EQueryTickPhase::FrameEnd),
+		[](IQueryContext& Context, const RowHandle* Rows)
 		{
 			Context.RemoveColumns<FTypedElementSyncFromWorldTag>(TConstArrayView<RowHandle>(Rows, Context.GetRowCount()));
 		}
diff --git a/Engine/Plugins/Experimental/EditorDataStorage/Source/TedsCore/Private/TypedElementDatabase.cpp b/Engine/Plugins/Experimental/EditorDataStorage/Source/TedsCore/Private/TypedElementDatabase.cpp
--- a/Engine/Plugins/Experimental/EditorDataStorage/Source/TedsCore/Private/TypedElementDatabase.cpp
+++ b/Engine/Plugins/Experimental/EditorDataStorage/Source/TedsCore/Private/TypedElementDatabase.cpp
@@ -104,7 +104,6 @@ namespace TypedElementDatabasePrivate
 
 void UEditorDataStorage::Initialize()
 {
-	using namespace TypedElementDataStorage;
 	using namespace UE::Editor::DataStorage;
 
 	check(GEditor);
@@ -325,8 +324,6 @@ UE::Editor::DataStorage::RowHandle UEditorDataStorage::ReserveRow()
 
 void UEditorDataStorage::BatchReserveRows(int32 Count, TFunctionRef<void(RowHandle)> ReservationCallback)
 {
-	using namespace TypedElementDataStorage;
-
 	if (ActiveEditorEntityManager)
 	{
 		TArrayView<FMassEntityHandle> ReservedEntities(Environment->GetScratchBuffer().EmplaceArray<FMassEntityHandle>(Count), Count);
@@ -341,8 +338,6 @@ void UEditorDataStorage::BatchReserveRows(int32 Count, TFunctionRef<void(RowHand
 
 void UEditorDataStorage::BatchReserveRows(TArrayView<RowHandle> ReservedRows)
 {
-	using namespace TypedElementDataStorage;
-
 	if (ActiveEditorEntityManager)
 	{
 		// Depend on the fact that a row handle is an alias for an entity within the Mass powered backend. This
@@ -969,7 +964,7 @@ const UScriptStruct* UEditorDataStorage::GenerateDynamicColumn(const UE::Editor:
 }
 
 void UEditorDataStorage::RegisterTickGroup(
-	FName GroupName, EQueryTickPhase Phase, FName BeforeGroup, FName AfterGroup, TypedElementDataStorage::EExecutionMode ExecutionMode)
+	FName GroupName, EQueryTickPhase Phase, FName BeforeGroup, FName AfterGroup, UE::Editor::DataStorage::EExecutionMode ExecutionMode)
 {
 	Environment->GetQueryStore().RegisterTickGroup(GroupName, Phase, BeforeGroup, AfterGroup, ExecutionMode);
 }
@@ -1082,29 +1077,29 @@ void UEditorDataStorage::ActivateQueries(FName ActivationName)
 	}
 }
 
-UE::Editor::DataStorage::RowHandle UEditorDataStorage::FindIndexedRow(TypedElementDataStorage::IndexHash Index) const
+UE::Editor::DataStorage::RowHandle UEditorDataStorage::FindIndexedRow(UE::Editor::DataStorage::IndexHash Index) const
 {
 	return Environment->GetIndexTable().FindIndexedRow(UE::Editor::DataStorage::EGlobalLockScope::Public, Index);
 }
 
-void UEditorDataStorage::IndexRow(TypedElementDataStorage::IndexHash Index, RowHandle Row)
+void UEditorDataStorage::IndexRow(UE::Editor::DataStorage::IndexHash Index, RowHandle Row)
 {
 	Environment->GetIndexTable().IndexRow(UE::Editor::DataStorage::EGlobalLockScope::Public, Index, Row);
 }
 
 void UEditorDataStorage::BatchIndexRows(
-	TConstArrayView<TPair<TypedElementDataStorage::IndexHash, RowHandle>> IndexRowPairs)
+	TConstArrayView<TPair<UE::Editor::DataStorage::IndexHash, RowHandle>> IndexRowPairs)
 {
 	Environment->GetIndexTable().BatchIndexRows(UE::Editor::DataStorage::EGlobalLockScope::Public, IndexRowPairs);
 }
 
-void UEditorDataStorage::ReindexRow(TypedElementDataStorage::IndexHash OriginalIndex, TypedElementDataStorage::IndexHash NewIndex, 
+void UEditorDataStorage::ReindexRow(UE::Editor::DataStorage::IndexHash OriginalIndex, UE::Editor::DataStorage::IndexHash NewIndex, 
 	RowHandle RowHandle)
 {
 	Environment->GetIndexTable().ReindexRow(UE::Editor::DataStorage::EGlobalLockScope::Public, OriginalIndex, NewIndex, RowHandle);
 }
 
-void UEditorDataStorage::RemoveIndex(TypedElementDataStorage::IndexHash Index)
+void UEditorDataStorage::RemoveIndex(UE::Editor::DataStorage::IndexHash Index)
 {
 	Environment->GetIndexTable().RemoveIndex(UE::Editor::DataStorage::EGlobalLockScope::Public, Index);
 }
diff --git a/Engine/Plugins/Experimental/EditorDataStorage/Source/TedsCore/Private/TypedElementDatabase.h b/Engine/Plugins/Experimental/EditorDataStorage/Source/TedsCore/Private/TypedElementDatabase.h
--- a/Engine/Plugins/Experimental/EditorDataStorage/Source/TedsCore/Private/TypedElementDatabase.h
+++ b/Engine/Plugins/Experimental/EditorDataStorage/Source/TedsCore/Private/TypedElementDatabase.h
@@ -139,7 +139,7 @@ public:
 	const UScriptStruct* FindDynamicColumn(const UE::Editor::DataStorage::FDynamicColumnDescription& Description) const override;
 	const UScriptStruct* GenerateDynamicColumn(const UE::Editor::DataStorage::FDynamicColumnDescription& Description) override;
 
-	void RegisterTickGroup(FName GroupName, EQueryTickPhase Phase, FName BeforeGroup, FName AfterGroup, TypedElementDataStorage::EExecutionMode ExecutionMode);
+	void RegisterTickGroup(FName GroupName, EQueryTickPhase Phase, FName BeforeGroup, FName AfterGroup, UE::Editor::DataStorage::EExecutionMode ExecutionMode);
 	void UnregisterTickGroup(FName GroupName, EQueryTickPhase Phase);
 
 	QueryHandle RegisterQuery(FQueryDescription&& Query) override;
@@ -152,15 +152,15 @@ public:
 		DirectQueryCallbackRef Callback) override;
 	virtual void ActivateQueries(FName ActivationName) override;
 
-	virtual RowHandle FindIndexedRow(TypedElementDataStorage::IndexHash Index) const override;
-	virtual void IndexRow(TypedElementDataStorage::IndexHash Index, RowHandle Row) override;
+	virtual RowHandle FindIndexedRow(UE::Editor::DataStorage::IndexHash Index) const override;
+	virtual void IndexRow(UE::Editor::DataStorage::IndexHash Index, RowHandle Row) override;
 	virtual void BatchIndexRows(
-		TConstArrayView<TPair<TypedElementDataStorage::IndexHash, RowHandle>> IndexRowPairs) override;
+		TConstArrayView<TPair<UE::Editor::DataStorage::IndexHash, RowHandle>> IndexRowPairs) override;
 	virtual void ReindexRow(
-		TypedElementDataStorage::IndexHash OriginalIndex, 
-		TypedElementDataStorage::IndexHash NewIndex, 
+		UE::Editor::DataStorage::IndexHash OriginalIndex, 
+		UE::Editor::DataStorage::IndexHash NewIndex, 
 		RowHandle Row) override;
-	virtual void RemoveIndex(TypedElementDataStorage::IndexHash Index) override;
+	virtual void RemoveIndex(UE::Editor::DataStorage::IndexHash Index) override;
 
 	virtual FTypedElementOnDataStorageUpdate& OnUpdate() override;
 	virtual FTypedElementOnDataStorageUpdate& OnUpdateCompleted() override;
diff --git a/Engine/Plugins/Experimental/EditorDataStorage/Source/TedsCore/Private/TypedElementDatabaseCommandBuffer.cpp b/Engine/Plugins/Experimental/EditorDataStorage/Source/TedsCore/Private/TypedElementDatabaseCommandBuffer.cpp
--- a/Engine/Plugins/Experimental/EditorDataStorage/Source/TedsCore/Private/TypedElementDatabaseCommandBuffer.cpp
+++ b/Engine/Plugins/Experimental/EditorDataStorage/Source/TedsCore/Private/TypedElementDatabaseCommandBuffer.cpp
@@ -9,7 +9,6 @@
 
 namespace UE::Editor::DataStorage::Legacy
 {
-	using namespace UE::Editor::DataStorage;
 	// 
 	// Commands section
 	//
@@ -302,8 +301,6 @@ namespace UE::Editor::DataStorage::Legacy
 
 	void FCommandBuffer::ProcessCommands()
 	{
-		using namespace TypedElementDataStorage;
-
 		Commands.StableSort(
 			[](const FCommand& Lhs, const FCommand& Rhs)
 			{
diff --git a/Engine/Plugins/Experimental/EditorDataStorage/Source/TedsCore/Private/TypedElementDatabaseCompatibility.cpp b/Engine/Plugins/Experimental/EditorDataStorage/Source/TedsCore/Private/TypedElementDatabaseCompatibility.cpp
--- a/Engine/Plugins/Experimental/EditorDataStorage/Source/TedsCore/Private/TypedElementDatabaseCompatibility.cpp
+++ b/Engine/Plugins/Experimental/EditorDataStorage/Source/TedsCore/Private/TypedElementDatabaseCompatibility.cpp
@@ -66,7 +66,6 @@ namespace UE::Editor::DataStorage::Private
 
 void UEditorDataStorageCompatibility::Initialize(UEditorDataStorage* InStorage)
 {
-	using namespace TypedElementDataStorage;
 	using namespace UE::Editor::DataStorage;
 
 	checkf(InStorage, TEXT("TEDS Compatibility is being initialized with an invalid storage target."));
@@ -220,7 +219,6 @@ UE::Editor::DataStorage::RowHandle UEditorDataStorageCompatibility::AddCompatibl
 
 UE::Editor::DataStorage::RowHandle UEditorDataStorageCompatibility::AddCompatibleObjectExplicit(void* Object, TWeakObjectPtr<UScriptStruct> TypeInfo)
 {
-	using namespace TypedElementDataStorage;
 	using namespace UE::Editor::DataStorage;
 	
 	checkf(Storage, TEXT("Trying to add an object to Typed Element's Data Storage before the storage is available."));
@@ -250,7 +248,6 @@ void UEditorDataStorageCompatibility::RemoveCompatibleObjectExplicit(UObject* Ob
 
 void UEditorDataStorageCompatibility::RemoveCompatibleObjectExplicit(void* Object)
 {
-	using namespace TypedElementDataStorage;
 	using namespace UE::Editor::DataStorage;
 
 	checkf(Storage, TEXT("Removing compatible objects is not supported before Typed Element's Database compatibility manager has been initialized."));
@@ -279,7 +276,6 @@ void UEditorDataStorageCompatibility::RemoveCompatibleObjectExplicit(void* Objec
 
 UE::Editor::DataStorage::RowHandle UEditorDataStorageCompatibility::FindRowWithCompatibleObjectExplicit(const UObject* Object) const
 {
-	using namespace TypedElementDataStorage;
 	using namespace UE::Editor::DataStorage;
 
 	if (Object && Storage && Storage->IsAvailable())
@@ -850,7 +846,6 @@ void UEditorDataStorageCompatibility::PendingRegistration<AddressType>::Reset()
 void UEditorDataStorageCompatibility::TickPendingCommands()
 {
 	using namespace UE::Editor::DataStorage;
-	using namespace TypedElementDataStorage;
 
 	// Thread safe because it's only called from functions that already lock.
 	SIZE_T CommandCount = QueuedCommands.Collect(PendingCommands);
@@ -1239,7 +1234,6 @@ UEditorDataStorageCompatibility::FRegistrationCommandChange::~FRegistrationComma
 
 void UEditorDataStorageCompatibility::FRegistrationCommandChange::Apply(UObject* Object)
 {
-	using namespace TypedElementDataStorage;
 	using namespace UE::Editor::DataStorage;
 
 	checkf(Owner.IsValid() && Owner.Get() == Object, 
@@ -1264,7 +1258,6 @@ void UEditorDataStorageCompatibility::FRegistrationCommandChange::Apply(UObject*
 
 void UEditorDataStorageCompatibility::FRegistrationCommandChange::Revert(UObject* Object)
 {
-	using namespace TypedElementDataStorage;
 	using namespace UE::Editor::DataStorage;
 
 	checkf(Owner.IsValid() && Owner.Get() == Object,
@@ -1308,7 +1301,6 @@ UEditorDataStorageCompatibility::FDeregistrationCommandChange::FDeregistrationCo
 	: Owner(InOwner)
 	, TargetObject(InTargetObject)
 {
-	using namespace TypedElementDataStorage;
 	using namespace UE::Editor::DataStorage;
 
 	ITypedElementDataStorageInterface* DataStorage = InOwner->Storage;
@@ -1351,8 +1343,6 @@ void UEditorDataStorageCompatibility::FDeregistrationCommandChange::Apply(UObjec
 {
 	// All function calls are guaranteed to be thread safe.
 
-	using namespace TypedElementDataStorage;
-
 	checkf(Owner.IsValid() && Owner.Get() == Object,
 		TEXT("Applying deregistration transaction command within TEDS Compat was called after TEDS is not longer available."));
 	UEditorDataStorageCompatibility* DataStorageCompat = Owner.Get();
@@ -1364,7 +1354,6 @@ void UEditorDataStorageCompatibility::FDeregistrationCommandChange::Apply(UObjec
 
 void UEditorDataStorageCompatibility::FDeregistrationCommandChange::Revert(UObject* Object)
 {
-	using namespace TypedElementDataStorage;
 	using namespace UE::Editor::DataStorage;
 
 	checkf(Owner.IsValid() && Owner.Get() == Object,
diff --git a/Engine/Plugins/Experimental/EditorDataStorage/Source/TedsCore/Private/TypedElementDatabaseIndexTable.cpp b/Engine/Plugins/Experimental/EditorDataStorage/Source/TedsCore/Private/TypedElementDatabaseIndexTable.cpp
--- a/Engine/Plugins/Experimental/EditorDataStorage/Source/TedsCore/Private/TypedElementDatabaseIndexTable.cpp
+++ b/Engine/Plugins/Experimental/EditorDataStorage/Source/TedsCore/Private/TypedElementDatabaseIndexTable.cpp
@@ -5,10 +5,8 @@
 namespace UE::Editor::DataStorage
 {
 	RowHandle FIndexTable::FindIndexedRow(
-		EGlobalLockScope LockScope, TypedElementDataStorage::IndexHash Index) const
-	{
-		using namespace TypedElementDataStorage;
-	
+		EGlobalLockScope LockScope, IndexHash Index) const
+	{	
 		FScopedSharedLock Lock(LockScope);
 	
 		const RowHandle* Result = IndexLookupMap.Find(Index);
@@ -16,10 +14,8 @@ namespace UE::Editor::DataStorage
 	}
 
 	void FIndexTable::BatchIndexRows(EGlobalLockScope LockScope,
-		TConstArrayView<TPair<TypedElementDataStorage::IndexHash, RowHandle>> IndexRowPairs)
+		TConstArrayView<TPair<IndexHash, RowHandle>> IndexRowPairs)
 	{
-		using namespace TypedElementDataStorage;
-
 		FScopedExclusiveLock Lock(LockScope);
 	
 		IndexLookupMap.Reserve(IndexLookupMap.Num() + IndexRowPairs.Num());
@@ -32,7 +28,7 @@ namespace UE::Editor::DataStorage
 	}
 
 	void FIndexTable::IndexRow(EGlobalLockScope LockScope,
-		TypedElementDataStorage::IndexHash Index, RowHandle Row)
+		IndexHash Index, RowHandle Row)
 	{
 
 		FScopedExclusiveLock Lock(LockScope);
@@ -40,7 +36,7 @@ namespace UE::Editor::DataStorage
 	}
 
 	void FIndexTable::ReindexRow(EGlobalLockScope LockScope,
-		TypedElementDataStorage::IndexHash OriginalIndex, TypedElementDataStorage::IndexHash NewIndex, 
+		IndexHash OriginalIndex, IndexHash NewIndex, 
 		RowHandle Row)
 	{
 		FScopedExclusiveLock Lock(LockScope);
@@ -49,7 +45,7 @@ namespace UE::Editor::DataStorage
 		IndexRowUnguarded(NewIndex, Row);
 	}
 
-	void FIndexTable::RemoveIndex(EGlobalLockScope LockScope, TypedElementDataStorage::IndexHash Index)
+	void FIndexTable::RemoveIndex(EGlobalLockScope LockScope, IndexHash Index)
 	{
 		FScopedExclusiveLock Lock(LockScope);
 		RemoveIndexUnguarded(Index);
@@ -57,8 +53,6 @@ namespace UE::Editor::DataStorage
 
 	void FIndexTable::RemoveRow(EGlobalLockScope LockScope, RowHandle Row)
 	{
-		using namespace TypedElementDataStorage;
-
 		FScopedExclusiveLock Lock(LockScope);
 	
 		if (TMultiMap<RowHandle, IndexHash>::TKeyIterator It = ReverseIndexLookupMap.CreateKeyIterator(Row); It)
@@ -72,17 +66,15 @@ namespace UE::Editor::DataStorage
 		}
 	}
 
-	void FIndexTable::IndexRowUnguarded(TypedElementDataStorage::IndexHash Index, RowHandle Row)
+	void FIndexTable::IndexRowUnguarded(IndexHash Index, RowHandle Row)
 	{
 		IndexLookupMap.Add(Index, Row);
 		ReverseIndexLookupMap.Add(Row, Index);
 	}
 
 
-	void FIndexTable::RemoveIndexUnguarded(TypedElementDataStorage::IndexHash Index)
+	void FIndexTable::RemoveIndexUnguarded(IndexHash Index)
 	{
-		using namespace TypedElementDataStorage;
-
 		if (const RowHandle* Row = IndexLookupMap.Find(Index))
 		{
 			IndexLookupMap.Remove(Index);
diff --git a/Engine/Plugins/Experimental/EditorDataStorage/Source/TedsCore/Private/TypedElementDatabaseIndexTable.h b/Engine/Plugins/Experimental/EditorDataStorage/Source/TedsCore/Private/TypedElementDatabaseIndexTable.h
--- a/Engine/Plugins/Experimental/EditorDataStorage/Source/TedsCore/Private/TypedElementDatabaseIndexTable.h
+++ b/Engine/Plugins/Experimental/EditorDataStorage/Source/TedsCore/Private/TypedElementDatabaseIndexTable.h
@@ -20,32 +20,31 @@ namespace UE::Editor::DataStorage
 	public:
 		RowHandle FindIndexedRow(
 			EGlobalLockScope LockScope,
-			TypedElementDataStorage::IndexHash Index) const;
+			IndexHash Index) const;
 		void IndexRow(
 			EGlobalLockScope LockScope,
-			TypedElementDataStorage::IndexHash Index, 
+			IndexHash Index, 
 			RowHandle Row);
 		void BatchIndexRows(
 			EGlobalLockScope LockScope,
-			TConstArrayView<TPair<TypedElementDataStorage::IndexHash, 
-			RowHandle>> IndexRowPairs);
+			TConstArrayView<TPair<IndexHash, RowHandle>> IndexRowPairs);
 		void ReindexRow(
 			EGlobalLockScope LockScope,
-			TypedElementDataStorage::IndexHash OriginalIndex,
-			TypedElementDataStorage::IndexHash NewIndex,
+			IndexHash OriginalIndex,
+			IndexHash NewIndex,
 			RowHandle Row);
 		void RemoveIndex(
 			EGlobalLockScope LockScope,
-			TypedElementDataStorage::IndexHash Index);
+			IndexHash Index);
 		void RemoveRow(
 			EGlobalLockScope LockScope,
 			RowHandle Row);
 
 	private:
-		TMap<TypedElementDataStorage::IndexHash, RowHandle> IndexLookupMap;
-		TMultiMap<RowHandle, TypedElementDataStorage::IndexHash> ReverseIndexLookupMap;
+		TMap<IndexHash, RowHandle> IndexLookupMap;
+		TMultiMap<RowHandle, IndexHash> ReverseIndexLookupMap;
 	
-		void IndexRowUnguarded(TypedElementDataStorage::IndexHash Index, RowHandle Row);
-		void RemoveIndexUnguarded(TypedElementDataStorage::IndexHash Index);
+		void IndexRowUnguarded(IndexHash Index, RowHandle Row);
+		void RemoveIndexUnguarded(IndexHash Index);
 	};
 } // namespace UE::Editor::DataStorage
diff --git a/Engine/Plugins/Experimental/EditorDataStorage/Source/TedsCore/Tests/DynamicColumnTests.cpp b/Engine/Plugins/Experimental/EditorDataStorage/Source/TedsCore/Tests/DynamicColumnTests.cpp
--- a/Engine/Plugins/Experimental/EditorDataStorage/Source/TedsCore/Tests/DynamicColumnTests.cpp
+++ b/Engine/Plugins/Experimental/EditorDataStorage/Source/TedsCore/Tests/DynamicColumnTests.cpp
@@ -9,9 +9,7 @@
 
 #include "Misc/AutomationTest.h"
 
-namespace UE::Editor::DataStorage
-{
-namespace Tests
+namespace UE::Editor::DataStorage::Tests
 {
 	BEGIN_DEFINE_SPEC(DynamicColumnTestFixture, "Editor.DataStorage.DynamicColumns", EAutomationTestFlags::EditorContext | EAutomationTestFlags::EngineFilter)
 		ITypedElementDataStorageInterface* TedsInterface = nullptr;
@@ -193,7 +191,7 @@ namespace Tests
 								Compile());
 						
 							SetExpectedMatches({Rows[0]});
-							TypedElementDataStorage::FQueryResult Result = TedsInterface->RunQuery(Query, Callback);
+							FQueryResult Result = TedsInterface->RunQuery(Query, Callback);
 							TestEqual("Match Row[0]", Result.Count, GetMatchCount());
 						}
 						{
@@ -204,7 +202,7 @@ namespace Tests
 										All<FTestDynamicTag>(Identifiers[0]).
 								Compile());
 							SetExpectedMatches({Rows[0], Rows[1]});
-							TypedElementDataStorage::FQueryResult Result = TedsInterface->RunQuery(Query, Callback);
+							FQueryResult Result = TedsInterface->RunQuery(Query, Callback);
 							TestEqual("Match Row[0] and Row[1]", Result.Count, GetMatchCount());
 						}
 						{
@@ -216,7 +214,7 @@ namespace Tests
 										None<FTestDynamicTag>(Identifiers[1]).
 								Compile());
 							SetExpectedMatches({Rows[1]});
-							TypedElementDataStorage::FQueryResult Result = TedsInterface->RunQuery(Query, Callback);
+							FQueryResult Result = TedsInterface->RunQuery(Query, Callback);
 							TestEqual("Match Row[1]", Result.Count, GetMatchCount());
 						}
 						{
@@ -227,7 +225,7 @@ namespace Tests
 										All<FTestDynamicTag>(Identifiers[1]).
 								Compile());
 							SetExpectedMatches({Rows[0], Rows[2]});
-							TypedElementDataStorage::FQueryResult Result = TedsInterface->RunQuery(Query, Callback);
+							FQueryResult Result = TedsInterface->RunQuery(Query, Callback);
 							TestEqual("Match Row[0] and Row[2]", Result.Count, GetMatchCount());
 						}
 						{
@@ -239,7 +237,7 @@ namespace Tests
 										All<FTestDynamicTag>(Identifiers[1]).
 								Compile());
 							SetExpectedMatches({Rows[2]});
-							TypedElementDataStorage::FQueryResult Result = TedsInterface->RunQuery(Query, Callback);
+							FQueryResult Result = TedsInterface->RunQuery(Query, Callback);
 							TestEqual("Match Row[2]", Result.Count, GetMatchCount());
 						}
 						{
@@ -251,7 +249,7 @@ namespace Tests
 										Any<FTestDynamicTag>(Identifiers[1]).
 								Compile());
 							SetExpectedMatches({Rows[0], Rows[1], Rows[2]});
-							TypedElementDataStorage::FQueryResult Result = TedsInterface->RunQuery(Query, Callback);
+							FQueryResult Result = TedsInterface->RunQuery(Query, Callback);
 							TestEqual("Match All Rows", Result.Count, GetMatchCount());
 						}
 					}).Wait();
@@ -299,7 +297,6 @@ namespace Tests
 						};
 						
 						using namespace TypedElementQueryBuilder;
-						using namespace TypedElementDataStorage;
 						// Setup an activatable processor for the test
 						ActivationKeys.Emplace(TEXT("TEST: Match Row[0]"));
 						QueryExpectedMatchRows.Emplace(TArray({Rows[0]}));
@@ -488,7 +485,6 @@ namespace Tests
 					Async(EAsyncExecution::TaskGraphMainThread,[this, &ActivationKeys, &AllTestExpectations, &UnexpectedRowCount]()
 					{
 						using namespace TypedElementQueryBuilder;
-						using namespace TypedElementDataStorage;
 
 						auto RunTest = [this, &AllTestExpectations, &UnexpectedRowCount](IQueryContext& Context, const RowHandle* RowsPtr, int32 TestIndex)
 						{
@@ -613,7 +609,6 @@ namespace Tests
 			TedsInterface = nullptr;
 		});
 	}
-} // namespace Tests
-} // namespace UE::Editor::DataStorage
+} // namespace UE::Editor::DataStorage::Tests
 
-#endif // WITH_TESTS
\ No newline at end of file
+#endif // WITH_TESTS
diff --git a/Engine/Plugins/Experimental/EditorDataStorage/Source/TedsUI/Private/Processors/WidgetReferenceColumnUpdateProcessor.cpp b/Engine/Plugins/Experimental/EditorDataStorage/Source/TedsUI/Private/Processors/WidgetReferenceColumnUpdateProcessor.cpp
--- a/Engine/Plugins/Experimental/EditorDataStorage/Source/TedsUI/Private/Processors/WidgetReferenceColumnUpdateProcessor.cpp
+++ b/Engine/Plugins/Experimental/EditorDataStorage/Source/TedsUI/Private/Processors/WidgetReferenceColumnUpdateProcessor.cpp
@@ -14,14 +14,13 @@ void UWidgetReferenceColumnUpdateFactory::RegisterQueries(ITypedElementDataStora
 void UWidgetReferenceColumnUpdateFactory::RegisterDeleteRowOnWidgetDeleteQuery(ITypedElementDataStorageInterface& DataStorage) const
 {
 	using namespace TypedElementQueryBuilder;
-	using DSI = ITypedElementDataStorageInterface;
-	using namespace UE::Editor;
+	using namespace UE::Editor::DataStorage;
 
 	DataStorage.RegisterQuery(
     	Select(
     		TEXT("Delete row with deleted widget"),
-    		FPhaseAmble(FPhaseAmble::ELocation::Preamble, DSI::EQueryTickPhase::FrameEnd),
-    		[](DSI::IQueryContext& Context, DataStorage::RowHandle Row, const FTypedElementSlateWidgetReferenceColumn& WidgetReference)
+    		FPhaseAmble(FPhaseAmble::ELocation::Preamble, EQueryTickPhase::FrameEnd),
+    		[](IQueryContext& Context, RowHandle Row, const FTypedElementSlateWidgetReferenceColumn& WidgetReference)
     		{
     			if (!WidgetReference.TedsWidget.IsValid())
     			{
@@ -38,14 +37,13 @@ void UWidgetReferenceColumnUpdateFactory::RegisterDeleteRowOnWidgetDeleteQuery(I
 void UWidgetReferenceColumnUpdateFactory::RegisterDeleteColumnOnWidgetDeleteQuery(ITypedElementDataStorageInterface& DataStorage) const
 {
 	using namespace TypedElementQueryBuilder;
-	using DSI = ITypedElementDataStorageInterface;
-	using namespace UE::Editor;
+	using namespace UE::Editor::DataStorage;
 
 	DataStorage.RegisterQuery(
 		Select(
 			TEXT("Delete widget columns for deleted widget"),
-			FPhaseAmble(FPhaseAmble::ELocation::Preamble, DSI::EQueryTickPhase::FrameEnd),
-			[](DSI::IQueryContext& Context, DataStorage::RowHandle Row, const FTypedElementSlateWidgetReferenceColumn& WidgetReference)
+			FPhaseAmble(FPhaseAmble::ELocation::Preamble, EQueryTickPhase::FrameEnd),
+			[](IQueryContext& Context, RowHandle Row, const FTypedElementSlateWidgetReferenceColumn& WidgetReference)
 			{
 				if (!WidgetReference.TedsWidget.IsValid())
 				{
diff --git a/Engine/Plugins/Experimental/EditorDataStorage/Source/TedsUI/Private/Widgets/CounterWidget.cpp b/Engine/Plugins/Experimental/EditorDataStorage/Source/TedsUI/Private/Widgets/CounterWidget.cpp
--- a/Engine/Plugins/Experimental/EditorDataStorage/Source/TedsUI/Private/Widgets/CounterWidget.cpp
+++ b/Engine/Plugins/Experimental/EditorDataStorage/Source/TedsUI/Private/Widgets/CounterWidget.cpp
@@ -46,7 +46,7 @@ UCounterWidgetFactory::UCounterWidgetFactory()
 void UCounterWidgetFactory::RegisterQueries(ITypedElementDataStorageInterface& DataStorage)
 {
 	using namespace TypedElementQueryBuilder;
-	using namespace TypedElementDataStorage;
+	using namespace UE::Editor::DataStorage;
 	
 	DataStorage.RegisterQuery(Select(TEXT("Sync counter widgets"), 
 		FProcessor(EQueryTickPhase::FrameEnd, DataStorage.GetQueryTickGroupName(EQueryTickGroups::SyncWidgets))
diff --git a/Engine/Plugins/Experimental/EditorDataStorage/Source/TedsUI/Private/Widgets/ExportedTextWidget.cpp b/Engine/Plugins/Experimental/EditorDataStorage/Source/TedsUI/Private/Widgets/ExportedTextWidget.cpp
--- a/Engine/Plugins/Experimental/EditorDataStorage/Source/TedsUI/Private/Widgets/ExportedTextWidget.cpp
+++ b/Engine/Plugins/Experimental/EditorDataStorage/Source/TedsUI/Private/Widgets/ExportedTextWidget.cpp
@@ -55,7 +55,6 @@ namespace UE::Editor
 		static DataStorage::RowHandle RegisterUpdateCallback(ITypedElementDataStorageInterface& DataStorage, const UScriptStruct* Target)
 		{
 			using namespace TypedElementQueryBuilder;
-			using namespace TypedElementDataStorage;
 			using namespace UE::Editor::DataStorage;
 			
 			QueryHandle TypeDataQuery = DataStorage.RegisterQuery(
diff --git a/Engine/Plugins/Experimental/EditorDataStorageFeatures/Source/TedsAssetData/Private/CB/TedsAssetDataCBDataSource.cpp b/Engine/Plugins/Experimental/EditorDataStorageFeatures/Source/TedsAssetData/Private/CB/TedsAssetDataCBDataSource.cpp
--- a/Engine/Plugins/Experimental/EditorDataStorageFeatures/Source/TedsAssetData/Private/CB/TedsAssetDataCBDataSource.cpp
+++ b/Engine/Plugins/Experimental/EditorDataStorageFeatures/Source/TedsAssetData/Private/CB/TedsAssetDataCBDataSource.cpp
@@ -18,12 +18,12 @@
 
 namespace UE::Editor::AssetData::Private
 {
+	using namespace UE::Editor::DataStorage;
 
 FTedsAssetDataCBDataSource::FTedsAssetDataCBDataSource(ITypedElementDataStorageInterface& InDatabase)
 	: Database(InDatabase)
 {
 	using namespace TypedElementQueryBuilder;
-	using namespace UE::Editor::DataStorage;
 
 	InitVirtualPathProcessor();
 
@@ -41,7 +41,7 @@ FTedsAssetDataCBDataSource::FTedsAssetDataCBDataSource(ITypedElementDataStorageI
 	ProcessPathQuery = Database.RegisterQuery(
 		Select(
 			TEXT("FTedsAssetDataCBDataSource: Process Path updates"),
-			FProcessor(TypedElementDataStorage::EQueryTickPhase::DuringPhysics, Database.GetQueryTickGroupName(TypedElementDataStorage::EQueryTickGroups::Update)),
+			FProcessor(EQueryTickPhase::DuringPhysics, Database.GetQueryTickGroupName(EQueryTickGroups::Update)),
 			[GenerateVirtualPaths](IQueryContext& Context, const RowHandle* Rows, const FAssetPathColumn_Experimental* PathColumn)
 			{
 				int32 NumOfRowToProcess = Context.GetRowCount();
@@ -71,7 +71,7 @@ FTedsAssetDataCBDataSource::FTedsAssetDataCBDataSource(ITypedElementDataStorageI
 	ProcessAssetDataPathUpdateQuery = Database.RegisterQuery(
 		Select(
 			TEXT("FTedsAssetDataCBDataSource:: Process Asset Data Path Update"),
-			FProcessor(TypedElementDataStorage::EQueryTickPhase::DuringPhysics, Database.GetQueryTickGroupName(TypedElementDataStorage::EQueryTickGroups::Update)),
+			FProcessor(EQueryTickPhase::DuringPhysics, Database.GetQueryTickGroupName(EQueryTickGroups::Update)),
 			[GenerateVirtualPaths](IQueryContext& Context, const RowHandle Row, const FAssetDataColumn_Experimental& AssetDataColumn)
 			{
 				int32 NumOfRowToProcess = Context.GetRowCount();
@@ -132,8 +132,8 @@ FTedsAssetDataCBDataSource::FTedsAssetDataCBDataSource(ITypedElementDataStorageI
 	ProcessAssetDataAndPathUpdateQuery = Database.RegisterQuery(
 		Select(
 			TEXT("FTedsAssetDataCBDataSource: Process Asset Data and Path Updates"),
-			FProcessor(TypedElementDataStorage::EQueryTickPhase::DuringPhysics, Database.GetQueryTickGroupName(TypedElementDataStorage::EQueryTickGroups::Update)),
-			[GenerateVirtualPaths, AddAssetDataColumns, AssetRegistry = static_cast<const IAssetRegistry*>(&IAssetRegistry::GetChecked())](UE::Editor::DataStorage::IQueryContext& Context, const RowHandle* Rows, const FAssetDataColumn_Experimental* AssetDataColumn)
+			FProcessor(EQueryTickPhase::DuringPhysics, Database.GetQueryTickGroupName(EQueryTickGroups::Update)),
+			[GenerateVirtualPaths, AddAssetDataColumns, AssetRegistry = static_cast<const IAssetRegistry*>(&IAssetRegistry::GetChecked())](IQueryContext& Context, const RowHandle* Rows, const FAssetDataColumn_Experimental* AssetDataColumn)
 			{
 				const int32 RowCount = Context.GetRowCount();
 
@@ -182,8 +182,8 @@ FTedsAssetDataCBDataSource::FTedsAssetDataCBDataSource(ITypedElementDataStorageI
 	ProcessAssetDataUpdateQuery = Database.RegisterQuery(
 		Select(
 			TEXT("FTedsAssetDataCBDataSource: Process Asset Data updates"),
-			FProcessor(TypedElementDataStorage::EQueryTickPhase::DuringPhysics, Database.GetQueryTickGroupName(TypedElementDataStorage::EQueryTickGroups::Update)),
-			[AddAssetDataColumns,  AssetRegistry = static_cast<const IAssetRegistry*>(&IAssetRegistry::GetChecked())](UE::Editor::DataStorage::IQueryContext& Context, const RowHandle* Rows, const FAssetDataColumn_Experimental* AssetDataColumn)
+			FProcessor(EQueryTickPhase::DuringPhysics, Database.GetQueryTickGroupName(EQueryTickGroups::Update)),
+			[AddAssetDataColumns,  AssetRegistry = static_cast<const IAssetRegistry*>(&IAssetRegistry::GetChecked())](IQueryContext& Context, const RowHandle* Rows, const FAssetDataColumn_Experimental* AssetDataColumn)
 			{
 				const int32 RowCount = Context.GetRowCount();
 				TArray<FName, TInlineAllocator<32>> PackageNames;
diff --git a/Engine/Plugins/Experimental/EditorDataStorageFeatures/Source/TedsAssetData/Private/Debug/TedsAssetDataDebugCommands.cpp b/Engine/Plugins/Experimental/EditorDataStorageFeatures/Source/TedsAssetData/Private/Debug/TedsAssetDataDebugCommands.cpp
--- a/Engine/Plugins/Experimental/EditorDataStorageFeatures/Source/TedsAssetData/Private/Debug/TedsAssetDataDebugCommands.cpp
+++ b/Engine/Plugins/Experimental/EditorDataStorageFeatures/Source/TedsAssetData/Private/Debug/TedsAssetDataDebugCommands.cpp
@@ -23,9 +23,7 @@
 
 DEFINE_LOG_CATEGORY_STATIC(LogTEDSAssetRegistry, Log, All)
 
-namespace UE::Editor::DataStorage
-{
-namespace Debug::Private
+namespace UE::Editor::DataStorage::Debug::Private
 {
 	static FAutoConsoleCommand CCMDTestFolderRowData(
 		TEXT("TEDS.Debug.ShowDataOfAssetFolder"),
@@ -38,7 +36,7 @@ namespace Debug::Private
 				for (const FString& Path : InArgs)
 				{
 					FName PathAsName(*Path);
-					RowHandle Row = Database->FindIndexedRow(TypedElementDataStorage::GenerateIndexHash(PathAsName));
+					RowHandle Row = Database->FindIndexedRow(GenerateIndexHash(PathAsName));
 
 					if (Database->IsRowAssigned(Row))
 					{
@@ -100,7 +98,7 @@ namespace Debug::Private
 					AssetRegistry.GetAssetsByPath(PathAsName, Assets);
 					for (const FAssetData& Asset : Assets)
 					{
-						RowHandle AssetRow = Database->FindIndexedRow(TypedElementDataStorage::GenerateIndexHash(Asset.GetSoftObjectPath()));
+						RowHandle AssetRow = Database->FindIndexedRow(GenerateIndexHash(Asset.GetSoftObjectPath()));
 
 						if (const FUnresolvedAssetsInPathColumn_Experimental* UnresolvedAssetsInPathColumn = Database->GetColumn<FUnresolvedAssetsInPathColumn_Experimental>(AssetRow))
 						{
@@ -221,7 +219,6 @@ namespace Debug::Private
 			{
 				TRACE_CPUPROFILER_EVENT_SCOPE(TEDS.Debug.CreateDebugAssetRows);
 
-				using namespace TypedElementDataStorage;
 				using namespace TypedElementQueryBuilder;
 
 				ITypedElementDataStorageInterface* DataStorage = UTypedElementRegistry::GetInstance()->GetMutableDataStorage();
@@ -254,7 +251,6 @@ namespace Debug::Private
 		FConsoleCommandWithArgsDelegate::CreateLambda(
 			[](const TArray<FString>& Args)
 			{
-				using namespace TypedElementDataStorage;
 				using namespace TypedElementQueryBuilder;
 
 				ITypedElementDataStorageInterface* DataStorage = UTypedElementRegistry::GetInstance()->GetMutableDataStorage();
@@ -281,5 +277,4 @@ namespace Debug::Private
 				}
 			}
 		));
-} // namespace Debug::Private
-} // namespace UE::Editor::DataStorage
+} // namespace UE::Editor::DataStorage::Debug::Private
diff --git a/Engine/Plugins/Experimental/EditorDataStorageFeatures/Source/TedsAssetData/Private/TedsAssetData.cpp b/Engine/Plugins/Experimental/EditorDataStorageFeatures/Source/TedsAssetData/Private/TedsAssetData.cpp
--- a/Engine/Plugins/Experimental/EditorDataStorageFeatures/Source/TedsAssetData/Private/TedsAssetData.cpp
+++ b/Engine/Plugins/Experimental/EditorDataStorageFeatures/Source/TedsAssetData/Private/TedsAssetData.cpp
@@ -16,19 +16,19 @@
 
 #define TRACK_TEDSASSETDATA_MEMORY 0
 
-namespace UE::Editor
+namespace UE::Editor::AssetData::Private
 {
-	namespace AssetData::Private
-	{	
+using namespace DataStorage;
+
 constexpr int32 ParallelForMinBatchSize = 1024 * 4;
 
 struct FPopulateAssetDataRowArgs
 {
 	FAssetData AssetData;
 
-	TypedElementDataStorage::IndexHash ObjectPathHash;
-	TypedElementDataStorage::IndexHash PathHash;
-	DataStorage::RowHandle PathRow = DataStorage::InvalidRowHandle;
+	IndexHash ObjectPathHash;
+	IndexHash PathHash;
+	RowHandle PathRow = InvalidRowHandle;
 };
 
 // Only safe if the GT is blocked during the operation
@@ -36,7 +36,7 @@ template<typename TAssetData>
 FPopulateAssetDataRowArgs ThreadSafe_PopulateAssetDataTableRow(TAssetData&& InAssetData, const ITypedElementDataStorageInterface& Database)
 {
 	FPopulateAssetDataRowArgs Output;
-	Output.ObjectPathHash = TypedElementDataStorage::GenerateIndexHash(InAssetData.GetSoftObjectPath());
+	Output.ObjectPathHash = GenerateIndexHash(InAssetData.GetSoftObjectPath());
 
 	// Looks safe but might not be depending on the implementation of the database
 	if (Database.IsRowAssigned(Database.FindIndexedRow(Output.ObjectPathHash)))
@@ -45,7 +45,7 @@ FPopulateAssetDataRowArgs ThreadSafe_PopulateAssetDataTableRow(TAssetData&& InAs
 		return Output;
 	}
 
-	Output.PathHash = TypedElementDataStorage::GenerateIndexHash(InAssetData.PackagePath);
+	Output.PathHash = GenerateIndexHash(InAssetData.PackagePath);
 
 	// Looks safe but might not be depending on the implementation of the database
 	Output.PathRow = Database.FindIndexedRow(Output.PathHash);
@@ -55,9 +55,9 @@ FPopulateAssetDataRowArgs ThreadSafe_PopulateAssetDataTableRow(TAssetData&& InAs
 }
 
 
-void PopulateAssetDataTableRow(FPopulateAssetDataRowArgs&& InAssetDataRowArgs, ITypedElementDataStorageInterface& Database, DataStorage::RowHandle RowHandle)
+void PopulateAssetDataTableRow(FPopulateAssetDataRowArgs&& InAssetDataRowArgs, ITypedElementDataStorageInterface& Database, RowHandle RowHandle)
 {
-	if (InAssetDataRowArgs.PathRow !=  DataStorage::InvalidRowHandle)
+	if (InAssetDataRowArgs.PathRow !=  InvalidRowHandle)
 	{
 		Database.GetColumn<FAssetsInPathColumn_Experimental>(InAssetDataRowArgs.PathRow)->AssetsRow.Add(RowHandle);
 	}
@@ -75,8 +75,8 @@ void PopulateAssetDataTableRow(FPopulateAssetDataRowArgs&& InAssetDataRowArgs, I
 struct FPopulatePathRowArgs
 {
 	FName AssetRegistryPath;
-	TypedElementDataStorage::IndexHash AssetRegistryPathHash;
-	TypedElementDataStorage::IndexHash ParentAssetRegistryPathHash;
+	IndexHash AssetRegistryPathHash;
+	IndexHash ParentAssetRegistryPathHash;
 	uint32 PathDepth;
 
 	operator bool() const 
@@ -113,16 +113,16 @@ void GetPathDepthAndParentFolderIndex(FStringView Path, uint32& OutDepth, int32&
 }
 
 // Only thread safe if the game thread is blocked 
-FPopulatePathRowArgs ThreadSafe_PopulatePathRowArgs(TypedElementDataStorage::IndexHash AssetRegistryPathHash, FName InAssetRegistryPath, FStringView PathAsString)
+FPopulatePathRowArgs ThreadSafe_PopulatePathRowArgs(IndexHash AssetRegistryPathHash, FName InAssetRegistryPath, FStringView PathAsString)
 {
-	TypedElementDataStorage::IndexHash ParentAssetRegistryPathHash = DataStorage::InvalidRowHandle;
+	IndexHash ParentAssetRegistryPathHash = InvalidRowHandle;
 	int32 CharacterIndex;
 	uint32 Depth;
 	GetPathDepthAndParentFolderIndex(PathAsString, Depth, CharacterIndex);
 	if (CharacterIndex != INDEX_NONE)
 	{
 		FStringView ParentPath = PathAsString.Left(CharacterIndex);
-		ParentAssetRegistryPathHash = TypedElementDataStorage::GenerateIndexHash(FName(ParentPath));
+		ParentAssetRegistryPathHash = GenerateIndexHash(FName(ParentPath));
 	}
 
 	FPopulatePathRowArgs Args;
@@ -134,11 +134,11 @@ FPopulatePathRowArgs ThreadSafe_PopulatePathRowArgs(TypedElementDataStorage::Ind
 	return Args;
 }
 
-void PopulatePathDataTableRow(FPopulatePathRowArgs&& InPopulatePathRowArgs, ITypedElementDataStorageInterface& Database, DataStorage::RowHandle InRowHandle)
+void PopulatePathDataTableRow(FPopulatePathRowArgs&& InPopulatePathRowArgs, ITypedElementDataStorageInterface& Database, RowHandle InRowHandle)
 {
-	if (InPopulatePathRowArgs.ParentAssetRegistryPathHash != DataStorage::InvalidRowHandle)
+	if (InPopulatePathRowArgs.ParentAssetRegistryPathHash != InvalidRowHandle)
 	{
-		const DataStorage::RowHandle ParentRow = Database.FindIndexedRow(InPopulatePathRowArgs.ParentAssetRegistryPathHash);
+		const RowHandle ParentRow = Database.FindIndexedRow(InPopulatePathRowArgs.ParentAssetRegistryPathHash);
 		if (Database.IsRowAssigned(ParentRow))
 		{
 			Database.GetColumn<FChildrenAssetPathColumn_Experimental>(ParentRow)->ChildrenRows.Add(InRowHandle);
@@ -179,13 +179,13 @@ FTedsAssetData::FTedsAssetData(ITypedElementDataStorageInterface& InDatabase)
 
 	// Register data types to TEDS
 	PathsTable = Database.FindTable(FName(TEXT("Editor_AssetRegistryPathsTable")));
-	if (PathsTable == DataStorage::InvalidTableHandle)
+	if (PathsTable == InvalidTableHandle)
 	{
 		PathsTable = Database.RegisterTable<FAssetPathColumn_Experimental, FChildrenAssetPathColumn_Experimental, FParentAssetPathColumn_Experimental, FAssetsInPathColumn_Experimental>(FName(TEXT("Editor_AssetRegistryPathsTable")));
 	}
 
 	AssetsDataTable = Database.FindTable(FName(TEXT("Editor_AssetRegistryAssetDataTable")));
-	if (AssetsDataTable == DataStorage::InvalidTableHandle)
+	if (AssetsDataTable == InvalidTableHandle)
 	{
 		AssetsDataTable = Database.RegisterTable<FAssetDataColumn_Experimental, FUpdatedPathTag, FUpdatedAssetDataTag>(FName(TEXT("Editor_AssetRegistryAssetDataTable")));
 	}
@@ -194,10 +194,10 @@ FTedsAssetData::FTedsAssetData(ITypedElementDataStorageInterface& InDatabase)
 	RemoveUpdatedPathTagQuery = Database.RegisterQuery(
 			TypedElementQueryBuilder::Select(
 				TEXT("FTedsAssetData: Remove Updated Path Tag"),
-				TypedElementQueryBuilder::FPhaseAmble(TypedElementQueryBuilder::FPhaseAmble::ELocation::Postamble, TypedElementDataStorage::EQueryTickPhase::FrameEnd),
-				[](DataStorage::IQueryContext& Context, const DataStorage::RowHandle* Rows)
+				TypedElementQueryBuilder::FPhaseAmble(TypedElementQueryBuilder::FPhaseAmble::ELocation::Postamble, EQueryTickPhase::FrameEnd),
+				[](IQueryContext& Context, const RowHandle* Rows)
 				{
-					Context.RemoveColumns<FUpdatedPathTag>(TConstArrayView<DataStorage::RowHandle>(Rows, Context.GetRowCount()));
+					Context.RemoveColumns<FUpdatedPathTag>(TConstArrayView<RowHandle>(Rows, Context.GetRowCount()));
 				}
 			)
 			.Where()
@@ -208,10 +208,10 @@ FTedsAssetData::FTedsAssetData(ITypedElementDataStorageInterface& InDatabase)
 	RemoveUpdatedAssetDataTagQuery = Database.RegisterQuery(
 			TypedElementQueryBuilder::Select(
 				TEXT("FTedsAssetData: Remove Updated Asset Data Tag"),
-				TypedElementQueryBuilder::FPhaseAmble(TypedElementQueryBuilder::FPhaseAmble::ELocation::Postamble, TypedElementDataStorage::EQueryTickPhase::FrameEnd),
-				[](DataStorage::IQueryContext& Context, const DataStorage::RowHandle* Rows)
+				TypedElementQueryBuilder::FPhaseAmble(TypedElementQueryBuilder::FPhaseAmble::ELocation::Postamble, EQueryTickPhase::FrameEnd),
+				[](IQueryContext& Context, const RowHandle* Rows)
 				{
-					Context.RemoveColumns<FUpdatedAssetDataTag>(TConstArrayView<DataStorage::RowHandle>(Rows, Context.GetRowCount()));
+					Context.RemoveColumns<FUpdatedAssetDataTag>(TConstArrayView<RowHandle>(Rows, Context.GetRowCount()));
 				}
 			)
 			.Where()
@@ -228,14 +228,14 @@ FTedsAssetData::FTedsAssetData(ITypedElementDataStorageInterface& InDatabase)
 	ResolveMissingAssetInPathQuery = Database.RegisterQuery(
 		TypedElementQueryBuilder::Select(
 			TEXT("FTedsAssetData: Resolve Missing Asset In Path"),
-			TypedElementQueryBuilder::FProcessor(TypedElementDataStorage::EQueryTickPhase::FrameEnd, Database.GetQueryTickGroupName(TypedElementDataStorage::EQueryTickGroups::Default)),
-			[this](DataStorage::IQueryContext& Context, DataStorage::RowHandle Row, const FUnresolvedAssetsInPathColumn_Experimental& UnresolvedAssetPath)
+			TypedElementQueryBuilder::FProcessor(EQueryTickPhase::FrameEnd, Database.GetQueryTickGroupName(EQueryTickGroups::Default)),
+			[this](IQueryContext& Context, RowHandle Row, const FUnresolvedAssetsInPathColumn_Experimental& UnresolvedAssetPath)
 			{
 #if TRACK_TEDSASSETDATA_MEMORY
 				LLM_SCOPE_BYNAME(TEXT("FTedsAssetData"))
 #endif
 
-				const DataStorage::RowHandle PathRow = Context.FindIndexedRow(UnresolvedAssetPath.Hash);
+				const RowHandle PathRow = Context.FindIndexedRow(UnresolvedAssetPath.Hash);
 				if (Context.IsRowAssigned(PathRow))
 				{
 					Context.RemoveColumns<FUnresolvedAssetsInPathColumn_Experimental>(Row);
@@ -266,14 +266,14 @@ FTedsAssetData::FTedsAssetData(ITypedElementDataStorageInterface& InDatabase)
 	ResolveMissingParentPathQuery = Database.RegisterQuery(
 		TypedElementQueryBuilder::Select(
 			TEXT("FTedsAssetData: Resolve Missing Parent Path Row"),
-			TypedElementQueryBuilder::FProcessor(TypedElementDataStorage::EQueryTickPhase::FrameEnd, Database.GetQueryTickGroupName(TypedElementDataStorage::EQueryTickGroups::Default)),
-			[this](DataStorage::IQueryContext& Context, DataStorage::RowHandle Row, const FUnresolvedParentAssetPathColumn_Experimental& UnresolvedParentAssetPath, FParentAssetPathColumn_Experimental& ParentAssetPathColumn)
+			TypedElementQueryBuilder::FProcessor(EQueryTickPhase::FrameEnd, Database.GetQueryTickGroupName(EQueryTickGroups::Default)),
+			[this](IQueryContext& Context, RowHandle Row, const FUnresolvedParentAssetPathColumn_Experimental& UnresolvedParentAssetPath, FParentAssetPathColumn_Experimental& ParentAssetPathColumn)
 			{
 #if TRACK_TEDSASSETDATA_MEMORY
 				LLM_SCOPE_BYNAME(TEXT("FTedsAssetData"))
 #endif
 
-				const DataStorage::RowHandle ParentPathRow = Context.FindIndexedRow(UnresolvedParentAssetPath.Hash);
+				const RowHandle ParentPathRow = Context.FindIndexedRow(UnresolvedParentAssetPath.Hash);
 				if (Context.IsRowAssigned(ParentPathRow))
 				{
 					Context.RemoveColumns<FUnresolvedParentAssetPathColumn_Experimental>(Row);
@@ -317,27 +317,27 @@ FTedsAssetData::FTedsAssetData(ITypedElementDataStorageInterface& InDatabase)
 		{
 			FName Path = CachedPaths[Index];
 			Path.ToString(NameBuilder);
-			PopulatePathRowArgs[Index] = ThreadSafe_PopulatePathRowArgs(TypedElementDataStorage::GenerateIndexHash(Path), Path, NameBuilder);
+			PopulatePathRowArgs[Index] = ThreadSafe_PopulatePathRowArgs(GenerateIndexHash(Path), Path, NameBuilder);
 		});
 
-	TArray<DataStorage::RowHandle> ReservedRows;
+	TArray<RowHandle> ReservedRows;
 	ReservedRows.AddUninitialized(PopulatePathRowArgs.Num() + AssetsData.Num());
 	Database.BatchReserveRows(ReservedRows);
 
-	TArray<TPair<TypedElementDataStorage::IndexHash, DataStorage::RowHandle>> IndexesToReservedRows;
+	TArray<TPair<IndexHash, RowHandle>> IndexesToReservedRows;
 	IndexesToReservedRows.AddUninitialized(ReservedRows.Num());
 
 	// Index Reserved Path Rows
-	TConstArrayView<DataStorage::RowHandle> ReservedPopulatePathRows(ReservedRows.GetData(), PopulatePathRowArgs.Num());
-	TArrayView<TPair<TypedElementDataStorage::IndexHash, DataStorage::RowHandle>> IndexesToReservedPathRows(IndexesToReservedRows.GetData(), PopulatePathRowArgs.Num());
+	TConstArrayView<RowHandle> ReservedPopulatePathRows(ReservedRows.GetData(), PopulatePathRowArgs.Num());
+	TArrayView<TPair<IndexHash, RowHandle>> IndexesToReservedPathRows(IndexesToReservedRows.GetData(), PopulatePathRowArgs.Num());
 	ParallelFor(TEXT("Populating TEDS Asset Registry Path Data Indexes"), ReservedPopulatePathRows.Num(), ParallelForMinBatchSize, [&IndexesToReservedPathRows, &ReservedPopulatePathRows, &PopulatePathRowArgs](int32 Index)
 		{
-			IndexesToReservedPathRows[Index] = TPair<TypedElementDataStorage::IndexHash, DataStorage::RowHandle>(PopulatePathRowArgs[Index].AssetRegistryPathHash, ReservedPopulatePathRows[Index]);
+			IndexesToReservedPathRows[Index] = TPair<IndexHash, RowHandle>(PopulatePathRowArgs[Index].AssetRegistryPathHash, ReservedPopulatePathRows[Index]);
 		});
 	Database.BatchIndexRows(IndexesToReservedPathRows);
 
 	// Populate Path Rows
-	Database.BatchAddRow(PathsTable, ReservedPopulatePathRows, [PathRowArgs = MoveTemp(PopulatePathRowArgs), Index = 0, this](DataStorage::RowHandle InRowHandle) mutable
+	Database.BatchAddRow(PathsTable, ReservedPopulatePathRows, [PathRowArgs = MoveTemp(PopulatePathRowArgs), Index = 0, this](RowHandle InRowHandle) mutable
 		{
 			PopulatePathDataTableRow(MoveTemp(PathRowArgs[Index]), Database, InRowHandle);
 			++Index;
@@ -352,16 +352,16 @@ FTedsAssetData::FTedsAssetData(ITypedElementDataStorageInterface& InDatabase)
 		});
 
 	// Index Reserved Asset Data Rows
-	TConstArrayView<DataStorage::RowHandle> ReservedAssetDataRows(ReservedRows.GetData() + ReservedPopulatePathRows.Num(), AssetsData.Num());
-	TArrayView<TPair<TypedElementDataStorage::IndexHash, DataStorage::RowHandle>>  IndexesToResevedAssetRows(IndexesToReservedRows.GetData() + ReservedPopulatePathRows.Num(), ReservedAssetDataRows.Num());
+	TConstArrayView<RowHandle> ReservedAssetDataRows(ReservedRows.GetData() + ReservedPopulatePathRows.Num(), AssetsData.Num());
+	TArrayView<TPair<IndexHash, RowHandle>>  IndexesToResevedAssetRows(IndexesToReservedRows.GetData() + ReservedPopulatePathRows.Num(), ReservedAssetDataRows.Num());
 	ParallelFor(TEXT("Populating TEDS Asset Registry Asset Data Indexes"), PopulateAssetDataRowArgs.Num(), ParallelForMinBatchSize, [&PopulateAssetDataRowArgs, &ReservedAssetDataRows, &IndexesToResevedAssetRows](int32 Index)
 		{
-			IndexesToResevedAssetRows[Index] = TPair<TypedElementDataStorage::IndexHash, DataStorage::RowHandle>(PopulateAssetDataRowArgs[Index].ObjectPathHash, ReservedAssetDataRows[Index]);
+			IndexesToResevedAssetRows[Index] = TPair<IndexHash, RowHandle>(PopulateAssetDataRowArgs[Index].ObjectPathHash, ReservedAssetDataRows[Index]);
 		});
 	Database.BatchIndexRows(IndexesToResevedAssetRows);
 
 	// Populate Asset Rows
-	Database.BatchAddRow(AssetsDataTable, ReservedAssetDataRows,  [AssetDataRowArgs = MoveTemp(PopulateAssetDataRowArgs), Index = 0, this](DataStorage::RowHandle InRowHandle) mutable
+	Database.BatchAddRow(AssetsDataTable, ReservedAssetDataRows,  [AssetDataRowArgs = MoveTemp(PopulateAssetDataRowArgs), Index = 0, this](RowHandle InRowHandle) mutable
 		{
 			PopulateAssetDataTableRow(MoveTemp(AssetDataRowArgs[Index]), Database, InRowHandle);
 			++Index;
@@ -399,8 +399,8 @@ FTedsAssetData::~FTedsAssetData()
 	
 			AssetRegistry->EnumerateAllCachedPaths([this](FName InPath)
 				{
-					const TypedElementDataStorage::IndexHash PathHash = TypedElementDataStorage::GenerateIndexHash(InPath);
-					const DataStorage::RowHandle Row = Database.FindIndexedRow(PathHash);
+					const IndexHash PathHash = GenerateIndexHash(InPath);
+					const RowHandle Row = Database.FindIndexedRow(PathHash);
 					Database.RemoveRow(Row);
 					Database.RemoveIndex(PathHash);
 					return true;
@@ -408,8 +408,8 @@ FTedsAssetData::~FTedsAssetData()
 
 			AssetRegistry->EnumerateAllAssets([this](const FAssetData& InAssetData)
 				{
-					const TypedElementDataStorage::IndexHash AssetPathHash = TypedElementDataStorage::GenerateIndexHash(InAssetData.GetSoftObjectPath());
-					const DataStorage::RowHandle Row = Database.FindIndexedRow(AssetPathHash);
+					const IndexHash AssetPathHash = GenerateIndexHash(InAssetData.GetSoftObjectPath());
+					const RowHandle Row = Database.FindIndexedRow(AssetPathHash);
 					Database.RemoveRow(Row);
 					Database.RemoveIndex(AssetPathHash);
 					return true;
@@ -466,9 +466,9 @@ void FTedsAssetData::OnAssetsAdded(TConstArrayView<FAssetData> InAssetsAdded)
 	int32 Index = 0;
 	if (NewRowsCount > 0)
 	{
-		TArray<TPair<TypedElementDataStorage::IndexHash, DataStorage::RowHandle>> IndexToRow;
+		TArray<TPair<IndexHash, RowHandle>> IndexToRow;
 		IndexToRow.Reserve(NewRowsCount);
-		Database.BatchAddRow(AssetsDataTable, NewRowsCount, [RowArgs = MoveTemp(PopulateRowArgs), Index = 0, &IndexToRow, this](DataStorage::RowHandle InRowHandle) mutable
+		Database.BatchAddRow(AssetsDataTable, NewRowsCount, [RowArgs = MoveTemp(PopulateRowArgs), Index = 0, &IndexToRow, this](RowHandle InRowHandle) mutable
 			{
 				FPopulateAssetDataRowArgs ARowArgs = MoveTemp(RowArgs[Index]);
 				while (!ARowArgs.AssetData.IsValid())
@@ -496,14 +496,14 @@ void FTedsAssetData::OnAssetsRemoved(TConstArrayView<FAssetData> InAssetsRemoved
 
 	for (const FAssetData& Asset : InAssetsRemoved)
 	{
-		const TypedElementDataStorage::IndexHash AssetHash = TypedElementDataStorage::GenerateIndexHash(Asset.GetSoftObjectPath());
-		const DataStorage::RowHandle AssetRow = Database.FindIndexedRow(AssetHash);
+		const IndexHash AssetHash = GenerateIndexHash(Asset.GetSoftObjectPath());
+		const RowHandle AssetRow = Database.FindIndexedRow(AssetHash);
 		if (Database.IsRowAssigned(AssetRow))
 		{
 			if (const FAssetDataColumn_Experimental* AssetDataColunm = Database.GetColumn<FAssetDataColumn_Experimental>(AssetRow))
 			{
-				const TypedElementDataStorage::IndexHash FolderPathHash = TypedElementDataStorage::GenerateIndexHash(AssetDataColunm->AssetData.PackagePath);
-				const DataStorage::RowHandle FolderRow = Database.FindIndexedRow(FolderPathHash);
+				const IndexHash FolderPathHash = GenerateIndexHash(AssetDataColunm->AssetData.PackagePath);
+				const RowHandle FolderRow = Database.FindIndexedRow(FolderPathHash);
 				if (FAssetsInPathColumn_Experimental* AssetInFolder = Database.GetColumn<FAssetsInPathColumn_Experimental>(FolderRow))
 				{
 					AssetInFolder->AssetsRow.Remove(AssetRow);
@@ -525,8 +525,8 @@ void FTedsAssetData::OnAssetsUpdated(TConstArrayView<FAssetData> InAssetsUpdated
 
 	for (const FAssetData& Asset : InAssetsUpdated)
 	{
-		const TypedElementDataStorage::IndexHash AssetHash = TypedElementDataStorage::GenerateIndexHash(Asset.GetSoftObjectPath());
-		const DataStorage::RowHandle Row = Database.FindIndexedRow(AssetHash);
+		const IndexHash AssetHash = GenerateIndexHash(Asset.GetSoftObjectPath());
+		const RowHandle Row = Database.FindIndexedRow(AssetHash);
 		if (Database.IsRowAssigned(Row))
 		{
 			Database.GetColumn<FAssetDataColumn_Experimental>(Row)->AssetData = Asset;
@@ -545,8 +545,8 @@ void FTedsAssetData::OnAssetsUpdatedOnDisk(TConstArrayView<FAssetData> InAssetsU
 
 	for (const FAssetData& Asset : InAssetsUpdated)
 	{
-		const TypedElementDataStorage::IndexHash AssetHash = TypedElementDataStorage::GenerateIndexHash(Asset.GetSoftObjectPath());
-		const DataStorage::RowHandle Row = Database.FindIndexedRow(AssetHash);
+		const IndexHash AssetHash = GenerateIndexHash(Asset.GetSoftObjectPath());
+		const RowHandle Row = Database.FindIndexedRow(AssetHash);
 		if (Database.IsRowAssigned(Row))
 		{
 			Database.GetColumn<FAssetDataColumn_Experimental>(Row)->AssetData = Asset;
@@ -563,15 +563,15 @@ void FTedsAssetData::OnAssetRenamed(const FAssetData& InAsset, const FString& In
 	LLM_SCOPE_BYNAME(TEXT("FTedsAssetData"))
 #endif
 
-	const TypedElementDataStorage::IndexHash NewAssetHash = TypedElementDataStorage::GenerateIndexHash(InAsset.GetSoftObjectPath());
-	const TypedElementDataStorage::IndexHash OldAssetHash = TypedElementDataStorage::GenerateIndexHash(FSoftObjectPath(InOldObjectPath));
-	const DataStorage::RowHandle Row = Database.FindIndexedRow(OldAssetHash);
+	const IndexHash NewAssetHash = GenerateIndexHash(InAsset.GetSoftObjectPath());
+	const IndexHash OldAssetHash = GenerateIndexHash(FSoftObjectPath(InOldObjectPath));
+	const RowHandle Row = Database.FindIndexedRow(OldAssetHash);
 	if (Database.IsRowAssigned(Row))
 	{
 		Database.GetColumn<FAssetDataColumn_Experimental>(Row)->AssetData = InAsset;
 
 		// Update the asset in folder columns
-		const TypedElementDataStorage::IndexHash NewFolderHash = TypedElementDataStorage::GenerateIndexHash(InAsset.PackagePath);
+		const IndexHash NewFolderHash = GenerateIndexHash(InAsset.PackagePath);
 		FStringView OldPackagePath(InOldObjectPath);
 
 		{
@@ -580,12 +580,12 @@ void FTedsAssetData::OnAssetRenamed(const FAssetData& InAsset, const FString& In
 			OldPackagePath.LeftInline(CharacterIndex);
 		}
 
-		const TypedElementDataStorage::IndexHash OldFolderHash = TypedElementDataStorage::GenerateIndexHash(FName(OldPackagePath));
+		const IndexHash OldFolderHash = GenerateIndexHash(FName(OldPackagePath));
 
 		if (NewFolderHash != OldFolderHash)
 		{
-			const DataStorage::RowHandle NewPathRow = Database.FindIndexedRow(NewFolderHash);
-			if (NewPathRow !=  DataStorage::InvalidRowHandle)
+			const RowHandle NewPathRow = Database.FindIndexedRow(NewFolderHash);
+			if (NewPathRow !=  InvalidRowHandle)
 			{
 				Database.GetColumn<FAssetsInPathColumn_Experimental>(NewPathRow)->AssetsRow.Add(Row);
 			}
@@ -596,8 +596,8 @@ void FTedsAssetData::OnAssetRenamed(const FAssetData& InAsset, const FString& In
 				Database.AddColumn<FUnresolvedAssetsInPathColumn_Experimental>(Row, MoveTemp(UnresolvedAssetsInPathColumn));
 			}
 
-			const DataStorage::RowHandle OldPathRow = Database.FindIndexedRow(OldFolderHash);
-			if (NewPathRow !=  DataStorage::InvalidRowHandle)
+			const RowHandle OldPathRow = Database.FindIndexedRow(OldFolderHash);
+			if (NewPathRow !=  InvalidRowHandle)
 			{
 				Database.GetColumn<FAssetsInPathColumn_Experimental>(NewPathRow)->AssetsRow.Remove(Row);
 			}
@@ -628,10 +628,10 @@ void FTedsAssetData::OnPathsAdded(TConstArrayView<FStringView> InPathsAdded)
 		{
 			FStringView Path = InPathsAdded[Index];
 			FName PathName(Path);
-			const TypedElementDataStorage::IndexHash AssetRegistryPathHash = TypedElementDataStorage::GenerateIndexHash(PathName);
+			const IndexHash AssetRegistryPathHash = GenerateIndexHash(PathName);
 			FPopulatePathRowArgs RowArgs;
 
-			if (Database.FindIndexedRow(AssetRegistryPathHash) != DataStorage::InvalidRowHandle)
+			if (Database.FindIndexedRow(AssetRegistryPathHash) != InvalidRowHandle)
 			{
 				RowArgs.MarkAsInvalid();
 			}
@@ -653,11 +653,11 @@ void FTedsAssetData::OnPathsAdded(TConstArrayView<FStringView> InPathsAdded)
 
 	if (NewRowsCount > 0)
 	{
-		TArray<DataStorage::RowHandle> ReservedRow;
+		TArray<RowHandle> ReservedRow;
 		ReservedRow.AddUninitialized(NewRowsCount);
 		Database.BatchReserveRows(ReservedRow);
 
-		TArray<TPair<TypedElementDataStorage::IndexHash, DataStorage::RowHandle>> IndexesAndRows;
+		TArray<TPair<IndexHash, RowHandle>> IndexesAndRows;
 		IndexesAndRows.Reserve(NewRowsCount);
 
 		int32 RowCount = 0;
@@ -677,7 +677,7 @@ void FTedsAssetData::OnPathsAdded(TConstArrayView<FStringView> InPathsAdded)
 		Database.BatchIndexRows(IndexesAndRows);
 
 		int32 Index = 0;
-		Database.BatchAddRow(PathsTable, ReservedRow, [&PopulateRowArgs, &Index, this](DataStorage::RowHandle InRowHandle)
+		Database.BatchAddRow(PathsTable, ReservedRow, [&PopulateRowArgs, &Index, this](RowHandle InRowHandle)
 			{
 				FPopulatePathRowArgs RowArgs = MoveTemp(PopulateRowArgs[Index]);
 				while (!RowArgs)
@@ -702,12 +702,11 @@ void FTedsAssetData::OnPathsRemoved(TConstArrayView<FStringView> InPathsRemoved)
 
 	for (const FStringView Path : InPathsRemoved)
 	{
-		const TypedElementDataStorage::IndexHash PathHash = TypedElementDataStorage::GenerateIndexHash(FName(Path));
+		const IndexHash PathHash = GenerateIndexHash(FName(Path));
 		Database.RemoveRow(Database.FindIndexedRow(PathHash));
 		Database.RemoveIndex(PathHash);
 	}
 }
-} // namespace AssetData::Private
-} // namespace UE::Editor
+} // namespace UE::Editor::AssetData::Private
 
 #undef TRACK_TEDSASSETDATA_MEMORY
\ No newline at end of file
diff --git a/Engine/Plugins/Experimental/EditorDataStorageFeatures/Source/TedsAssetData/Private/Tests/TedsAssetDataTest.cpp b/Engine/Plugins/Experimental/EditorDataStorageFeatures/Source/TedsAssetData/Private/Tests/TedsAssetDataTest.cpp
--- a/Engine/Plugins/Experimental/EditorDataStorageFeatures/Source/TedsAssetData/Private/Tests/TedsAssetDataTest.cpp
+++ b/Engine/Plugins/Experimental/EditorDataStorageFeatures/Source/TedsAssetData/Private/Tests/TedsAssetDataTest.cpp
@@ -16,11 +16,9 @@
 #include "TedsAssetDataModule.h"
 #include "UObject/NameTypes.h"
 
-namespace UE::Editor
+namespace UE::Editor::AssetData::Tests
 {
-namespace AssetData::Tests
-{
-
+	using namespace DataStorage;
 IMPLEMENT_SIMPLE_AUTOMATION_TEST(FTedsAssetDataTest, "Editor.DataStorage.AssetRegistry.ValidateState", EAutomationTestFlags::EditorContext | EAutomationTestFlags::EngineFilter | EAutomationTestFlags::HighPriority)
 
 bool FTedsAssetDataTest::RunTest(const FString& Parameters)
@@ -71,9 +69,9 @@ bool FTedsAssetDataTest::RunTest(const FString& Parameters)
 	
 	for (const TPair<FName, int32>& PathAndAssetCount : AssetRegistryPathsAndAssetCount)
 	{
-		DataStorage::RowHandle Row = Database->FindIndexedRow(TypedElementDataStorage::GenerateIndexHash(PathAndAssetCount.Key));
+		RowHandle Row = Database->FindIndexedRow(GenerateIndexHash(PathAndAssetCount.Key));
 		bool bHadError = false;
-		bHadError |= !TestTrue(TEXT("Asset registry folder/path is indexed in TEDS"), Row != DataStorage::InvalidRowHandle);
+		bHadError |= !TestTrue(TEXT("Asset registry folder/path is indexed in TEDS"), Row != InvalidRowHandle);
 
 		FName NameStoredInTeds;
 		if (const FAssetPathColumn_Experimental* AssetPathColumn = Database->GetColumn<FAssetPathColumn_Experimental>(Row))
@@ -109,17 +107,17 @@ bool FTedsAssetDataTest::RunTest(const FString& Parameters)
 		// Verse package don't emit asset added event so we will ignore those for the test to avoid false errors
 		if (!FPackageName::IsVersePackage(Builder))
 		{ 
-			DataStorage::RowHandle Row = Database->FindIndexedRow(TypedElementDataStorage::GenerateIndexHash(AssetData.GetSoftObjectPath()));
+			RowHandle Row = Database->FindIndexedRow(GenerateIndexHash(AssetData.GetSoftObjectPath()));
 
 			bool bHadError = false;
 
-			bHadError |= !TestTrue(TEXT("Asset registry asset path is indexed in TEDS"), Row != DataStorage::InvalidRowHandle);
+			bHadError |= !TestTrue(TEXT("Asset registry asset path is indexed in TEDS"), Row != InvalidRowHandle);
 
 			bHadError |= !TestNotNull(TEXT("Teds doesn't have a asset column for an asset of the asset registry"), Database->GetColumn<FAssetDataColumn_Experimental>(Row));
 
 			if (!Database->GetColumn<FUnresolvedAssetsInPathColumn_Experimental>(Row))
 			{
-				DataStorage::RowHandle PathRow = Database->FindIndexedRow(TypedElementDataStorage::GenerateIndexHash(AssetData.PackagePath));
+				RowHandle PathRow = Database->FindIndexedRow(GenerateIndexHash(AssetData.PackagePath));
 			
 				const FAssetsInPathColumn_Experimental* AssetsInPathColumn = Database->GetColumn<FAssetsInPathColumn_Experimental>(PathRow);
 
@@ -152,5 +150,4 @@ bool FTedsAssetDataTest::RunTest(const FString& Parameters)
 	return true;
 }
 
-} // namespace AssetData::Tests
-} // namespace UE::Editor
+} // namespace UE::Editor::AssetData::Tests
diff --git a/Engine/Plugins/Experimental/EditorDataStorageFeatures/Source/TedsAssetData/Public/TedsAssetDataColumns.h b/Engine/Plugins/Experimental/EditorDataStorageFeatures/Source/TedsAssetData/Public/TedsAssetDataColumns.h
--- a/Engine/Plugins/Experimental/EditorDataStorageFeatures/Source/TedsAssetData/Public/TedsAssetDataColumns.h
+++ b/Engine/Plugins/Experimental/EditorDataStorageFeatures/Source/TedsAssetData/Public/TedsAssetDataColumns.h
@@ -56,7 +56,7 @@ struct FUnresolvedParentAssetPathColumn_Experimental : public FEditorDataStorage
 {
 	GENERATED_BODY()
 	
-	TypedElementDataStorage::IndexHash Hash;
+	UE::Editor::DataStorage::IndexHash Hash;
 };
 
 USTRUCT()
@@ -64,7 +64,7 @@ struct FUnresolvedAssetsInPathColumn_Experimental : public FEditorDataStorageCol
 {
 	GENERATED_BODY()
 	
-	TypedElementDataStorage::IndexHash Hash;
+	UE::Editor::DataStorage::IndexHash Hash;
 };
 
 // Tag to identify assets
diff --git a/Engine/Plugins/Experimental/EditorDataStorageFeatures/Source/TedsContentBrowser/Private/TedsContentBrowserModule.cpp b/Engine/Plugins/Experimental/EditorDataStorageFeatures/Source/TedsContentBrowser/Private/TedsContentBrowserModule.cpp
--- a/Engine/Plugins/Experimental/EditorDataStorageFeatures/Source/TedsContentBrowser/Private/TedsContentBrowserModule.cpp
+++ b/Engine/Plugins/Experimental/EditorDataStorageFeatures/Source/TedsContentBrowser/Private/TedsContentBrowserModule.cpp
@@ -17,8 +17,10 @@
 
 #define LOCTEXT_NAMESPACE "TedsContentBrowserModule"
 
-namespace UE::Editor::DataStorage::ContentBrowser
+namespace UE::Editor::ContentBrowser
 {
+	using namespace DataStorage;
+
 	static bool bEnableTedsContentBrowser = false;
 	
 	static FAutoConsoleVariableRef CVarUseTEDSOutliner(
@@ -91,12 +93,12 @@ namespace UE::Editor::DataStorage::ContentBrowser
 		
 		if (Item->GetItem().Legacy_TryGetAssetData(ItemAssetData))
 		{
-			TypedElementDataStorage::IndexHash IndexHash = TypedElementDataStorage::GenerateIndexHash(ItemAssetData.GetSoftObjectPath());
+			IndexHash IndexHash = GenerateIndexHash(ItemAssetData.GetSoftObjectPath());
 			RowHandle = DataStorage->FindIndexedRow(IndexHash);
 		}
 		else if(Item->GetItem().Legacy_TryGetPackagePath(PackagePath))
 		{
-			TypedElementDataStorage::IndexHash IndexHash = TypedElementDataStorage::GenerateIndexHash(PackagePath);
+			IndexHash IndexHash = GenerateIndexHash(PackagePath);
 			RowHandle = DataStorage->FindIndexedRow(IndexHash);
 		}
 		
@@ -128,7 +130,7 @@ namespace UE::Editor::DataStorage::ContentBrowser
 		RowQueryStack = MakeShared<FQueryStackNode_RowView>(&Rows);
 
 		// Create the table viewer widget
-		TableViewer = SNew(UE::Editor::DataStorage::STedsTableViewer)
+		TableViewer = SNew(STedsTableViewer)
 					.QueryStack(RowQueryStack)
 					.CellWidgetPurposes({TEXT("General.RowLabel"), TEXT("General.Cell")})
 					// Default list of columns to display
@@ -262,8 +264,8 @@ namespace UE::Editor::DataStorage::ContentBrowser
 		// CB 2.0 TODO: Implement using a Teds column
 		return false;
 	}
-}
+} // namespace UE::Editor::ContentBrowser
 
-IMPLEMENT_MODULE(UE::Editor::DataStorage::ContentBrowser::FTedsContentBrowserModule, TedsContentBrowser);
+IMPLEMENT_MODULE(UE::Editor::ContentBrowser::FTedsContentBrowserModule, TedsContentBrowser);
 
 #undef LOCTEXT_NAMESPACE
\ No newline at end of file
diff --git a/Engine/Plugins/Experimental/EditorDataStorageFeatures/Source/TedsContentBrowser/Private/TedsContentBrowserModule.h b/Engine/Plugins/Experimental/EditorDataStorageFeatures/Source/TedsContentBrowser/Private/TedsContentBrowserModule.h
--- a/Engine/Plugins/Experimental/EditorDataStorageFeatures/Source/TedsContentBrowser/Private/TedsContentBrowserModule.h
+++ b/Engine/Plugins/Experimental/EditorDataStorageFeatures/Source/TedsContentBrowser/Private/TedsContentBrowserModule.h
@@ -17,7 +17,7 @@ namespace UE::Editor::DataStorage
 
 class FAssetViewItem;
 
-namespace UE::Editor::DataStorage::ContentBrowser
+namespace UE::Editor::ContentBrowser
 {
 	// A custom view for the content browser that uses the TEDS asset registry integration to display rows with widgets using TEDS UI
 	class FTedsContentBrowserViewExtender : public IContentBrowserViewExtender
@@ -51,7 +51,7 @@ namespace UE::Editor::DataStorage::ContentBrowser
 		void AddRow(const TSharedPtr<FAssetViewItem>& Item);
 
 		// Get the internal FAssetViewItem from a row handle
-		TSharedPtr<FAssetViewItem> GetAssetViewItemFromRow(UE::Editor::DataStorage::RowHandle Row);
+		TSharedPtr<FAssetViewItem> GetAssetViewItemFromRow(DataStorage::RowHandle Row);
 
 
 	private:
@@ -60,16 +60,16 @@ namespace UE::Editor::DataStorage::ContentBrowser
 		ITypedElementDataStorageInterface* DataStorage;
 		
 		// The actual table viewer widget
-		TSharedPtr<STedsTableViewer> TableViewer;
+		TSharedPtr<DataStorage::STedsTableViewer> TableViewer;
 		
 		// Query stack used by the table viewer
-		TSharedPtr<FQueryStackNode_RowView> RowQueryStack;
+		TSharedPtr<DataStorage::FQueryStackNode_RowView> RowQueryStack;
 
 		// The row handles of the items currently in the list
-		TArray<UE::Editor::DataStorage::RowHandle> Rows;
+		TArray<DataStorage::RowHandle> Rows;
 
 		// A map from row handle -> FAssetView item for lookups
-		TMap<UE::Editor::DataStorage::RowHandle, TWeakPtr<FAssetViewItem>> ContentBrowserItemMap;
+		TMap<DataStorage::RowHandle, TWeakPtr<FAssetViewItem>> ContentBrowserItemMap;
 
 		// Delegates fired when specific events happen on the list
 		FOnSelectionChanged OnSelectionChangedDelegate;
@@ -97,4 +97,4 @@ namespace UE::Editor::DataStorage::ContentBrowser
 	private:
 		TSharedPtr<FTedsContentBrowserViewExtender> ViewExtender;
 	};
-}
+} // namespace UE::Editor::ContentBrowser
diff --git a/Engine/Plugins/Experimental/EditorDataStorageFeatures/Source/TedsDebugger/Private/QueryEditor/TedsQueryEditorModel.cpp b/Engine/Plugins/Experimental/EditorDataStorageFeatures/Source/TedsDebugger/Private/QueryEditor/TedsQueryEditorModel.cpp
--- a/Engine/Plugins/Experimental/EditorDataStorageFeatures/Source/TedsDebugger/Private/QueryEditor/TedsQueryEditorModel.cpp
+++ b/Engine/Plugins/Experimental/EditorDataStorageFeatures/Source/TedsDebugger/Private/QueryEditor/TedsQueryEditorModel.cpp
@@ -33,8 +33,8 @@ namespace UE::Editor::DataStorage::Debug::QueryEditor
 		TArray<const UScriptStruct*> Columns;
 		TArray<const UScriptStruct*> Tags;
 
-		UScriptStruct* ColumnType = UE::Editor::DataStorage::FColumn::StaticStruct();
-		UScriptStruct* TagType = UE::Editor::DataStorage::FTag::StaticStruct();
+		UScriptStruct* ColumnType = FColumn::StaticStruct();
+		UScriptStruct* TagType = FTag::StaticStruct();
 		// Not sure if there is a faster way to do this.  Would be nice to iterate only the derived classes
 		for(TObjectIterator< UScriptStruct > It; It; ++It)
 		{
@@ -136,7 +136,7 @@ namespace UE::Editor::DataStorage::Debug::QueryEditor
 				if (Entry.OperatorType == EOperatorType::Select)
 				{
 					Description.SelectionMetaData.Emplace(FColumnMetaData());
-					Description.SelectionAccessTypes.Add(TypedElementDataStorage::EQueryAccessType::ReadOnly);
+					Description.SelectionAccessTypes.Add(EQueryAccessType::ReadOnly);
 					Description.SelectionTypes.Add(Target);
 				}
 				else
@@ -260,7 +260,7 @@ namespace UE::Editor::DataStorage::Debug::QueryEditor
 			}
 		}
 
-		const UScriptStruct* TagType = UE::Editor::DataStorage::FTag::StaticStruct();
+		const UScriptStruct* TagType = FTag::StaticStruct();
 		for (FConditionEntryInternal& Entry : Conditions)
 		{
 			if (Entry.OperatorType == EOperatorType::Unset)
@@ -389,4 +389,4 @@ namespace UE::Editor::DataStorage::Debug::QueryEditor
 	{
 		return Conditions.FindByPredicate([Handle](const FConditionEntryInternal& Entry){ return Entry.Id == Handle.Id; });
 	}
-}
+} // namespace UE::Editor::DataStorage::Debug::QueryEditor
diff --git a/Engine/Plugins/Experimental/EditorDataStorageFeatures/Source/TedsDebugger/Private/Widgets/QueryEditor/TedsQueryEditorResultsView.cpp b/Engine/Plugins/Experimental/EditorDataStorageFeatures/Source/TedsDebugger/Private/Widgets/QueryEditor/TedsQueryEditorResultsView.cpp
--- a/Engine/Plugins/Experimental/EditorDataStorageFeatures/Source/TedsDebugger/Private/Widgets/QueryEditor/TedsQueryEditorResultsView.cpp
+++ b/Engine/Plugins/Experimental/EditorDataStorageFeatures/Source/TedsDebugger/Private/Widgets/QueryEditor/TedsQueryEditorResultsView.cpp
@@ -17,184 +17,181 @@
 
 #define LOCTEXT_NAMESPACE "TedsDebuggerModule"
 
-namespace UE::Editor::DataStorage
+namespace UE::Editor::DataStorage::Debug::QueryEditor
 {
-	namespace Debug::QueryEditor
+	SResultsView::~SResultsView()
 	{
-		SResultsView::~SResultsView()
-		{
-			Model->GetModelChangedDelegate().Remove(ModelChangedDelegateHandle);
-		}
+		Model->GetModelChangedDelegate().Remove(ModelChangedDelegateHandle);
+	}
 
-		void SResultsView::Construct(const FArguments& InArgs, FTedsQueryEditorModel& InModel)
-		{
-			Model = &InModel;
-			ModelChangedDelegateHandle = Model->GetModelChangedDelegate().AddRaw(this, &SResultsView::OnModelChanged);
-			// RowQueryHandle = TypedElementDataStorage::InvalidQueryHandle;
+	void SResultsView::Construct(const FArguments& InArgs, FTedsQueryEditorModel& InModel)
+	{
+		Model = &InModel;
+		ModelChangedDelegateHandle = Model->GetModelChangedDelegate().AddRaw(this, &SResultsView::OnModelChanged);
+		// RowQueryHandle = InvalidQueryHandle;
 
-			// Create a custom column for the table viewer to display row handles
-			CreateRowHandleColumn();
+		// Create a custom column for the table viewer to display row handles
+		CreateRowHandleColumn();
 	
-			RowQueryStack = MakeShared<FQueryStackNode_RowView>(&TableViewerRows);
+		RowQueryStack = MakeShared<FQueryStackNode_RowView>(&TableViewerRows);
 
-			ChildSlot
+		ChildSlot
+		[
+			SNew(SVerticalBox)
+			+SVerticalBox::Slot()
 			[
-				SNew(SVerticalBox)
-				+SVerticalBox::Slot()
+				SNew(SSplitter)
+				+SSplitter::Slot()
+				.Value(0.5f)
 				[
-					SNew(SSplitter)
-					+SSplitter::Slot()
-					.Value(0.5f)
-					[
-						SAssignNew(TableViewer, STedsTableViewer)
-						.QueryStack(RowQueryStack)
-						.EmptyRowsMessage(LOCTEXT("EmptyRowsMessage", "The provided query has no results."))
-						.OnSelectionChanged(STedsTableViewer::FOnSelectionChanged::CreateLambda(
-							[this](RowHandle SelectedRow)
+					SAssignNew(TableViewer, STedsTableViewer)
+					.QueryStack(RowQueryStack)
+					.EmptyRowsMessage(LOCTEXT("EmptyRowsMessage", "The provided query has no results."))
+					.OnSelectionChanged(STedsTableViewer::FOnSelectionChanged::CreateLambda(
+						[this](RowHandle SelectedRow)
+							{
+								if(RowDetailsWidget)
 								{
-									if(RowDetailsWidget)
+									if(SelectedRow != InvalidRowHandle)
 									{
-										if(SelectedRow != InvalidRowHandle)
-										{
-											RowDetailsWidget->SetRow(SelectedRow);
-										}
-										else
-										{
-											RowDetailsWidget->ClearRow();
-										}
-								
+										RowDetailsWidget->SetRow(SelectedRow);
 									}
-								}))
-					]
-					+SSplitter::Slot()
-					.Value(0.5f)
-					[
-						SAssignNew(RowDetailsWidget, SRowDetails)
-					]
-			
+									else
+									{
+										RowDetailsWidget->ClearRow();
+									}
+							
+								}
+							}))
 				]
-				+SVerticalBox::Slot()
-				.AutoHeight()
+				+SSplitter::Slot()
+				.Value(0.5f)
 				[
-					SNew(STextBlock)
-					.Text_Lambda([this]()
-					{
-						auto& TedsInterface = Model->GetTedsInterface();
-						TypedElementDataStorage::FQueryResult QueryResult = TedsInterface.RunQuery(CountQueryHandle);
-						if (QueryResult.Completed == TypedElementDataStorage::FQueryResult::ECompletion::Fully)
-						{
-							FString String = FString::Printf(TEXT("Element Count: %u"), QueryResult.Count);
-							return FText::FromString(String);
-						}
-						else
-						{
-							return FText::FromString(TEXT("Invalid query"));
-						}
-					})
+					SAssignNew(RowDetailsWidget, SRowDetails)
 				]
-			];
+		
+			]
+			+SVerticalBox::Slot()
+			.AutoHeight()
+			[
+				SNew(STextBlock)
+				.Text_Lambda([this]()
+				{
+					auto& TedsInterface = Model->GetTedsInterface();
+					FQueryResult QueryResult = TedsInterface.RunQuery(CountQueryHandle);
+					if (QueryResult.Completed == FQueryResult::ECompletion::Fully)
+					{
+						FString String = FString::Printf(TEXT("Element Count: %u"), QueryResult.Count);
+						return FText::FromString(String);
+					}
+					else
+					{
+						return FText::FromString(TEXT("Invalid query"));
+					}
+				})
+			]
+		];
 
-			if(RowHandleColumn)
-			{
-				TableViewer->AddCustomColumn(RowHandleColumn.ToSharedRef());
-			}
+		if(RowHandleColumn)
+		{
+			TableViewer->AddCustomColumn(RowHandleColumn.ToSharedRef());
 		}
+	}
 
-		void SResultsView::Tick(const FGeometry& AllottedGeometry, const double InCurrentTime, const float InDeltaTime)
-		{
-			using namespace TypedElementQueryBuilder;
+	void SResultsView::Tick(const FGeometry& AllottedGeometry, const double InCurrentTime, const float InDeltaTime)
+	{
+		using namespace TypedElementQueryBuilder;
 
-			ITypedElementDataStorageInterface& TedsInterface = Model->GetTedsInterface();
+		ITypedElementDataStorageInterface& TedsInterface = Model->GetTedsInterface();
 
-			if(bModelDirty)
+		if(bModelDirty)
+		{
 			{
+				FQueryDescription CountQueryDescription = Model->GenerateNoSelectQueryDescription();
+	
+				if (CountQueryHandle != InvalidQueryHandle)
 				{
-					FQueryDescription CountQueryDescription = Model->GenerateNoSelectQueryDescription();
-		
-					if (CountQueryHandle != InvalidQueryHandle)
-					{
-						TedsInterface.UnregisterQuery(CountQueryHandle);
-						CountQueryHandle = InvalidQueryHandle;
-					}
-					CountQueryHandle = TedsInterface.RegisterQuery(MoveTemp(CountQueryDescription));
+					TedsInterface.UnregisterQuery(CountQueryHandle);
+					CountQueryHandle = InvalidQueryHandle;
 				}
+				CountQueryHandle = TedsInterface.RegisterQuery(MoveTemp(CountQueryDescription));
+			}
 
-				{
-					FQueryDescription TableViewerQueryDescription = Model->GenerateQueryDescription();
-
-					// Update the columns in the table viewer using the selection types from the query description
-					TableViewer->SetColumns(TArray<TWeakObjectPtr<const UScriptStruct>>(TableViewerQueryDescription.SelectionTypes));
+			{
+				FQueryDescription TableViewerQueryDescription = Model->GenerateQueryDescription();
 
-					// Since setting the columns clears all columns, we are going to re-add the custom column back
-					if(RowHandleColumn)
-					{
-						TableViewer->AddCustomColumn(RowHandleColumn.ToSharedRef());
-					}
-			
-					if (TableViewerQueryHandle != InvalidQueryHandle)
-					{
-						TedsInterface.UnregisterQuery(TableViewerQueryHandle);
-						TableViewerQueryHandle = InvalidQueryHandle;
-					}
+				// Update the columns in the table viewer using the selection types from the query description
+				TableViewer->SetColumns(TArray<TWeakObjectPtr<const UScriptStruct>>(TableViewerQueryDescription.SelectionTypes));
 
-					// Mass doesn't like empty queries, so we only set it if there are actual conditions
-					if(TableViewerQueryDescription.ConditionTypes.Num() || TableViewerQueryDescription.SelectionTypes.Num())
-					{
-						TableViewerQueryHandle = TedsInterface.RegisterQuery(MoveTemp(TableViewerQueryDescription));
-					}
+				// Since setting the columns clears all columns, we are going to re-add the custom column back
+				if(RowHandleColumn)
+				{
+					TableViewer->AddCustomColumn(RowHandleColumn.ToSharedRef());
 				}
 		
-				bModelDirty = false;
-			}
-
-			TSet<RowHandle> NewTableViewerRows_Set;
-	
-			// Every frame we re-run the query to update the rows the table viewer is showing
-			if(TableViewerQueryHandle != InvalidQueryHandle)
-			{
-				NewTableViewerRows_Set.Reserve(TableViewerRows_Set.Num());
+				if (TableViewerQueryHandle != InvalidQueryHandle)
+				{
+					TedsInterface.UnregisterQuery(TableViewerQueryHandle);
+					TableViewerQueryHandle = InvalidQueryHandle;
+				}
 
-				TypedElementDataStorage::FQueryResult QueryResult = Model->GetTedsInterface().RunQuery(TableViewerQueryHandle,
-					CreateDirectQueryCallbackBinding([&NewTableViewerRows_Set](const ITypedElementDataStorageInterface::IDirectQueryContext& Context, const RowHandle*)
+				// Mass doesn't like empty queries, so we only set it if there are actual conditions
+				if(TableViewerQueryDescription.ConditionTypes.Num() || TableViewerQueryDescription.SelectionTypes.Num())
 				{
-					NewTableViewerRows_Set.Append(Context.GetRowHandles());
-				}));
+					TableViewerQueryHandle = TedsInterface.RegisterQuery(MoveTemp(TableViewerQueryDescription));
+				}
 			}
 	
-			// Check if the two sets are equal, i.e not changes and no need to update the table viewer
-			const bool bSetsEqual = (TableViewerRows_Set.Num() == NewTableViewerRows_Set.Num()) && TableViewerRows_Set.Includes(NewTableViewerRows_Set);
+			bModelDirty = false;
+		}
 
-			if(!bSetsEqual)
+		TSet<RowHandle> NewTableViewerRows_Set;
+	
+		// Every frame we re-run the query to update the rows the table viewer is showing
+		if(TableViewerQueryHandle != InvalidQueryHandle)
+		{
+			NewTableViewerRows_Set.Reserve(TableViewerRows_Set.Num());
+
+			FQueryResult QueryResult = Model->GetTedsInterface().RunQuery(TableViewerQueryHandle,
+				CreateDirectQueryCallbackBinding([&NewTableViewerRows_Set](const ITypedElementDataStorageInterface::IDirectQueryContext& Context, const RowHandle*)
 			{
-				Swap(TableViewerRows_Set, NewTableViewerRows_Set);
-				TableViewerRows = TableViewerRows_Set.Array();
-				RowQueryStack->MarkDirty();
-			}
+				NewTableViewerRows_Set.Append(Context.GetRowHandles());
+			}));
 		}
+	
+		// Check if the two sets are equal, i.e not changes and no need to update the table viewer
+		const bool bSetsEqual = (TableViewerRows_Set.Num() == NewTableViewerRows_Set.Num()) && TableViewerRows_Set.Includes(NewTableViewerRows_Set);
 
-		void SResultsView::OnModelChanged()
+		if(!bSetsEqual)
 		{
-			bModelDirty = true;
-			Invalidate(EInvalidateWidgetReason::Layout);
+			Swap(TableViewerRows_Set, NewTableViewerRows_Set);
+			TableViewerRows = TableViewerRows_Set.Array();
+			RowQueryStack->MarkDirty();
 		}
+	}
 
-		void SResultsView::CreateRowHandleColumn()
+	void SResultsView::OnModelChanged()
+	{
+		bModelDirty = true;
+		Invalidate(EInvalidateWidgetReason::Layout);
+	}
+
+	void SResultsView::CreateRowHandleColumn()
+	{
+		auto AssignWidgetToColumn = [this](TUniquePtr<FTypedElementWidgetConstructor> Constructor, TConstArrayView<TWeakObjectPtr<const UScriptStruct>>)
 		{
-			auto AssignWidgetToColumn = [this](TUniquePtr<FTypedElementWidgetConstructor> Constructor, TConstArrayView<TWeakObjectPtr<const UScriptStruct>>)
-			{
-				TSharedPtr<FTypedElementWidgetConstructor> WidgetConstructor(Constructor.Release());
-				RowHandleColumn = MakeShared<FTedsTableViewerColumn>(TEXT("Row Handle"), WidgetConstructor);
-				return false;
-			};
+			TSharedPtr<FTypedElementWidgetConstructor> WidgetConstructor(Constructor.Release());
+			RowHandleColumn = MakeShared<FTedsTableViewerColumn>(TEXT("Row Handle"), WidgetConstructor);
+			return false;
+		};
 	
-			UTypedElementRegistry* Registry = UTypedElementRegistry::GetInstance();
-			checkf(Registry, TEXT("SResultsView created before UTypedElementRegistry is available."));
-			ITypedElementDataStorageUiInterface* StorageUi = Registry->GetMutableDataStorageUi();
-			checkf(StorageUi, TEXT("SResultsView created before data storage interfaces were initialized."))
-
-			StorageUi->CreateWidgetConstructors(TEXT("General.Cell.RowHandle"), FMetaDataView(), AssignWidgetToColumn);
-		}
-	} // namespace Debug::QueryEditor
-} // namespace UE::Editor::DataStorage
+		UTypedElementRegistry* Registry = UTypedElementRegistry::GetInstance();
+		checkf(Registry, TEXT("SResultsView created before UTypedElementRegistry is available."));
+		ITypedElementDataStorageUiInterface* StorageUi = Registry->GetMutableDataStorageUi();
+		checkf(StorageUi, TEXT("SResultsView created before data storage interfaces were initialized."))
+
+		StorageUi->CreateWidgetConstructors(TEXT("General.Cell.RowHandle"), FMetaDataView(), AssignWidgetToColumn);
+	}
+} // namespace UE::Editor::DataStorage::Debug::QueryEditor
 
 #undef LOCTEXT_NAMESPACE
diff --git a/Engine/Plugins/Experimental/EditorDataStorageFeatures/Source/TedsDebugger/Private/Widgets/RowReferenceWidget.cpp b/Engine/Plugins/Experimental/EditorDataStorageFeatures/Source/TedsDebugger/Private/Widgets/RowReferenceWidget.cpp
--- a/Engine/Plugins/Experimental/EditorDataStorageFeatures/Source/TedsDebugger/Private/Widgets/RowReferenceWidget.cpp
+++ b/Engine/Plugins/Experimental/EditorDataStorageFeatures/Source/TedsDebugger/Private/Widgets/RowReferenceWidget.cpp
@@ -19,8 +19,6 @@
 
 namespace UE::Editor::DataStorage::Debug::Private
 {
-	using namespace UE::Editor::DataStorage;
-
 	void OnNavigateHyperlink(const ITypedElementDataStorageInterface* DataStorage,  RowHandle TargetRowHandle, RowHandle UiRowHandle)
 	{
 		const FTedsOutlinerColumn* TedsOutlinerColumn = DataStorage->GetColumn<FTedsOutlinerColumn>(UiRowHandle);
@@ -109,7 +107,6 @@ void URowReferenceWidgetFactory::RegisterWidgetConstructors(ITypedElementDataSto
 void URowReferenceWidgetFactory::RegisterQueries(ITypedElementDataStorageInterface& DataStorage)
 {
 	using namespace TypedElementQueryBuilder;
-	using namespace TypedElementDataStorage;
 	using namespace UE::Editor::DataStorage;
 	
 	const QueryHandle UpdateRowReferenceWidget = DataStorage.RegisterQuery(
diff --git a/Engine/Plugins/Experimental/EditorDataStorageFeatures/Source/TedsOutliner/Private/Compatibility/SceneOutlinerTedsBridge.cpp b/Engine/Plugins/Experimental/EditorDataStorageFeatures/Source/TedsOutliner/Private/Compatibility/SceneOutlinerTedsBridge.cpp
--- a/Engine/Plugins/Experimental/EditorDataStorageFeatures/Source/TedsOutliner/Private/Compatibility/SceneOutlinerTedsBridge.cpp
+++ b/Engine/Plugins/Experimental/EditorDataStorageFeatures/Source/TedsOutliner/Private/Compatibility/SceneOutlinerTedsBridge.cpp
@@ -35,7 +35,6 @@ FAutoConsoleCommand BindColumnsToSceneOutlinerConsoleCommand(
 	FConsoleCommandWithArgsDelegate::CreateLambda([](const TArray<FString>& Args)
 		{
 			using namespace TypedElementQueryBuilder;
-			using namespace TypedElementDataStorage;
 			
 		    const FName WidgetPurposes[] = {TEXT("SceneOutliner.Cell"), TEXT("General.Cell")};
 
diff --git a/Engine/Plugins/Experimental/EditorDataStorageFeatures/Source/TedsOutliner/Private/TedsOutlinerImpl.cpp b/Engine/Plugins/Experimental/EditorDataStorageFeatures/Source/TedsOutliner/Private/TedsOutlinerImpl.cpp
--- a/Engine/Plugins/Experimental/EditorDataStorageFeatures/Source/TedsOutliner/Private/TedsOutlinerImpl.cpp
+++ b/Engine/Plugins/Experimental/EditorDataStorageFeatures/Source/TedsOutliner/Private/TedsOutlinerImpl.cpp
@@ -364,20 +364,20 @@ bool FTedsOutlinerImpl::CanDisplayRow(UE::Editor::DataStorage::RowHandle ItemRow
 void FTedsOutlinerImpl::CreateItemsFromQuery(TArray<FSceneOutlinerTreeItemPtr>& OutItems, ISceneOutlinerMode* InMode) const
 {
 	using namespace TypedElementQueryBuilder;
-	using DSI = ITypedElementDataStorageInterface;
+	using namespace UE::Editor::DataStorage;
 	
-	TArray<UE::Editor::DataStorage::RowHandle> Rows;
+	TArray<RowHandle> Rows;
 	
-	UE::Editor::DataStorage::DirectQueryCallback RowCollector = CreateDirectQueryCallbackBinding(
-		[&Rows](DSI::IDirectQueryContext& Context)
+	DirectQueryCallback RowCollector = CreateDirectQueryCallbackBinding(
+		[&Rows](IDirectQueryContext& Context)
 		{
-			TConstArrayView<UE::Editor::DataStorage::RowHandle> ContextRows = Context.GetRowHandles();
+			TConstArrayView<RowHandle> ContextRows = Context.GetRowHandles();
 			Rows.Append(ContextRows);
 		});
 
 	Storage->RunQuery(RowHandleQuery, RowCollector);
 	
-	for (const UE::Editor::DataStorage::RowHandle& Row : Rows)
+	for (const RowHandle& Row : Rows)
 	{
 		if (!CanDisplayRow(Row))
 		{
@@ -409,7 +409,7 @@ void FTedsOutlinerImpl::CreateChildren(const FSceneOutlinerTreeItemPtr& Item, TA
 	}
 
 	using namespace TypedElementQueryBuilder;
-	using DSI = ITypedElementDataStorageInterface;
+	using namespace UE::Editor::DataStorage;
 	
 	const FTedsOutlinerTreeItem* TedsTreeItem = Item->CastTo<FTedsOutlinerTreeItem>();
 
@@ -419,7 +419,7 @@ void FTedsOutlinerImpl::CreateChildren(const FSceneOutlinerTreeItemPtr& Item, TA
 		return;
 	}
 		
-	UE::Editor::DataStorage::RowHandle ItemRowHandle = TedsTreeItem->GetRowHandle();
+	RowHandle ItemRowHandle = TedsTreeItem->GetRowHandle();
 
 	if(!Storage->IsRowAssigned(ItemRowHandle))
 	{
@@ -427,31 +427,31 @@ void FTedsOutlinerImpl::CreateChildren(const FSceneOutlinerTreeItemPtr& Item, TA
 	}
 
 
-	TSet<UE::Editor::DataStorage::RowHandle> MatchedRowsWithParentColumn;
+	TSet<RowHandle> MatchedRowsWithParentColumn;
 	
 	// Collect all entities that are owned by our entity
-	UE::Editor::DataStorage::DirectQueryCallback ChildRowCollector = CreateDirectQueryCallbackBinding(
-	[&MatchedRowsWithParentColumn] (const DSI::IDirectQueryContext& Context)
+	DirectQueryCallback ChildRowCollector = CreateDirectQueryCallbackBinding(
+	[&MatchedRowsWithParentColumn] (const IDirectQueryContext& Context)
 	{
 		MatchedRowsWithParentColumn.Append(Context.GetRowHandles());
 	});
 
 	Storage->RunQuery(ChildRowHandleQuery, ChildRowCollector);
 
-	TArray<UE::Editor::DataStorage::RowHandle> ChildItems;
+	TArray<RowHandle> ChildItems;
 
 	// Recursively get the children for each entity
-	TFunction<void(UE::Editor::DataStorage::RowHandle)> GetChildrenRecursive = [&ChildItems, &MatchedRowsWithParentColumn, DataStorage = Storage, &GetChildrenRecursive, InHierarchyData = HierarchyData]
-	(UE::Editor::DataStorage::RowHandle EntityRowHandle) -> void
+	TFunction<void(RowHandle)> GetChildrenRecursive = [&ChildItems, &MatchedRowsWithParentColumn, DataStorage = Storage, &GetChildrenRecursive, InHierarchyData = HierarchyData]
+	(RowHandle EntityRowHandle) -> void
 	{
-		for(UE::Editor::DataStorage::RowHandle ChildEntityRowHandle : MatchedRowsWithParentColumn)
+		for(RowHandle ChildEntityRowHandle : MatchedRowsWithParentColumn)
 		{
 			void* ParentColumnData = DataStorage->GetColumnData(ChildEntityRowHandle, InHierarchyData.GetValue().HierarchyColumn);
 
 			if (ensureMsgf(ParentColumnData, TEXT("We should always the a parent column since we only grabbed rows with those ")))
 			{
 				// Get the parent row handle
-				const UE::Editor::DataStorage::RowHandle ParentRowHandle = InHierarchyData.GetValue().GetParent.Execute(ParentColumnData);
+				const RowHandle ParentRowHandle = InHierarchyData.GetValue().GetParent.Execute(ParentColumnData);
 				
 				// Check if this entity is owned by the entity we are looking children for
 				if (ParentRowHandle == EntityRowHandle)
@@ -470,7 +470,7 @@ void FTedsOutlinerImpl::CreateChildren(const FSceneOutlinerTreeItemPtr& Item, TA
 	GetChildrenRecursive(ItemRowHandle);
 
 	// Actually create the items for the child entities 
-	for (UE::Editor::DataStorage::RowHandle ChildItemRowHandle : ChildItems)
+	for (RowHandle ChildItemRowHandle : ChildItems)
 	{
 		if (!CanDisplayRow(ChildItemRowHandle))
 		{
@@ -553,7 +553,6 @@ void FTedsOutlinerImpl::OnItemMoved(UE::Editor::DataStorage::RowHandle ItemRowHa
 void FTedsOutlinerImpl::RecompileQueries()
 {
 	using namespace TypedElementQueryBuilder;
-	using namespace TypedElementDataStorage;
 
 	UnregisterQueries();
 
@@ -713,17 +712,17 @@ void FTedsOutlinerImpl::ClearSelection() const
 	}
 	
 	using namespace TypedElementQueryBuilder;
-	using DSI = ITypedElementDataStorageInterface;
+	using namespace UE::Editor::DataStorage;
 
-	TArray<UE::Editor::DataStorage::RowHandle> RowsToRemoveSelectionColumn;
+	TArray<RowHandle> RowsToRemoveSelectionColumn;
 
 	// Query to remove the selection column from all rows that belong to this selection set
-	UE::Editor::DataStorage::DirectQueryCallback RowCollector = CreateDirectQueryCallbackBinding(
-	[this, &RowsToRemoveSelectionColumn](const DSI::IDirectQueryContext& Context, const UE::Editor::DataStorage::RowHandle* RowHandles)
+	DirectQueryCallback RowCollector = CreateDirectQueryCallbackBinding(
+	[this, &RowsToRemoveSelectionColumn](const IDirectQueryContext& Context, const RowHandle* RowHandles)
 	{
-		const TConstArrayView<UE::Editor::DataStorage::RowHandle> Rows(RowHandles, Context.GetRowCount());
+		const TConstArrayView<RowHandle> Rows(RowHandles, Context.GetRowCount());
 
-		for(const UE::Editor::DataStorage::RowHandle RowHandle : Rows)
+		for(const RowHandle RowHandle : Rows)
 		{
 			if (const FTypedElementSelectionColumn* SelectionColumn = Storage->GetColumn<FTypedElementSelectionColumn>(RowHandle))
 			{
@@ -737,7 +736,7 @@ void FTedsOutlinerImpl::ClearSelection() const
 
 	Storage->RunQuery(SelectedRowsQuery, RowCollector);
 
-	for(const UE::Editor::DataStorage::RowHandle RowHandle : RowsToRemoveSelectionColumn)
+	for(const RowHandle RowHandle : RowsToRemoveSelectionColumn)
 	{
 		Storage->RemoveColumn<FTypedElementSelectionColumn>(RowHandle);
 	}
diff --git a/Engine/Plugins/Experimental/EditorDataStorageFeatures/Source/TedsOutliner/Private/Widgets/AlertWidget.cpp b/Engine/Plugins/Experimental/EditorDataStorageFeatures/Source/TedsOutliner/Private/Widgets/AlertWidget.cpp
--- a/Engine/Plugins/Experimental/EditorDataStorageFeatures/Source/TedsOutliner/Private/Widgets/AlertWidget.cpp
+++ b/Engine/Plugins/Experimental/EditorDataStorageFeatures/Source/TedsOutliner/Private/Widgets/AlertWidget.cpp
@@ -169,7 +169,6 @@ void UAlertWidgetFactory::RegisterQueries(ITypedElementDataStorageInterface& Dat
 void UAlertWidgetFactory::RegisterAlertQueries(ITypedElementDataStorageInterface& DataStorage)
 {
 	using namespace TypedElementQueryBuilder;
-	using namespace TypedElementDataStorage;
 	using namespace UE::Editor::DataStorage;
 	
 	QueryHandle UpdateWidget_OnlyAlert = DataStorage.RegisterQuery(
@@ -252,7 +251,6 @@ void UAlertWidgetFactory::RegisterAlertQueries(ITypedElementDataStorageInterface
 void UAlertWidgetFactory::RegisterAlertHeaderQueries(ITypedElementDataStorageInterface& DataStorage)
 {
 	using namespace TypedElementQueryBuilder;
-	using namespace TypedElementDataStorage;
 	using namespace UE::Editor::DataStorage;
 	
 	QueryHandle AlertCount = DataStorage.RegisterQuery(
diff --git a/Engine/Plugins/Experimental/EditorDataStorageFeatures/Source/TedsOutliner/Private/Widgets/TransformHeadsUpWidget.cpp b/Engine/Plugins/Experimental/EditorDataStorageFeatures/Source/TedsOutliner/Private/Widgets/TransformHeadsUpWidget.cpp
--- a/Engine/Plugins/Experimental/EditorDataStorageFeatures/Source/TedsOutliner/Private/Widgets/TransformHeadsUpWidget.cpp
+++ b/Engine/Plugins/Experimental/EditorDataStorageFeatures/Source/TedsOutliner/Private/Widgets/TransformHeadsUpWidget.cpp
@@ -155,7 +155,7 @@ static void UpdateTransformHeadsUpDisplay(FTypedElementSlateWidgetReferenceColum
 void UTransformHeadsUpWidgetFactory::RegisterQueries(ITypedElementDataStorageInterface& DataStorage)
 {
 	using namespace TypedElementQueryBuilder;
-	using namespace TypedElementDataStorage;
+	using namespace UE::Editor::DataStorage;
 		
 	UE::Editor::DataStorage::QueryHandle UpdateTransformWidget = DataStorage.RegisterQuery(
 		Select()
diff --git a/Engine/Plugins/Experimental/EditorDataStorageFeatures/Source/TedsRevisionControl/Private/Processors/RevisionControlProcessors.cpp b/Engine/Plugins/Experimental/EditorDataStorageFeatures/Source/TedsRevisionControl/Private/Processors/RevisionControlProcessors.cpp
--- a/Engine/Plugins/Experimental/EditorDataStorageFeatures/Source/TedsRevisionControl/Private/Processors/RevisionControlProcessors.cpp
+++ b/Engine/Plugins/Experimental/EditorDataStorageFeatures/Source/TedsRevisionControl/Private/Processors/RevisionControlProcessors.cpp
@@ -238,7 +238,7 @@ void URevisionControlDataStorageFactory::RegisterQueries(ITypedElementDataStorag
 void URevisionControlDataStorageFactory::RegisterFetchUpdates(ITypedElementDataStorageInterface& DataStorage)
 {
 	using namespace TypedElementQueryBuilder;
-	using namespace TypedElementDataStorage;
+	using namespace UE::Editor::DataStorage;
 	
 	FSourceControlFileStatusMonitor& FileStatusMonitor = ISourceControlModule::Get().GetSourceControlFileStatusMonitor();
 
@@ -268,7 +268,6 @@ void URevisionControlDataStorageFactory::RegisterFetchUpdates(ITypedElementDataS
 
 void URevisionControlDataStorageFactory::RegisterApplyOverlays(ITypedElementDataStorageInterface& DataStorage)
 {
-	using namespace TypedElementDataStorage;
 	using namespace TypedElementQueryBuilder;
 	using namespace UE::Editor::DataStorage;
 	using namespace UE::Editor::RevisionControl::Private;
@@ -389,7 +388,6 @@ void URevisionControlDataStorageFactory::RegisterApplyOverlays(ITypedElementData
 void URevisionControlDataStorageFactory::RegisterRemoveOverlays(ITypedElementDataStorageInterface& DataStorage)
 {
 	using namespace TypedElementQueryBuilder;
-	using namespace TypedElementDataStorage;
 	using namespace UE::Editor::DataStorage;
 	
 	if (RemoveOverlays == InvalidQueryHandle)
diff --git a/Engine/Plugins/Experimental/EditorDataStorageFeatures/Source/TedsRevisionControl/Private/Queries/ObjectPackagePathToColumnQueries.cpp b/Engine/Plugins/Experimental/EditorDataStorageFeatures/Source/TedsRevisionControl/Private/Queries/ObjectPackagePathToColumnQueries.cpp
--- a/Engine/Plugins/Experimental/EditorDataStorageFeatures/Source/TedsRevisionControl/Private/Queries/ObjectPackagePathToColumnQueries.cpp
+++ b/Engine/Plugins/Experimental/EditorDataStorageFeatures/Source/TedsRevisionControl/Private/Queries/ObjectPackagePathToColumnQueries.cpp
@@ -46,7 +46,6 @@ namespace UE::Editor::RevisionControl::Private
 void UTypedElementUObjectPackagePathFactory::RegisterQueries(ITypedElementDataStorageInterface& DataStorage)
 {
 	using namespace TypedElementQueryBuilder;
-	using namespace TypedElementDataStorage;
 	using namespace UE::Editor::DataStorage;
 	
 	UE::Editor::RevisionControl::Private::CVarAutoPopulateState->AsVariable()->OnChangedDelegate().AddLambda(
@@ -96,7 +95,6 @@ void UTypedElementUObjectPackagePathFactory::RegisterQueries(ITypedElementDataSt
 void UTypedElementUObjectPackagePathFactory::RegisterTryAddPackageRef(ITypedElementDataStorageInterface& DataStorage)
 {
 	using namespace TypedElementQueryBuilder;
-	using namespace TypedElementDataStorage;
 	using namespace UE::Editor::DataStorage;
 	
 	TryAddPackageRef = DataStorage.RegisterQuery(
@@ -115,7 +113,7 @@ void UTypedElementUObjectPackagePathFactory::RegisterTryAddPackageRef(ITypedElem
 					{
 						FPaths::NormalizeFilename(PackageFilename);
 						FString FullPackageFilename = FPaths::ConvertRelativePathToFull(PackageFilename);
-						TypedElementDataStorage::IndexHash Index = TypedElementDataStorage::GenerateIndexHash(FullPackageFilename);
+						IndexHash Index = GenerateIndexHash(FullPackageFilename);
 						RowHandle PackageRow = Context.FindIndexedRow(Index);
 						if (Context.IsRowAvailable(PackageRow))
 						{
diff --git a/Engine/Plugins/Experimental/EditorDataStorageFeatures/Source/TedsSettings/Private/TedsSettingsManager.cpp b/Engine/Plugins/Experimental/EditorDataStorageFeatures/Source/TedsSettings/Private/TedsSettingsManager.cpp
--- a/Engine/Plugins/Experimental/EditorDataStorageFeatures/Source/TedsSettings/Private/TedsSettingsManager.cpp
+++ b/Engine/Plugins/Experimental/EditorDataStorageFeatures/Source/TedsSettings/Private/TedsSettingsManager.cpp
@@ -234,7 +234,6 @@ void FTedsSettingsManager::UnregisterSettings()
 
 void FTedsSettingsManager::UpdateSettingsCategory(TSharedPtr<ISettingsCategory> SettingsCategory, UE::Editor::DataStorage::RowHandle ContainerRow, const bool bQueryExistingRows)
 {
-	using namespace TypedElementDataStorage;
 	using namespace UE::Editor::DataStorage;
 
 	TRACE_CPUPROFILER_EVENT_SCOPE(TedsSettingsManager.UpdateSettingsCategory);
@@ -276,11 +275,10 @@ void FTedsSettingsManager::UpdateSettingsCategory(TSharedPtr<ISettingsCategory>
 	if (bQueryExistingRows)
 	{
 		using namespace TypedElementQueryBuilder;
-		using DSI = ITypedElementDataStorageInterface;
 
 		DataStorage->RunQuery(SelectAllSettingsQuery, CreateDirectQueryCallbackBinding(
 			[&OldRowHandles, &OldSectionNames, &ContainerName, &CategoryName](
-				DSI::IDirectQueryContext& Context,
+				IDirectQueryContext& Context,
 				const FSettingsContainerReferenceColumn* ContainerColumns,
 				const FSettingsCategoryReferenceColumn* CategoryColumns,
 				const FNameColumn* SectionNameColumns)
diff --git a/Engine/Plugins/Experimental/EditorDataStorageFeatures/Source/TedsSettings/Tests/TedsSettingsTests.cpp b/Engine/Plugins/Experimental/EditorDataStorageFeatures/Source/TedsSettings/Tests/TedsSettingsTests.cpp
--- a/Engine/Plugins/Experimental/EditorDataStorageFeatures/Source/TedsSettings/Tests/TedsSettingsTests.cpp
+++ b/Engine/Plugins/Experimental/EditorDataStorageFeatures/Source/TedsSettings/Tests/TedsSettingsTests.cpp
@@ -33,9 +33,7 @@ namespace UE::Editor::Settings::Tests
 
 	uint32 CountSettingsRowsInDataStorage()
 	{
-		using DSI = ITypedElementDataStorageInterface;
-
-		DSI::FQueryResult Result = DataStorage->RunQuery(CountAllSettingsQuery);
+		DataStorage::FQueryResult Result = DataStorage->RunQuery(CountAllSettingsQuery);
 
 		return Result.Count;
 	}
diff --git a/Engine/Plugins/Experimental/EditorDataStorageFeatures/Source/TedsTableViewer/Private/TedsTableViewerColumn.cpp b/Engine/Plugins/Experimental/EditorDataStorageFeatures/Source/TedsTableViewer/Private/TedsTableViewerColumn.cpp
--- a/Engine/Plugins/Experimental/EditorDataStorageFeatures/Source/TedsTableViewer/Private/TedsTableViewerColumn.cpp
+++ b/Engine/Plugins/Experimental/EditorDataStorageFeatures/Source/TedsTableViewer/Private/TedsTableViewerColumn.cpp
@@ -157,7 +157,6 @@ namespace UE::Editor::DataStorage
 	void FTedsTableViewerColumn::RegisterQueries()
 	{
 		using namespace TypedElementQueryBuilder;
-		using namespace TypedElementDataStorage;
 
 		const TConstArrayView<TWeakObjectPtr<const UScriptStruct>> ColumnTypes = GetMatchedColumns();
 		
@@ -268,7 +267,6 @@ namespace UE::Editor::DataStorage
 		});
 		
 		using namespace TypedElementQueryBuilder;
-		using namespace TypedElementDataStorage;
 
 		// Query to find all widgets that belong to the row handles that need updates
 		DirectQueryCallback RowCollector = CreateDirectQueryCallbackBinding(
diff --git a/Engine/Plugins/Experimental/InstanceDataObjectFixupTool/Source/Private/InstanceDataObjectFixupToolTedsQueries.cpp b/Engine/Plugins/Experimental/InstanceDataObjectFixupTool/Source/Private/InstanceDataObjectFixupToolTedsQueries.cpp
--- a/Engine/Plugins/Experimental/InstanceDataObjectFixupTool/Source/Private/InstanceDataObjectFixupToolTedsQueries.cpp
+++ b/Engine/Plugins/Experimental/InstanceDataObjectFixupTool/Source/Private/InstanceDataObjectFixupToolTedsQueries.cpp
@@ -18,7 +18,6 @@
 void UInstanceDataObjectFixupToolTedsQueryFactory::RegisterQueries(ITypedElementDataStorageInterface& DataStorage)
 {
 	using namespace UE;
-	using namespace TypedElementDataStorage;
 	using namespace TypedElementQueryBuilder;
 	using namespace UE::Editor::DataStorage;
 	
diff --git a/Engine/Source/Runtime/TypedElementFramework/Private/Elements/Framework/TypedElementQueryBuilder.cpp b/Engine/Source/Runtime/TypedElementFramework/Private/Elements/Framework/TypedElementQueryBuilder.cpp
--- a/Engine/Source/Runtime/TypedElementFramework/Private/Elements/Framework/TypedElementQueryBuilder.cpp
+++ b/Engine/Source/Runtime/TypedElementFramework/Private/Elements/Framework/TypedElementQueryBuilder.cpp
@@ -40,7 +40,7 @@ namespace TypedElementQueryBuilder
 	// DependsOn
 	//
 
-	FDependency::FDependency(ITypedElementDataStorageInterface::FQueryDescription* Query)
+	FDependency::FDependency(FQueryDescription* Query)
 		: Query(Query)
 	{
 	}
@@ -49,7 +49,7 @@ namespace TypedElementQueryBuilder
 	{
 		checkf(Target, TEXT("The Dependency section in the Typed Elements query builder doesn't support nullptrs as Read-Only input."));
 		Query->DependencyTypes.Emplace(Target);
-		Query->DependencyFlags.Emplace(ITypedElementDataStorageInterface::EQueryDependencyFlags::ReadOnly);
+		Query->DependencyFlags.Emplace(EQueryDependencyFlags::ReadOnly);
 		Query->CachedDependencies.AddDefaulted();
 		return *this;
 	}
@@ -72,7 +72,7 @@ namespace TypedElementQueryBuilder
 	{
 		checkf(Target, TEXT("The Dependency section in the Typed Elements query builder doesn't support nullptrs as Read/Write input."));
 		Query->DependencyTypes.Emplace(Target);
-		Query->DependencyFlags.Emplace(ITypedElementDataStorageInterface::EQueryDependencyFlags::None);
+		Query->DependencyFlags.Emplace(EQueryDependencyFlags::None);
 		Query->CachedDependencies.AddDefaulted();
 		return *this;
 	}
@@ -93,7 +93,7 @@ namespace TypedElementQueryBuilder
 
 	FDependency& FDependency::SubQuery(QueryHandle Handle)
 	{
-		checkf(Query->Callback.ExecutionMode != TypedElementDataStorage::EExecutionMode::ThreadedChunks,
+		checkf(Query->Callback.ExecutionMode != EExecutionMode::ThreadedChunks,
 			TEXT("TEDS sub-queries can not be added to queries with a callback that process chunks in parallel."));
 		Query->Subqueries.Add(Handle);
 		return *this;
@@ -101,13 +101,13 @@ namespace TypedElementQueryBuilder
 	
 	FDependency& FDependency::SubQuery(TConstArrayView<QueryHandle> Handles)
 	{
-		checkf(Query->Callback.ExecutionMode != TypedElementDataStorage::EExecutionMode::ThreadedChunks,
+		checkf(Query->Callback.ExecutionMode != EExecutionMode::ThreadedChunks,
 			TEXT("TEDS sub-queries can not be added to queries with a callback that process chunks in parallel."));
 		Query->Subqueries.Insert(Handles.GetData(), Handles.Num(), Query->Subqueries.Num());
 		return *this;
 	}
 
-	ITypedElementDataStorageInterface::FQueryDescription&& FDependency::Compile()
+	FQueryDescription&& FDependency::Compile()
 	{
 		return MoveTemp(*Query);
 	}
@@ -117,7 +117,7 @@ namespace TypedElementQueryBuilder
 	 * Simple Query
 	 */
 
-	FSimpleQuery::FSimpleQuery(ITypedElementDataStorageInterface::FQueryDescription* Query)
+	FSimpleQuery::FSimpleQuery(FQueryDescription* Query)
 		: Query(Query)
 	{
 	}
@@ -126,7 +126,7 @@ namespace TypedElementQueryBuilder
 	{
 		if (Target)
 		{
-			Query->ConditionTypes.Add(ITypedElementDataStorageInterface::FQueryDescription::EOperatorType::SimpleAll);
+			Query->ConditionTypes.Add(FQueryDescription::EOperatorType::SimpleAll);
 			Query->ConditionOperators.AddZeroed_GetRef().Type = Target;
 		}
 		return *this;
@@ -158,7 +158,6 @@ namespace TypedElementQueryBuilder
 
 	FSimpleQuery& FSimpleQuery::All(const UEnum& Enum)
 	{
-		using namespace UE::Editor::DataStorage;
 		return All(FValueTag(Enum.GetFName()));
 	}
 
@@ -170,14 +169,13 @@ namespace TypedElementQueryBuilder
 			UE_LOG(LogEditorDataStorage, Warning, TEXT("Invalid value '%lld' for enum '%s'"), Value, *Enum.GetName());
 			return *this;
 		}
-		using namespace UE::Editor::DataStorage;
 		return All(FValueTag(Enum.GetFName()), ValueName);
 	}
 	
 	FSimpleQuery& FSimpleQuery::All(const FDynamicColumnDescription& Description)
 	{
 		Query->DynamicConditionDescriptions.Add(Description);
-		Query->DynamicConditionOperations.Add(ITypedElementDataStorageInterface::FQueryDescription::EOperatorType::SimpleAll);
+		Query->DynamicConditionOperations.Add(FQueryDescription::EOperatorType::SimpleAll);
 		return *this;
 	}
 
@@ -190,7 +188,7 @@ namespace TypedElementQueryBuilder
 	{
 		if (Target)
 		{
-			Query->ConditionTypes.Add(ITypedElementDataStorageInterface::FQueryDescription::EOperatorType::SimpleAny);
+			Query->ConditionTypes.Add(FQueryDescription::EOperatorType::SimpleAny);
 			Query->ConditionOperators.AddZeroed_GetRef().Type = Target;
 		}
 		return *this;
@@ -220,7 +218,7 @@ namespace TypedElementQueryBuilder
 	{
 		if (Target)
 		{
-			Query->ConditionTypes.Add(ITypedElementDataStorageInterface::FQueryDescription::EOperatorType::SimpleNone);
+			Query->ConditionTypes.Add(FQueryDescription::EOperatorType::SimpleNone);
 			Query->ConditionOperators.AddZeroed_GetRef().Type = Target;
 		}
 		return *this;
@@ -251,7 +249,7 @@ namespace TypedElementQueryBuilder
 		return FDependency{ Query };
 	}
 
-	ITypedElementDataStorageInterface::FQueryDescription&& FSimpleQuery::Compile()
+	FQueryDescription&& FSimpleQuery::Compile()
 	{
 		Query->Callback.BeforeGroups.Shrink();
 		Query->Callback.AfterGroups.Shrink();
@@ -281,12 +279,12 @@ namespace TypedElementQueryBuilder
 	/**
 	 * FProcessor
 	 */
-	FProcessor::FProcessor(ITypedElementDataStorageInterface::EQueryTickPhase Phase, FName Group)
+	FProcessor::FProcessor(EQueryTickPhase Phase, FName Group)
 		: Phase(Phase)
 		, Group(Group)
 	{}
 
-	FProcessor& FProcessor::SetPhase(ITypedElementDataStorageInterface::EQueryTickPhase NewPhase)
+	FProcessor& FProcessor::SetPhase(EQueryTickPhase NewPhase)
 	{
 		Phase = NewPhase;
 		return *this;
@@ -310,7 +308,7 @@ namespace TypedElementQueryBuilder
 		return *this;
 	}
 
-	FProcessor& FProcessor::SetExecutionMode(TypedElementDataStorage::EExecutionMode Mode)
+	FProcessor& FProcessor::SetExecutionMode(EExecutionMode Mode)
 	{
 		ExecutionMode = Mode;
 		return *this;
@@ -344,7 +342,7 @@ namespace TypedElementQueryBuilder
 		return *this;
 	}
 
-	FObserver& FObserver::SetExecutionMode(TypedElementDataStorage::EExecutionMode Mode)
+	FObserver& FObserver::SetExecutionMode(EExecutionMode Mode)
 	{
 		ExecutionMode = Mode;
 		return *this;
@@ -361,7 +359,7 @@ namespace TypedElementQueryBuilder
 	 * FPhaseAmble
 	 */
 
-	FPhaseAmble::FPhaseAmble(ELocation InLocation, ITypedElementDataStorageInterface::EQueryTickPhase InPhase)
+	FPhaseAmble::FPhaseAmble(ELocation InLocation, EQueryTickPhase InPhase)
 		: Phase(InPhase)
 		, Location(InLocation)
 	{}
@@ -372,13 +370,13 @@ namespace TypedElementQueryBuilder
 		return *this;
 	}
 
-	FPhaseAmble& FPhaseAmble::SetPhase(ITypedElementDataStorageInterface::EQueryTickPhase NewPhase)
+	FPhaseAmble& FPhaseAmble::SetPhase(EQueryTickPhase NewPhase)
 	{
 		Phase = NewPhase;
 		return *this;
 	}
 
-	FPhaseAmble& FPhaseAmble::SetExecutionMode(TypedElementDataStorage::EExecutionMode Mode)
+	FPhaseAmble& FPhaseAmble::SetExecutionMode(EExecutionMode Mode)
 	{
 		ExecutionMode = Mode;
 		return *this;
@@ -397,14 +395,14 @@ namespace TypedElementQueryBuilder
 
 	Select::Select()
 	{
-		Query.Action = ITypedElementDataStorageInterface::FQueryDescription::EActionType::Select;
+		Query.Action = FQueryDescription::EActionType::Select;
 	}
 	
 	Select& Select::ReadOnly(const UScriptStruct* Target)
 	{
 		checkf(Target, TEXT("The Select section in the Typed Elements query builder doesn't support nullptrs as Read-Only input."));
 		Query.SelectionTypes.Emplace(Target);
-		Query.SelectionAccessTypes.Emplace(ITypedElementDataStorageInterface::EQueryAccessType::ReadOnly);
+		Query.SelectionAccessTypes.Emplace(EQueryAccessType::ReadOnly);
 		Query.SelectionMetaData.Emplace(Target, FColumnMetaData::EFlags::None);
 		
 		return *this;
@@ -427,7 +425,7 @@ namespace TypedElementQueryBuilder
 	Select& Select::ReadOnly(const FDynamicColumnDescription& Description)
 	{
 		Query.DynamicSelectionTypes.Emplace(Description);
-		Query.DynamicSelectionAccessTypes.Emplace(ITypedElementDataStorageInterface::EQueryAccessType::ReadOnly);
+		Query.DynamicSelectionAccessTypes.Emplace(EQueryAccessType::ReadOnly);
 		Query.DynamicSelectionMetaData.Emplace(FColumnMetaData::EFlags::None);
 		return *this;
 	}
@@ -437,8 +435,8 @@ namespace TypedElementQueryBuilder
 		checkf(Target, TEXT("The Select section in the Typed Elements query builder doesn't support nullptrs as Read-Only input."));
 		Query.SelectionTypes.Emplace(Target);
 		Query.SelectionAccessTypes.Emplace(Optional == EOptional::Yes
-			? ITypedElementDataStorageInterface::EQueryAccessType::OptionalReadOnly
-			: ITypedElementDataStorageInterface::EQueryAccessType::ReadOnly);
+			? EQueryAccessType::OptionalReadOnly
+			: EQueryAccessType::ReadOnly);
 		Query.SelectionMetaData.Emplace(Target, FColumnMetaData::EFlags::None);
 
 		return *this;
@@ -462,7 +460,7 @@ namespace TypedElementQueryBuilder
 	{
 		checkf(Target, TEXT("The Select section in the Typed Elements query builder doesn't support nullptrs as Read/Write input."));
 		Query.SelectionTypes.Emplace(Target);
-		Query.SelectionAccessTypes.Emplace(ITypedElementDataStorageInterface::EQueryAccessType::ReadWrite);
+		Query.SelectionAccessTypes.Emplace(EQueryAccessType::ReadWrite);
 		Query.SelectionMetaData.Emplace(Target, FColumnMetaData::EFlags::IsMutable);
 		return *this;
 	}
@@ -484,7 +482,7 @@ namespace TypedElementQueryBuilder
 	Select& Select::ReadWrite(const FDynamicColumnDescription& Descirption)
 	{
 		Query.DynamicSelectionTypes.Emplace(Descirption);
-		Query.DynamicSelectionAccessTypes.Emplace(ITypedElementDataStorageInterface::EQueryAccessType::ReadWrite);
+		Query.DynamicSelectionAccessTypes.Emplace(EQueryAccessType::ReadWrite);
 		Query.DynamicSelectionMetaData.Emplace(FColumnMetaData::EFlags::IsMutable);
 		return *this;
 	}
@@ -499,7 +497,7 @@ namespace TypedElementQueryBuilder
 		return FDependency{ &Query };
 	}
 
-	ITypedElementDataStorageInterface::FQueryDescription&& Select::Compile()
+	FQueryDescription&& Select::Compile()
 	{
 		return MoveTemp(Query);
 	}
@@ -511,7 +509,7 @@ namespace TypedElementQueryBuilder
 
 	Count::Count()
 	{
-		Query.Action = ITypedElementDataStorageInterface::FQueryDescription::EActionType::Count;
+		Query.Action = FQueryDescription::EActionType::Count;
 	}
 
 	FSimpleQuery Count::Where()
diff --git a/Engine/Source/Runtime/TypedElementFramework/Public/Elements/Columns/TypedElementPackageColumns.h b/Engine/Source/Runtime/TypedElementFramework/Public/Elements/Columns/TypedElementPackageColumns.h
--- a/Engine/Source/Runtime/TypedElementFramework/Public/Elements/Columns/TypedElementPackageColumns.h
+++ b/Engine/Source/Runtime/TypedElementFramework/Public/Elements/Columns/TypedElementPackageColumns.h
@@ -6,6 +6,7 @@
 #include "Elements/Common/TypedElementCommonTypes.h"
 #include "Elements/Common/TypedElementHandles.h"
 #include "Elements/Common/TypedElementQueryTypes.h"
+#include "HAL/Platform.h"
 #include "IO/PackageId.h"
 #include "Misc/PackagePath.h"
 #include "UObject/ObjectMacros.h"
@@ -22,7 +23,7 @@ struct FTypedElementPackageUnresolvedReference final : public FEditorDataStorage
 
 	FString PathOnDisk;
 	
-	TypedElementDataStorage::IndexHash Index;
+	UE::Editor::DataStorage::IndexHash Index;
 };
 
 /**
diff --git a/Engine/Source/Runtime/TypedElementFramework/Public/Elements/Common/TypedElementQueryDescription.h b/Engine/Source/Runtime/TypedElementFramework/Public/Elements/Common/TypedElementQueryDescription.h
--- a/Engine/Source/Runtime/TypedElementFramework/Public/Elements/Common/TypedElementQueryDescription.h
+++ b/Engine/Source/Runtime/TypedElementFramework/Public/Elements/Common/TypedElementQueryDescription.h
@@ -19,6 +19,7 @@ class UScriptStruct;
 namespace UE::Editor::DataStorage
 {
 	struct FQueryDescription;
+	struct IQueryContext;
 	
 	using QueryCallback = TFunction<void(const FQueryDescription&, IQueryContext&)>;
 	using QueryCallbackRef = TFunctionRef<void(const FQueryDescription&, IQueryContext&)>;
diff --git a/Engine/Source/Runtime/TypedElementFramework/Public/Elements/Common/TypedElementQueryTypes.h b/Engine/Source/Runtime/TypedElementFramework/Public/Elements/Common/TypedElementQueryTypes.h
--- a/Engine/Source/Runtime/TypedElementFramework/Public/Elements/Common/TypedElementQueryTypes.h
+++ b/Engine/Source/Runtime/TypedElementFramework/Public/Elements/Common/TypedElementQueryTypes.h
@@ -2,9 +2,10 @@
 
 #pragma once
 
-#include "CoreMinimal.h"
+#include "HAL/Platform.h"
+#include "Misc/EnumClassFlags.h"
 
-namespace TypedElementDataStorage
+namespace UE::Editor::DataStorage
 {
 	using IndexHash = uint64;
 
@@ -129,4 +130,4 @@ namespace TypedElementDataStorage
 		uint32 Count{ 0 }; /** The number of rows were processed. */
 		ECompletion Completed{ ECompletion::Unavailable };
 	};
-} // namespace TypedElementDataStorage
+} // namespace UE::Editor::DataStorage
diff --git a/Engine/Source/Runtime/TypedElementFramework/Public/Elements/Framework/TypedElementIndexHasher.h b/Engine/Source/Runtime/TypedElementFramework/Public/Elements/Framework/TypedElementIndexHasher.h
--- a/Engine/Source/Runtime/TypedElementFramework/Public/Elements/Framework/TypedElementIndexHasher.h
+++ b/Engine/Source/Runtime/TypedElementFramework/Public/Elements/Framework/TypedElementIndexHasher.h
@@ -11,7 +11,7 @@
 #include "UObject/StrongObjectPtr.h"
 #include "UObject/WeakObjectPtrTemplates.h"
 
-namespace TypedElementDataStorage
+namespace UE::Editor::DataStorage
 {
 	template<typename T>
 	IndexHash GenerateIndexHash(const T* Object);
@@ -27,6 +27,6 @@ namespace TypedElementDataStorage
 	inline IndexHash GenerateIndexHash(FStringView Object);
 	inline IndexHash GenerateIndexHash(FName Object);
 	inline IndexHash GenerateIndexHash(const FSoftObjectPath& ObjectPath);
-} // namespace TypedElementDataStorage
+} // namespace UE::Editor::DataStorage
 
 #include "Elements/Framework/TypedElementIndexHasher.inl"
\ No newline at end of file
diff --git a/Engine/Source/Runtime/TypedElementFramework/Public/Elements/Framework/TypedElementIndexHasher.inl b/Engine/Source/Runtime/TypedElementFramework/Public/Elements/Framework/TypedElementIndexHasher.inl
--- a/Engine/Source/Runtime/TypedElementFramework/Public/Elements/Framework/TypedElementIndexHasher.inl
+++ b/Engine/Source/Runtime/TypedElementFramework/Public/Elements/Framework/TypedElementIndexHasher.inl
@@ -5,7 +5,7 @@
 #include "Hash/CityHash.h"
 #include "UObject/TopLevelAssetPath.h"
 
-namespace TypedElementDataStorage
+namespace UE::Editor::DataStorage
 {
 	template<typename T>
 	IndexHash GenerateIndexHash(const T* Object)
@@ -58,4 +58,4 @@ namespace TypedElementDataStorage
 		Hash = CityHash128to64({ Hash, GenerateIndexHash(TopLevelAssetPath.GetAssetName()) });
 		return CityHash128to64({ Hash, GenerateIndexHash(ObjectPath.GetSubPathString()) });
 	}
-} // namespace TypedElementDataStorage
+} // namespace UE::Editor::DataStorage
diff --git a/Engine/Source/Runtime/TypedElementFramework/Public/Elements/Framework/TypedElementQueryBuilder.h b/Engine/Source/Runtime/TypedElementFramework/Public/Elements/Framework/TypedElementQueryBuilder.h
--- a/Engine/Source/Runtime/TypedElementFramework/Public/Elements/Framework/TypedElementQueryBuilder.h
+++ b/Engine/Source/Runtime/TypedElementFramework/Public/Elements/Framework/TypedElementQueryBuilder.h
@@ -48,7 +48,7 @@ class USubsystem;
  *	- void(<Context>&, RowHandle, [const]Column&...) 
  *	- void(<Context>&, [const]Column*...) 
  *	- void(<Context>&, const RowHandle*, [const]Column*...) 
- *	Where <Context> is ITypedElementDataStorageInterface::IQueryContext or FCachedQueryContext<...>	e.g.:
+ *	Where <Context> is IQueryContext or FCachedQueryContext<...>	e.g.:
  *		void(
  *			FCachedQueryContext<Subsystem1, const Subsystem2>& Context, 
  *			RowHandle Row, 
@@ -62,8 +62,8 @@ class USubsystem;
  *
  * The following is a simplified example of these options combined together:
  *		FProcessor Info(
- *			ITypedElementDataStorageInterface::EQueryTickPhase::FrameEnd, 
- *			DataStorage->GetQueryTickGroupName(ITypedElementDataStorageInterface::EQueryTickGroups::SyncExternalToDataStorage);
+ *			EQueryTickPhase::FrameEnd, 
+ *			DataStorage->GetQueryTickGroupName(EQueryTickGroups::SyncExternalToDataStorage);
  *		Query = Select(FName(TEXT("Example Callback")), Info, 
  *				[](FCachedQueryContext<Subsystem1, const Subsystem2>&, const FDataExample1&, FDataExample2&) {});
  * 
@@ -86,7 +86,7 @@ class USubsystem;
  * - ReadWrite: Indicates that the external system will be used to write data to.
  *
  * Usage example:
- * ITypedElementDataStorageInterface::FQueryDescription Query =
+ * FQueryDescription Query =
  *		Select()
  *			.ReadWrite({ FDataExample1::StaticStruct() })
  *			.ReadWrite<FDataExample2, FDataExample3>()
@@ -137,12 +137,12 @@ namespace TypedElementQueryBuilder
 		TYPEDELEMENTFRAMEWORK_API FDependency& SubQuery(QueryHandle Handle);
 		TYPEDELEMENTFRAMEWORK_API FDependency& SubQuery(TConstArrayView<QueryHandle> Handles);
 
-		TYPEDELEMENTFRAMEWORK_API ITypedElementDataStorageInterface::FQueryDescription&& Compile();
+		TYPEDELEMENTFRAMEWORK_API FQueryDescription&& Compile();
 
 	private:
-		TYPEDELEMENTFRAMEWORK_API explicit FDependency(ITypedElementDataStorageInterface::FQueryDescription* Query);
+		TYPEDELEMENTFRAMEWORK_API explicit FDependency(FQueryDescription* Query);
 
-		ITypedElementDataStorageInterface::FQueryDescription* Query;
+		FQueryDescription* Query;
 	};
 
 	class FSimpleQuery final
@@ -152,7 +152,7 @@ namespace TypedElementQueryBuilder
 		friend class Select;
 
 		TYPEDELEMENTFRAMEWORK_API FDependency DependsOn();
-		TYPEDELEMENTFRAMEWORK_API ITypedElementDataStorageInterface::FQueryDescription&& Compile();
+		TYPEDELEMENTFRAMEWORK_API FQueryDescription&& Compile();
 
 		template<TColumnType... TargetTypes>
 		FSimpleQuery& All();
@@ -165,14 +165,14 @@ namespace TypedElementQueryBuilder
 		TYPEDELEMENTFRAMEWORK_API FSimpleQuery& Any(TConstArrayView<const UScriptStruct*> Targets);
 
 		// Dynamic Column Support
-		TYPEDELEMENTFRAMEWORK_API FSimpleQuery& Any(const UE::Editor::DataStorage::FDynamicColumnDescription& Description);
-		template<UE::Editor::DataStorage::TColumnType T>
+		TYPEDELEMENTFRAMEWORK_API FSimpleQuery& Any(const FDynamicColumnDescription& Description);
+		template<TColumnType T>
 		FSimpleQuery& Any(const FName&);
 
 		// ValueTags not yet supported for Any
-		template<UE::Editor::DataStorage::TValueTagType>
+		template<TValueTagType>
 		FSimpleQuery& Any(const FName&) = delete;
-		template<UE::Editor::DataStorage::TValueTagType>
+		template<TValueTagType>
 		FSimpleQuery& Any(const FName&, const FName&) = delete;
 		
 		template<TColumnType... TargetTypes>
@@ -182,14 +182,14 @@ namespace TypedElementQueryBuilder
 		TYPEDELEMENTFRAMEWORK_API FSimpleQuery& None(TConstArrayView<const UScriptStruct*> Targets);
 		
 		// ValueTags not yet supported for None
-		template<UE::Editor::DataStorage::TValueTagType>
+		template<TValueTagType>
 		FSimpleQuery& None(const FName& Tag) = delete;
-		template<UE::Editor::DataStorage::TValueTagType>
+		template<TValueTagType>
 		FSimpleQuery& None(const FName& Tag, const FName& Value) = delete;
 
 		// Dynamic Column Support
-		TYPEDELEMENTFRAMEWORK_API FSimpleQuery& None(const UE::Editor::DataStorage::FDynamicColumnDescription& Description);
-		template<UE::Editor::DataStorage::TColumnType T>
+		TYPEDELEMENTFRAMEWORK_API FSimpleQuery& None(const FDynamicColumnDescription& Description);
+		template<TColumnType T>
 		FSimpleQuery& None(const FName&);
 		
 		// Value Tags
@@ -204,16 +204,16 @@ namespace TypedElementQueryBuilder
 		TYPEDELEMENTFRAMEWORK_API FSimpleQuery& All(const UEnum& Enum, int64 Value);
 
 		// Dynamic Column support
-		TYPEDELEMENTFRAMEWORK_API FSimpleQuery& All(const UE::Editor::DataStorage::FDynamicColumnDescription& Description);
-		template<UE::Editor::DataStorage::TColumnType T>
+		TYPEDELEMENTFRAMEWORK_API FSimpleQuery& All(const FDynamicColumnDescription& Description);
+		template<TColumnType T>
 		FSimpleQuery& All(const FName&);
 
 		// Adds a filter to the query which must match a given ValueTag.  The value of the DynamicTag will not be checked.
-		template<UE::Editor::DataStorage::TValueTagType>
+		template<TValueTagType>
 		FSimpleQuery& All(const FName& Tag);
 		// Adds a filter to the query which must match a given ValueTag.  The value of the DynamicTag must also match.
 		// Note: The query can only match a single value.  Multiple value queries are not supported at this time.
-		template<UE::Editor::DataStorage::TValueTagType>
+		template<TValueTagType>
 		FSimpleQuery& All(const FName& Tag, const FName& Value);
 		
 		template<TEnumType EnumT>
@@ -226,29 +226,29 @@ namespace TypedElementQueryBuilder
 		FSimpleQuery& All();
 
 	private:
-		TYPEDELEMENTFRAMEWORK_API explicit FSimpleQuery(ITypedElementDataStorageInterface::FQueryDescription* Query);
+		TYPEDELEMENTFRAMEWORK_API explicit FSimpleQuery(FQueryDescription* Query);
 
-		ITypedElementDataStorageInterface::FQueryDescription* Query;
+		FQueryDescription* Query;
 	};
 
 	struct FQueryCallbackType{};
 
 	struct FProcessor final : public FQueryCallbackType
 	{
-		TYPEDELEMENTFRAMEWORK_API FProcessor(ITypedElementDataStorageInterface::EQueryTickPhase Phase, FName Group);
-		TYPEDELEMENTFRAMEWORK_API FProcessor& SetPhase(ITypedElementDataStorageInterface::EQueryTickPhase NewPhase);
+		TYPEDELEMENTFRAMEWORK_API FProcessor(EQueryTickPhase Phase, FName Group);
+		TYPEDELEMENTFRAMEWORK_API FProcessor& SetPhase(EQueryTickPhase NewPhase);
 		TYPEDELEMENTFRAMEWORK_API FProcessor& SetGroup(FName GroupName);
 		TYPEDELEMENTFRAMEWORK_API FProcessor& SetBeforeGroup(FName GroupName);
 		TYPEDELEMENTFRAMEWORK_API FProcessor& SetAfterGroup(FName GroupName);
-		TYPEDELEMENTFRAMEWORK_API FProcessor& SetExecutionMode(TypedElementDataStorage::EExecutionMode Mode);
+		TYPEDELEMENTFRAMEWORK_API FProcessor& SetExecutionMode(EExecutionMode Mode);
 		TYPEDELEMENTFRAMEWORK_API FProcessor& MakeActivatable(FName Name);
 		
-		ITypedElementDataStorageInterface::EQueryTickPhase Phase;
+		EQueryTickPhase Phase;
 		FName Group;
 		FName BeforeGroup;
 		FName AfterGroup;
 		FName ActivationName;
-		TypedElementDataStorage::EExecutionMode ExecutionMode = TypedElementDataStorage::EExecutionMode::Default;
+		EExecutionMode ExecutionMode = EExecutionMode::Default;
 	};
 
 	struct FObserver final : public FQueryCallbackType
@@ -270,13 +270,13 @@ namespace TypedElementQueryBuilder
 		TYPEDELEMENTFRAMEWORK_API FObserver& SetMonitoredColumn(const UScriptStruct* MonitoredColumn);
 		template<TColumnType ColumnType>
 		FObserver& SetMonitoredColumn();
-		TYPEDELEMENTFRAMEWORK_API FObserver& SetExecutionMode(TypedElementDataStorage::EExecutionMode Mode);
+		TYPEDELEMENTFRAMEWORK_API FObserver& SetExecutionMode(EExecutionMode Mode);
 		TYPEDELEMENTFRAMEWORK_API FObserver& MakeActivatable(FName Name);
 
 		const UScriptStruct* Monitor;
 		EEvent Event;
 		FName ActivationName;
-		TypedElementDataStorage::EExecutionMode ExecutionMode = TypedElementDataStorage::EExecutionMode::Default;
+		EExecutionMode ExecutionMode = EExecutionMode::Default;
 	};
 
 	struct FPhaseAmble final : public FQueryCallbackType
@@ -287,40 +287,40 @@ namespace TypedElementQueryBuilder
 			Postamble
 		};
 
-		TYPEDELEMENTFRAMEWORK_API FPhaseAmble(ELocation InLocation, ITypedElementDataStorageInterface::EQueryTickPhase InPhase);
+		TYPEDELEMENTFRAMEWORK_API FPhaseAmble(ELocation InLocation, EQueryTickPhase InPhase);
 		TYPEDELEMENTFRAMEWORK_API FPhaseAmble& SetLocation(ELocation NewLocation);
-		TYPEDELEMENTFRAMEWORK_API FPhaseAmble& SetPhase(ITypedElementDataStorageInterface::EQueryTickPhase NewPhase);
-		TYPEDELEMENTFRAMEWORK_API FPhaseAmble& SetExecutionMode(TypedElementDataStorage::EExecutionMode Mode);
+		TYPEDELEMENTFRAMEWORK_API FPhaseAmble& SetPhase(EQueryTickPhase NewPhase);
+		TYPEDELEMENTFRAMEWORK_API FPhaseAmble& SetExecutionMode(EExecutionMode Mode);
 		TYPEDELEMENTFRAMEWORK_API FPhaseAmble& MakeActivatable(FName Name);
 
-		ITypedElementDataStorageInterface::EQueryTickPhase Phase;
+		EQueryTickPhase Phase;
 		ELocation Location;
 		FName ActivationName;
-		TypedElementDataStorage::EExecutionMode ExecutionMode = TypedElementDataStorage::EExecutionMode::Default;
+		EExecutionMode ExecutionMode = EExecutionMode::Default;
 	};
 
 	// Because this is a thin wrapper called from within a query callback, it's better to inline fully so all
 	// function pre/postambles can be optimized away.
-	struct FQueryContextForwarder : public ITypedElementDataStorageInterface::IQueryContext
+	struct FQueryContextForwarder : public IQueryContext
 	{
 		inline FQueryContextForwarder(
-			const ITypedElementDataStorageInterface::FQueryDescription& InDescription, 
-			ITypedElementDataStorageInterface::IQueryContext& InParentContext);
+			const FQueryDescription& InDescription, 
+			IQueryContext& InParentContext);
 		inline ~FQueryContextForwarder() = default;
 
 		inline const void* GetColumn(const UScriptStruct* ColumnType) const override;
 		inline void* GetMutableColumn(const UScriptStruct* ColumnType) override;
 		inline void GetColumns(TArrayView<char*> RetrievedAddresses, TConstArrayView<TWeakObjectPtr<const UScriptStruct>> ColumnTypes,
-			TConstArrayView<ITypedElementDataStorageInterface::EQueryAccessType> AccessTypes) override;
+			TConstArrayView<EQueryAccessType> AccessTypes) override;
 		inline void GetColumnsUnguarded(int32 TypeCount, char** RetrievedAddresses, const TWeakObjectPtr<const UScriptStruct>* ColumnTypes,
-			const ITypedElementDataStorageInterface::EQueryAccessType* AccessTypes) override;
+			const EQueryAccessType* AccessTypes) override;
 
 		inline bool HasColumn(const UScriptStruct* ColumnType) const override;
 		
 		inline UObject* GetMutableDependency(const UClass* DependencyClass) override;
 		inline const UObject* GetDependency(const UClass* DependencyClass) override;
 		inline void GetDependencies(TArrayView<UObject*> RetrievedAddresses, TConstArrayView<TWeakObjectPtr<const UClass>> SubsystemTypes,
-			TConstArrayView<ITypedElementDataStorageInterface::EQueryAccessType> AccessTypes) override;
+			TConstArrayView<EQueryAccessType> AccessTypes) override;
 
 		inline uint32 GetRowCount() const override;
 		inline TConstArrayView<RowHandle> GetRowHandles() const override;
@@ -332,24 +332,24 @@ namespace TypedElementQueryBuilder
 		inline void RemoveColumns(RowHandle Row, TConstArrayView<const UScriptStruct*> ColumnTypes) override;
 		inline void RemoveColumns(TConstArrayView<RowHandle> Rows, TConstArrayView<const UScriptStruct*> ColumnTypes) override;
 
-		inline TypedElementDataStorage::FQueryResult RunQuery(QueryHandle Query) override;
-		inline TypedElementDataStorage::FQueryResult RunSubquery(int32 SubqueryIndex) override;
-		inline TypedElementDataStorage::FQueryResult RunSubquery(int32 SubqueryIndex, SubqueryCallbackRef Callback) override;
-		inline TypedElementDataStorage::FQueryResult RunSubquery(int32 SubqueryIndex, RowHandle Row,
+		inline FQueryResult RunQuery(QueryHandle Query) override;
+		inline FQueryResult RunSubquery(int32 SubqueryIndex) override;
+		inline FQueryResult RunSubquery(int32 SubqueryIndex, SubqueryCallbackRef Callback) override;
+		inline FQueryResult RunSubquery(int32 SubqueryIndex, RowHandle Row,
 			SubqueryCallbackRef Callback) override;
 
-		ITypedElementDataStorageInterface::IQueryContext& ParentContext;
-		const ITypedElementDataStorageInterface::FQueryDescription& Description;
+		IQueryContext& ParentContext;
+		const FQueryDescription& Description;
 	};
 
 	template<typename... Dependencies>
 	struct FCachedQueryContext final : public FQueryContextForwarder
 	{
 		explicit FCachedQueryContext(
-			const ITypedElementDataStorageInterface::FQueryDescription& InDescription, 
-			ITypedElementDataStorageInterface::IQueryContext& InParentContext);
+			const FQueryDescription& InDescription, 
+			IQueryContext& InParentContext);
 		
-		static void Register(ITypedElementDataStorageInterface::FQueryDescription& Query);
+		static void Register(FQueryDescription& Query);
 
 		template<typename Dependency>
 		Dependency& GetCachedMutableDependency();
@@ -377,7 +377,7 @@ namespace TypedElementQueryBuilder
 		TYPEDELEMENTFRAMEWORK_API Select& ReadOnly(const UScriptStruct* Target);
 		/** Request read-only access to the listed columns. */
 		TYPEDELEMENTFRAMEWORK_API Select& ReadOnly(TConstArrayView<const UScriptStruct*> Targets);
-		TYPEDELEMENTFRAMEWORK_API Select& ReadOnly(const UE::Editor::DataStorage::FDynamicColumnDescription& Description);
+		TYPEDELEMENTFRAMEWORK_API Select& ReadOnly(const FDynamicColumnDescription& Description);
 		/** 
 		 * Request read-only access to the listed columns. If optional is true read access will be given if the column is in the table but
 		 * it will not be used for finding matching tables. Columns bound with optional can not be bound to a query callback argument.
@@ -394,7 +394,7 @@ namespace TypedElementQueryBuilder
 		 * it will not be used for finding matching tables. Columns bound with optional can not be bound to a query callback argument.
 		 */
 		TYPEDELEMENTFRAMEWORK_API Select& ReadOnly(TConstArrayView<const UScriptStruct*> Targets, EOptional Optional);
-		template<UE::Editor::DataStorage::TDataColumnType Target>
+		template<TDataColumnType Target>
 		Select& ReadOnly(const FName& Identifier);
 		/** Request read and write access to the listed columns. */
 		template<TDataColumnType... TargetTypes>
@@ -403,16 +403,16 @@ namespace TypedElementQueryBuilder
 		TYPEDELEMENTFRAMEWORK_API Select& ReadWrite(const UScriptStruct* Target);
 		/** Request read and write access to the listed columns. */
 		TYPEDELEMENTFRAMEWORK_API Select& ReadWrite(TConstArrayView<const UScriptStruct*> Targets);
-		TYPEDELEMENTFRAMEWORK_API Select& ReadWrite(const UE::Editor::DataStorage::FDynamicColumnDescription& Description);
-		template<UE::Editor::DataStorage::TDataColumnType Target>
+		TYPEDELEMENTFRAMEWORK_API Select& ReadWrite(const FDynamicColumnDescription& Description);
+		template<TDataColumnType Target>
 		Select& ReadWrite(const FName& Identifier);
 
-		TYPEDELEMENTFRAMEWORK_API ITypedElementDataStorageInterface::FQueryDescription&& Compile();
+		TYPEDELEMENTFRAMEWORK_API FQueryDescription&& Compile();
 		TYPEDELEMENTFRAMEWORK_API FSimpleQuery Where();
 		TYPEDELEMENTFRAMEWORK_API FDependency DependsOn();
 
 	private:
-		ITypedElementDataStorageInterface::FQueryDescription Query;
+		FQueryDescription Query;
 	};
 
 	// Explicitly not following the naming convention in order to keep readability consistent. It now reads like a query sentence.
@@ -425,7 +425,7 @@ namespace TypedElementQueryBuilder
 		TYPEDELEMENTFRAMEWORK_API FDependency DependsOn();
 
 	private:
-		ITypedElementDataStorageInterface::FQueryDescription Query;
+		FQueryDescription Query;
 	};
 
 	template<typename Function>
diff --git a/Engine/Source/Runtime/TypedElementFramework/Public/Elements/Framework/TypedElementQueryBuilder.inl b/Engine/Source/Runtime/TypedElementFramework/Public/Elements/Framework/TypedElementQueryBuilder.inl
--- a/Engine/Source/Runtime/TypedElementFramework/Public/Elements/Framework/TypedElementQueryBuilder.inl
+++ b/Engine/Source/Runtime/TypedElementFramework/Public/Elements/Framework/TypedElementQueryBuilder.inl
@@ -32,37 +32,37 @@ namespace TypedElementQueryBuilder
 		}
 
 		template<typename Type>
-		constexpr TypedElementDataStorage::EQueryAccessType GetAccessType()
+		constexpr EQueryAccessType GetAccessType()
 		{
 			using BaseType = typename std::remove_reference_t<Type>;
 			if constexpr (TIsConst<BaseType>::Value)
 			{
-				return TypedElementDataStorage::EQueryAccessType::ReadOnly;
+				return EQueryAccessType::ReadOnly;
 			}
 			else
 			{
-				return TypedElementDataStorage::EQueryAccessType::ReadWrite;
+				return EQueryAccessType::ReadWrite;
 			}
 		}
 
 		template<typename Type>
-		constexpr TypedElementDataStorage::EQueryDependencyFlags GetDependencyFlags()
+		constexpr EQueryDependencyFlags GetDependencyFlags()
 		{
 			using BaseType = typename std::remove_reference_t<Type>;
 			using SubsystemTraits = TTypedElementSubsystemTraits<std::remove_const_t<BaseType>>;
 
-			TypedElementDataStorage::EQueryDependencyFlags Result = TypedElementDataStorage::EQueryDependencyFlags::None;
+			EQueryDependencyFlags Result = EQueryDependencyFlags::None;
 			if constexpr (SubsystemTraits::RequiresGameThread())
 			{
-				EnumAddFlags(Result, TypedElementDataStorage::EQueryDependencyFlags::GameThreadBound);
+				EnumAddFlags(Result, EQueryDependencyFlags::GameThreadBound);
 			}
 			if constexpr (SubsystemTraits::IsHotReloadable())
 			{
-				EnumAddFlags(Result, TypedElementDataStorage::EQueryDependencyFlags::AlwaysRefresh);
+				EnumAddFlags(Result, EQueryDependencyFlags::AlwaysRefresh);
 			}
 			if constexpr (TIsConst<BaseType>::Value)
 			{
-				EnumAddFlags(Result, TypedElementDataStorage::EQueryDependencyFlags::ReadOnly);
+				EnumAddFlags(Result, EQueryDependencyFlags::ReadOnly);
 			}
 			
 			return Result;
@@ -116,8 +116,8 @@ namespace TypedElementQueryBuilder
 	//
 
 	FQueryContextForwarder::FQueryContextForwarder(
-		const UE::Editor::DataStorage::FQueryDescription& InDescription,
-		UE::Editor::DataStorage::IQueryContext& InParentContext)
+		const FQueryDescription& InDescription,
+		IQueryContext& InParentContext)
 		: ParentContext(InParentContext)
 		, Description(InDescription)
 	{}
@@ -133,13 +133,13 @@ namespace TypedElementQueryBuilder
 	}
 
 	void FQueryContextForwarder::GetColumns(TArrayView<char*> RetrievedAddresses, TConstArrayView<TWeakObjectPtr<const UScriptStruct>> ColumnTypes,
-		TConstArrayView<TypedElementDataStorage::EQueryAccessType> AccessTypes)
+		TConstArrayView<EQueryAccessType> AccessTypes)
 	{
 		ParentContext.GetColumns(RetrievedAddresses, ColumnTypes, AccessTypes);
 	}
 
 	void FQueryContextForwarder::GetColumnsUnguarded(int32 TypeCount, char** RetrievedAddresses, const TWeakObjectPtr<const UScriptStruct>* ColumnTypes,
-		const TypedElementDataStorage::EQueryAccessType* AccessTypes)
+		const EQueryAccessType* AccessTypes)
 	{
 		ParentContext.GetColumnsUnguarded(TypeCount, RetrievedAddresses, ColumnTypes, AccessTypes);
 	}
@@ -160,7 +160,7 @@ namespace TypedElementQueryBuilder
 	}
 
 	void FQueryContextForwarder::GetDependencies(TArrayView<UObject*> RetrievedAddresses, TConstArrayView<TWeakObjectPtr<const UClass>> DependencyTypes,
-		TConstArrayView<TypedElementDataStorage::EQueryAccessType> AccessTypes)
+		TConstArrayView<EQueryAccessType> AccessTypes)
 	{
 		ParentContext.GetDependencies(RetrievedAddresses, DependencyTypes, AccessTypes);
 	}
@@ -205,24 +205,24 @@ namespace TypedElementQueryBuilder
 		ParentContext.RemoveColumns(Rows, ColumnTypes);
 	}
 
-	TypedElementDataStorage::FQueryResult FQueryContextForwarder::RunQuery(QueryHandle Query)
+	FQueryResult FQueryContextForwarder::RunQuery(QueryHandle Query)
 	{
 		return ParentContext.RunQuery(Query);
 	}
 	
-	TypedElementDataStorage::FQueryResult FQueryContextForwarder::RunSubquery(int32 SubqueryIndex)
+	FQueryResult FQueryContextForwarder::RunSubquery(int32 SubqueryIndex)
 	{
 		return ParentContext.RunSubquery(SubqueryIndex);
 	}
 
-	TypedElementDataStorage::FQueryResult FQueryContextForwarder::RunSubquery(int32 SubqueryIndex,
-		UE::Editor::DataStorage::SubqueryCallbackRef Callback)
+	FQueryResult FQueryContextForwarder::RunSubquery(int32 SubqueryIndex,
+		SubqueryCallbackRef Callback)
 	{
 		return ParentContext.RunSubquery(SubqueryIndex, Callback);
 	}
 
-	TypedElementDataStorage::FQueryResult FQueryContextForwarder::RunSubquery(int32 SubqueryIndex,
-		RowHandle Row, UE::Editor::DataStorage::SubqueryCallbackRef Callback)
+	FQueryResult FQueryContextForwarder::RunSubquery(int32 SubqueryIndex,
+		RowHandle Row, SubqueryCallbackRef Callback)
 	{
 		return ParentContext.RunSubquery(SubqueryIndex, Row, Callback);
 	}
@@ -233,13 +233,13 @@ namespace TypedElementQueryBuilder
 	
 	template<typename... Dependencies>
 	FCachedQueryContext<Dependencies...>::FCachedQueryContext(
-		const UE::Editor::DataStorage::FQueryDescription& InDescription,
-		UE::Editor::DataStorage::IQueryContext& InParentContext)
+		const FQueryDescription& InDescription,
+		IQueryContext& InParentContext)
 		: FQueryContextForwarder(InDescription, InParentContext)
 	{}
 
 	template<typename... Dependencies>
-	void FCachedQueryContext<Dependencies...>::Register(UE::Editor::DataStorage::FQueryDescription& Query)
+	void FCachedQueryContext<Dependencies...>::Register(FQueryDescription& Query)
 	{
 		Query.DependencyTypes.Reserve(sizeof...(Dependencies));
 		Query.DependencyFlags.Reserve(sizeof...(Dependencies));
@@ -291,7 +291,7 @@ namespace TypedElementQueryBuilder
 
 		template<typename Context>
 		concept SourceQueryContext = requires(Context& Ctx, int32 TypeCount, char** RetrievedAddresses,
-			const TWeakObjectPtr<const UScriptStruct>*ColumnTypes, const TypedElementDataStorage::EQueryAccessType * AccessTypes)
+			const TWeakObjectPtr<const UScriptStruct>*ColumnTypes, const EQueryAccessType * AccessTypes)
 		{
 			{ Ctx.GetColumnsUnguarded(TypeCount, RetrievedAddresses, ColumnTypes, AccessTypes) };
 			{ Ctx.GetRowHandles() } -> UE::convertible_to<TConstArrayView<RowHandle>>;
@@ -386,7 +386,7 @@ namespace TypedElementQueryBuilder
 			static constexpr bool bArePointerColumns = AreAllColumnsPointers<ColumnTypes...>();
 
 			template<typename Column>
-			bool CheckValidity(const UE::Editor::DataStorage::FQueryDescription& Description)
+			bool CheckValidity(const FQueryDescription& Description)
 			{
 				const UScriptStruct* ColumnType = Column::StaticStruct();
 				if (const TWeakObjectPtr<const UScriptStruct>* Match = Description.SelectionTypes.FindByPredicate(
@@ -396,7 +396,7 @@ namespace TypedElementQueryBuilder
 					}); Match != nullptr)
 				{
 					size_t Index = Match - Description.SelectionTypes.GetData();
-					if (Description.SelectionAccessTypes[Index] == TypedElementDataStorage::EQueryAccessType::ReadOnly)
+					if (Description.SelectionAccessTypes[Index] == EQueryAccessType::ReadOnly)
 					{
 						if (std::is_const_v<Column>)
 						{
@@ -424,7 +424,7 @@ namespace TypedElementQueryBuilder
 				}
 			}
 
-			FFunctionColumnInfo(const UE::Editor::DataStorage::FQueryDescription& Description, SourceContext& Context)
+			FFunctionColumnInfo(const FQueryDescription& Description, SourceContext& Context)
 			{
 				if constexpr (ValidateColumns)
 				{
@@ -448,7 +448,7 @@ namespace TypedElementQueryBuilder
 		template<SourceQueryContext SourceContext, bool ValidateColumns>
 		struct FFunctionColumnInfo<SourceContext, ValidateColumns>
 		{
-			FFunctionColumnInfo(const UE::Editor::DataStorage::FQueryDescription&, SourceContext&) {}
+			FFunctionColumnInfo(const FQueryDescription&, SourceContext&) {}
 			static constexpr bool bArePointerColumns = false;
 		};
 
@@ -459,7 +459,7 @@ namespace TypedElementQueryBuilder
 			using BaseContextType = std::remove_reference_t<TargetType>;
 			BaseContextType ContextWrapper;
 
-			FContextInfo(const UE::Editor::DataStorage::FQueryDescription& Description, SourceType& Context)
+			FContextInfo(const FQueryDescription& Description, SourceType& Context)
 				: ContextWrapper(Description, Context)
 			{}
 		};
@@ -469,7 +469,7 @@ namespace TypedElementQueryBuilder
 		{
 			const SourceType& ContextWrapper;
 
-			FContextInfo(const UE::Editor::DataStorage::FQueryDescription& Description, SourceType& Context)
+			FContextInfo(const FQueryDescription& Description, SourceType& Context)
 				: ContextWrapper(Context)
 			{}
 		};
@@ -479,7 +479,7 @@ namespace TypedElementQueryBuilder
 		{
 			SourceType& ContextWrapper;
 
-			FContextInfo(const UE::Editor::DataStorage::FQueryDescription& Description, SourceType& Context)
+			FContextInfo(const FQueryDescription& Description, SourceType& Context)
 				: ContextWrapper(Context)
 			{}
 		};
@@ -493,7 +493,7 @@ namespace TypedElementQueryBuilder
 			using SuperColumn = FFunctionColumnInfo<SourceContext, ValidateColumns, Args...>;
 			using SuperContext = FContextInfo<SourceContext, TargetContext>;
 
-			FContextRowHandleColumnsFunction(const UE::Editor::DataStorage::FQueryDescription& Description, SourceContext& Context)
+			FContextRowHandleColumnsFunction(const FQueryDescription& Description, SourceContext& Context)
 				: SuperColumn(Description, Context)
 				, SuperContext(Description, Context)
 			{}
@@ -546,7 +546,7 @@ namespace TypedElementQueryBuilder
 			using SuperColumn = FFunctionColumnInfo<SourceContext, ValidateColumns, Args...>;
 			using SuperContext = FContextInfo<SourceContext, TargetContext>;
 
-			FContextColumnsFunction(const UE::Editor::DataStorage::FQueryDescription& Description, SourceContext& Context)
+			FContextColumnsFunction(const FQueryDescription& Description, SourceContext& Context)
 				: SuperColumn(Description, Context)
 				, SuperContext(Description, Context)
 			{}
@@ -594,7 +594,7 @@ namespace TypedElementQueryBuilder
 		{
 			using Super = FFunctionColumnInfo<SourceContext, ValidateColumns, Columns...>;
 
-			FRowHandleColumnsFunction(const UE::Editor::DataStorage::FQueryDescription& Description, SourceContext& Context)
+			FRowHandleColumnsFunction(const FQueryDescription& Description, SourceContext& Context)
 				: Super(Description, Context)
 			{}
 
@@ -642,7 +642,7 @@ namespace TypedElementQueryBuilder
 		{
 			using Super = FFunctionColumnInfo<SourceContext, ValidateColumns, Args...>;
 			
-			FColumnsFunction(const UE::Editor::DataStorage::FQueryDescription& Description, SourceContext& Context)
+			FColumnsFunction(const FQueryDescription& Description, SourceContext& Context)
 				: Super(Description, Context)
 			{}
 
@@ -716,105 +716,105 @@ namespace TypedElementQueryBuilder
 		{
 			using Super = typename FFunctionInfoHelper<SourceContext, ValidateColumns, Args...>::BaseClass;
 
-			FFunctionInfo(const UE::Editor::DataStorage::FQueryDescription& Description, SourceContext& Context)
+			FFunctionInfo(const FQueryDescription& Description, SourceContext& Context)
 				: Super(Description, Context)
 			{}
 		};
 
 		template<SourceQueryContext SourceContext, bool ValidateColumns, typename... Args>
-		void BindQueryFunction(UE::Editor::DataStorage::QueryCallback& Function,
+		void BindQueryFunction(QueryCallback& Function,
 			void (*Callback)(Args...))
 		{
-			Function = [Callback](const UE::Editor::DataStorage::FQueryDescription& Description, SourceContext& Context)
+			Function = [Callback](const FQueryDescription& Description, SourceContext& Context)
 			{
 				FFunctionInfo<SourceContext, ValidateColumns, Args...>(Description, Context).Call(Context, Callback);
 			};
 		}
 
 		template<SourceQueryContext SourceContext, bool ValidateColumns, typename Class, typename... Args>
-		void BindQueryFunction(UE::Editor::DataStorage::QueryCallback& Function, Class* Target,
+		void BindQueryFunction(QueryCallback& Function, Class* Target,
 			void (Class::*Callback)(Args...))
 		{
-			Function = [Target, Callback](const UE::Editor::DataStorage::FQueryDescription& Description, SourceContext& Context)
+			Function = [Target, Callback](const FQueryDescription& Description, SourceContext& Context)
 			{
 				FFunctionInfo<SourceContext, ValidateColumns, Args...>(Description, Context).Call(Context, (Target->*Callback));
 			};
 		}
 
 		template<SourceQueryContext SourceContext, bool ValidateColumns, typename Class, typename... Args>
-		void BindQueryFunction(UE::Editor::DataStorage::QueryCallback& Function, Class* Target,
+		void BindQueryFunction(QueryCallback& Function, Class* Target,
 			void (Class::*Callback)(Args...) const)
 		{
-			Function = [Target, Callback](const UE::Editor::DataStorage::FQueryDescription& Description, SourceContext& Context)
+			Function = [Target, Callback](const FQueryDescription& Description, SourceContext& Context)
 			{
 				FFunctionInfo<SourceContext, ValidateColumns, Args...>(Description, Context).Call(Context, (Target->*Callback));
 			};
 		}
 
 		template<SourceQueryContext SourceContext, bool ValidateColumns, typename Functor, typename Class, typename... Args>
-		void BindQueryFunction_Expand(UE::Editor::DataStorage::QueryCallback& Function, Functor&& CallbackObject,
+		void BindQueryFunction_Expand(QueryCallback& Function, Functor&& CallbackObject,
 			void (Class::*Callback)(Args...) const)
 		{
-			Function = [CallbackObject = std::forward<Functor>(CallbackObject)](const UE::Editor::DataStorage::FQueryDescription& Description, SourceContext& Context)
+			Function = [CallbackObject = std::forward<Functor>(CallbackObject)](const FQueryDescription& Description, SourceContext& Context)
 			{
 				FFunctionInfo<SourceContext, ValidateColumns, Args...>(Description, Context).Call(Context, CallbackObject);
 			};
 		}
 
 		template<SourceQueryContext SourceContext, bool ValidateColumns, typename Functor, typename Class, typename... Args>
-		void BindQueryFunction_Expand(UE::Editor::DataStorage::QueryCallback& Function, Functor&& CallbackObject,
+		void BindQueryFunction_Expand(QueryCallback& Function, Functor&& CallbackObject,
 			void (Class::*Callback)(Args...))
 		{
-			Function = [CallbackObject = std::forward<Functor>(CallbackObject)](const UE::Editor::DataStorage::FQueryDescription& Description, SourceContext& Context)
+			Function = [CallbackObject = std::forward<Functor>(CallbackObject)](const FQueryDescription& Description, SourceContext& Context)
 			{
 				FFunctionInfo<SourceContext, ValidateColumns, Args...>(Description, Context).Call(Context, CallbackObject);
 			};
 		}
 
 		template<SourceQueryContext SourceContext, bool ValidateColumns, typename Functor>
-		void BindQueryFunction(UE::Editor::DataStorage::QueryCallback& Function, Functor&& Callback)
+		void BindQueryFunction(QueryCallback& Function, Functor&& Callback)
 		{
 			BindQueryFunction_Expand<SourceContext, ValidateColumns>(Function, std::forward<Functor>(Callback), &Functor::operator());
 		}
 
 		
 		template<bool ValidateColumns, SourceQueryContext SourceContext, typename... Args>
-		void CallQueryFunction(const UE::Editor::DataStorage::FQueryDescription& Description, SourceContext& Context,
-			UE::Editor::DataStorage::QueryCallback& Function, void (*Callback)(Args...))
+		void CallQueryFunction(const FQueryDescription& Description, SourceContext& Context,
+			QueryCallback& Function, void (*Callback)(Args...))
 		{
 			FFunctionInfo<SourceContext, ValidateColumns, Args...>(Description, Context).Call(Context, Callback);
 		}
 
 		template<bool ValidateColumns, SourceQueryContext SourceContext, typename Class, typename... Args>
-		void CallQueryFunction(const UE::Editor::DataStorage::FQueryDescription& Description, SourceContext& Context,
+		void CallQueryFunction(const FQueryDescription& Description, SourceContext& Context,
 			Class* Target, void (Class::* Callback)(Args...))
 		{
 			FFunctionInfo<SourceContext, ValidateColumns, Args...>(Description, Context).Call(Context, (Target->*Callback));
 		}
 
 		template<bool ValidateColumns, SourceQueryContext SourceContext, typename Class, typename... Args>
-		void CallQueryFunction(const UE::Editor::DataStorage::FQueryDescription& Description, SourceContext& Context,
+		void CallQueryFunction(const FQueryDescription& Description, SourceContext& Context,
 			Class* Target, void (Class::* Callback)(Args...) const)
 		{
 			FFunctionInfo<SourceContext, ValidateColumns, Args...>(Description, Context).Call(Context, (Target->*Callback));
 		}
 
 		template<bool ValidateColumns, SourceQueryContext SourceContext, typename Functor, typename Class, typename... Args>
-		void CallQueryFunction_Expand(const UE::Editor::DataStorage::FQueryDescription& Description, SourceContext& Context,
+		void CallQueryFunction_Expand(const FQueryDescription& Description, SourceContext& Context,
 			const Functor& CallbackObject, void (Class::* Callback)(Args...) const)
 		{
 			FFunctionInfo<SourceContext, ValidateColumns, Args...>(Description, Context).Call(Context, CallbackObject);
 		}
 
 		template<bool ValidateColumns, SourceQueryContext SourceContext, typename Functor, typename Class, typename... Args>
-		void CallQueryFunction_Expand(const UE::Editor::DataStorage::FQueryDescription& Description, SourceContext& Context,
+		void CallQueryFunction_Expand(const FQueryDescription& Description, SourceContext& Context,
 			const Functor& CallbackObject, void (Class::* Callback)(Args...))
 		{
 			FFunctionInfo<SourceContext, ValidateColumns, Args...>(Description, Context).Call(Context, CallbackObject);
 		}
 
 		template<bool ValidateColumns, SourceQueryContext SourceContext, typename Functor>
-		void CallQueryFunction(const UE::Editor::DataStorage::FQueryDescription& Description, SourceContext& Context,
+		void CallQueryFunction(const FQueryDescription& Description, SourceContext& Context,
 			const Functor& Callback)
 		{
 			CallQueryFunction_Expand<ValidateColumns>(Description, Context, Callback, &Functor::operator());
@@ -834,7 +834,7 @@ namespace TypedElementQueryBuilder
 			}
 		}
 
-		template<typename ContextType> void RegisterDependencies(UE::Editor::DataStorage::FQueryDescription& Query)
+		template<typename ContextType> void RegisterDependencies(FQueryDescription& Query)
 		{
 			using BaseType = typename TRemoveReference<ContextType>::Type;
 			if constexpr (TIsDerivedFrom<BaseType, FQueryContextForwarder>::Value)
@@ -846,7 +846,7 @@ namespace TypedElementQueryBuilder
 		template<SourceQueryContext SourceContext, typename... Args>
 		struct RegisterFunctionArgumentsHelper
 		{
-			void Register(UE::Editor::DataStorage::FQueryDescription&, Select& Target)
+			void Register(FQueryDescription&, Select& Target)
 			{
 				(AddColumnToSelect<Args>(Target), ...);
 			}
@@ -855,13 +855,13 @@ namespace TypedElementQueryBuilder
 		template<SourceQueryContext SourceContext>
 		struct RegisterFunctionArgumentsHelper<SourceContext>
 		{
-			static void Register(UE::Editor::DataStorage::FQueryDescription&, Select&) {}
+			static void Register(FQueryDescription&, Select&) {}
 		};
 
 		template<SourceQueryContext SourceContext, typename Arg0>
 		struct RegisterFunctionArgumentsHelper<SourceContext, Arg0>
 		{
-			static void Register(UE::Editor::DataStorage::FQueryDescription& Query, Select& Target)
+			static void Register(FQueryDescription& Query, Select& Target)
 			{
 				if constexpr (IsCompatibleTargetContextType<SourceContext, Arg0>)
 				{
@@ -878,7 +878,7 @@ namespace TypedElementQueryBuilder
 		template<SourceQueryContext SourceContext, typename Arg0, typename Arg1, typename... Args>
 		struct RegisterFunctionArgumentsHelper<SourceContext, Arg0, Arg1, Args...>
 		{
-			static void Register(UE::Editor::DataStorage::FQueryDescription& Query, Select& Target)
+			static void Register(FQueryDescription& Query, Select& Target)
 			{
 				if constexpr (IsCompatibleTargetContextType<SourceContext, Arg0>)
 				{
@@ -901,25 +901,25 @@ namespace TypedElementQueryBuilder
 		};
 
 		template<SourceQueryContext SourceContext, typename... Args>
-		void RegisterFunctionArguments(UE::Editor::DataStorage::FQueryDescription& Query, Select& Target, void (*)(Args...))
+		void RegisterFunctionArguments(FQueryDescription& Query, Select& Target, void (*)(Args...))
 		{
 			RegisterFunctionArgumentsHelper<SourceContext, Args...>::Register(Query, Target);
 		}
 
 		template<SourceQueryContext SourceContext, typename Class, typename... Args>
-		void RegisterFunctionArguments(UE::Editor::DataStorage::FQueryDescription& Query, Select& Target, void (Class::*)(Args...))
+		void RegisterFunctionArguments(FQueryDescription& Query, Select& Target, void (Class::*)(Args...))
 		{
 			RegisterFunctionArgumentsHelper<SourceContext, Args...>::Register(Query, Target);
 		}
 
 		template<SourceQueryContext SourceContext, typename Class, typename... Args>
-		void RegisterFunctionArguments(UE::Editor::DataStorage::FQueryDescription& Query, Select& Target, void (Class::*)(Args...) const)
+		void RegisterFunctionArguments(FQueryDescription& Query, Select& Target, void (Class::*)(Args...) const)
 		{
 			RegisterFunctionArgumentsHelper<SourceContext, Args...>::Register(Query, Target);
 		}
 
 		template<SourceQueryContext SourceContext, typename Functor>
-		void RegisterFunctionArguments(UE::Editor::DataStorage::FQueryDescription& Query, Select& Target, Functor)
+		void RegisterFunctionArguments(FQueryDescription& Query, Select& Target, Functor)
 		{
 			RegisterFunctionArguments<SourceContext>(Query, Target, &Functor::operator());
 		}
@@ -1025,9 +1025,9 @@ namespace TypedElementQueryBuilder
 			}
 		};
 		
-		inline void PrepareForQueryBinding(UE::Editor::DataStorage::FQueryDescription& Query, const FProcessor& Processor)
+		inline void PrepareForQueryBinding(FQueryDescription& Query, const FProcessor& Processor)
 		{
-			Query.Callback.Type = TypedElementDataStorage::EQueryCallbackType::Processor;
+			Query.Callback.Type = EQueryCallbackType::Processor;
 			Query.Callback.Phase = Processor.Phase;
 			Query.Callback.Group = Processor.Group;
 			Query.Callback.ActivationName = Processor.ActivationName;
@@ -1043,15 +1043,15 @@ namespace TypedElementQueryBuilder
 			Query.Callback.ExecutionMode = Processor.ExecutionMode;
 		}
 
-		inline void PrepareForQueryBinding(UE::Editor::DataStorage::FQueryDescription& Query, const FObserver& Observer)
+		inline void PrepareForQueryBinding(FQueryDescription& Query, const FObserver& Observer)
 		{
 			switch (Observer.Event)
 			{
 			case FObserver::EEvent::Add:
-				Query.Callback.Type = TypedElementDataStorage::EQueryCallbackType::ObserveAdd;
+				Query.Callback.Type = EQueryCallbackType::ObserveAdd;
 				break;
 			case FObserver::EEvent::Remove:
-				Query.Callback.Type = TypedElementDataStorage::EQueryCallbackType::ObserveRemove;
+				Query.Callback.Type = EQueryCallbackType::ObserveRemove;
 				break;
 			}
 			Query.Callback.ActivationName = Observer.ActivationName;
@@ -1060,15 +1060,15 @@ namespace TypedElementQueryBuilder
 			Query.Callback.ExecutionMode = Observer.ExecutionMode;
 		}
 
-		inline void PrepareForQueryBinding(UE::Editor::DataStorage::FQueryDescription& Query, const FPhaseAmble& PhaseAmble)
+		inline void PrepareForQueryBinding(FQueryDescription& Query, const FPhaseAmble& PhaseAmble)
 		{
 			switch (PhaseAmble.Location)
 			{
 			case FPhaseAmble::ELocation::Preamble:
-				Query.Callback.Type = TypedElementDataStorage::EQueryCallbackType::PhasePreparation;
+				Query.Callback.Type = EQueryCallbackType::PhasePreparation;
 				break;
 			case FPhaseAmble::ELocation::Postamble:
-				Query.Callback.Type = TypedElementDataStorage::EQueryCallbackType::PhaseFinalization;
+				Query.Callback.Type = EQueryCallbackType::PhaseFinalization;
 				break;
 			}
 			Query.Callback.ActivationName = PhaseAmble.ActivationName;
@@ -1078,12 +1078,12 @@ namespace TypedElementQueryBuilder
 		}
 		
 		template<typename CallbackType, typename Function>
-		void PrepareForQueryBinding(Select& Target, UE::Editor::DataStorage::FQueryDescription& Query, FName Name,
+		void PrepareForQueryBinding(Select& Target, FQueryDescription& Query, FName Name,
 			const CallbackType& Type, Function Callback)
 		{
 			static_assert(TIsDerivedFrom<CallbackType, FQueryCallbackType>::Value, "The callback type provided isn't one of the available "
 				"classes derived from FQueryCallbackType.");
-			static_assert(IsValidSelectFunctionSignature<UE::Editor::DataStorage::IQueryContext, Function>(),
+			static_assert(IsValidSelectFunctionSignature<IQueryContext, Function>(),
 				R"(The function provided to the Query Builder's Select call wasn't invocable or doesn't contain a supported combination of arguments.
 The following options are supported:
 - void([const]Column&...) 
@@ -1092,10 +1092,10 @@ The following options are supported:
 - void(<Context>&, RowHandle, [const]Column&...) 
 - void(<Context>&, [const]Column*...) 
 - void(<Context>&, const RowHandle*, [const]Column*...) 
-Where <Context> is UE::Editor::DataStorage::IQueryContext or FCachedQueryContext<...>
+Where <Context> is IQueryContext or FCachedQueryContext<...>
 e.g. void(FCachedQueryContext<Subsystem1, const Subsystem2>& Context, RowHandle Row, ColumnType0& ColumnA, const ColumnType1& ColumnB) {...}
 )");
-			RegisterFunctionArguments<UE::Editor::DataStorage::IQueryContext>(Query, Target, Callback);
+			RegisterFunctionArguments<IQueryContext>(Query, Target, Callback);
 			PrepareForQueryBinding(Query, Type);
 			Query.Callback.Name = Name;
 		}
@@ -1107,7 +1107,7 @@ e.g. void(FCachedQueryContext<Subsystem1, const Subsystem2>& Context, RowHandle
 	{
 		static constexpr bool ValidateColumns = false;
 		Private::PrepareForQueryBinding(*this, Query, Name, Type, Callback);
-		Private::BindQueryFunction<UE::Editor::DataStorage::IQueryContext, ValidateColumns>(Query.Callback.Function, std::forward<Function>(Callback));
+		Private::BindQueryFunction<IQueryContext, ValidateColumns>(Query.Callback.Function, std::forward<Function>(Callback));
 	}
 
 	template<typename CallbackType, typename Class, typename Function>
@@ -1116,7 +1116,7 @@ e.g. void(FCachedQueryContext<Subsystem1, const Subsystem2>& Context, RowHandle
 	{
 		static constexpr bool ValidateColumns = false;
 		Private::PrepareForQueryBinding(*this, Query, Name, Type, Callback);
-		Private::BindQueryFunction<UE::Editor::DataStorage::IQueryContext, ValidateColumns>(Query.Callback.Function, Instance, Callback);
+		Private::BindQueryFunction<IQueryContext, ValidateColumns>(Query.Callback.Function, Instance, Callback);
 	}
 
 	template<TDataColumnType... TargetTypes>
@@ -1133,10 +1133,10 @@ e.g. void(FCachedQueryContext<Subsystem1, const Subsystem2>& Context, RowHandle
 		return *this;
 	}
 
-	template<UE::Editor::DataStorage::TDataColumnType Target>
+	template<TDataColumnType Target>
 	Select& Select::ReadOnly(const FName& Identifier)
 	{
-		ReadOnly(UE::Editor::DataStorage::FDynamicColumnDescription
+		ReadOnly(FDynamicColumnDescription
 			{
 				.TemplateType = Target::StaticStruct(),
 				.Identifier = Identifier
@@ -1154,7 +1154,7 @@ e.g. void(FCachedQueryContext<Subsystem1, const Subsystem2>& Context, RowHandle
 	template<TDataColumnType Target>
 	Select& Select::ReadWrite(const FName& Identifier)
 	{
-		ReadWrite(UE::Editor::DataStorage::FDynamicColumnDescription
+		ReadWrite(FDynamicColumnDescription
 			{
 				.TemplateType = Target::StaticStruct(),
 				.Identifier = Identifier
@@ -1184,7 +1184,7 @@ e.g. void(FCachedQueryContext<Subsystem1, const Subsystem2>& Context, RowHandle
 	template <TColumnType DynamicColumnTemplate>
 	FSimpleQuery& FSimpleQuery::Any(const FName& Identifier)
 	{
-		return Any(UE::Editor::DataStorage::FDynamicColumnDescription
+		return Any(FDynamicColumnDescription
 		{
 			.TemplateType = DynamicColumnTemplate::StaticStruct(),
 			.Identifier = Identifier
@@ -1198,22 +1198,22 @@ e.g. void(FCachedQueryContext<Subsystem1, const Subsystem2>& Context, RowHandle
 		return *this;
 	}
 
-	template<UE::Editor::DataStorage::TValueTagType>
+	template<TValueTagType>
 	FSimpleQuery& FSimpleQuery::All(const FName& Tag)
 	{
-		return All(UE::Editor::DataStorage::FValueTag(Tag));
+		return All(FValueTag(Tag));
 	}
 
-	template<UE::Editor::DataStorage::TValueTagType>
+	template<TValueTagType>
 	FSimpleQuery& FSimpleQuery::All(const FName& Tag, const FName& Value)
 	{
-		return All(UE::Editor::DataStorage::FValueTag(Tag), Value);
+		return All(FValueTag(Tag), Value);
 	}
 	
 	template <TColumnType DynamicColumnTemplate>
 	FSimpleQuery& FSimpleQuery::None(const FName& Identifier)
 	{
-		return None(UE::Editor::DataStorage::FDynamicColumnDescription
+		return None(FDynamicColumnDescription
 			{
 				.TemplateType = DynamicColumnTemplate::StaticStruct(),
 				.Identifier = Identifier
@@ -1223,7 +1223,7 @@ e.g. void(FCachedQueryContext<Subsystem1, const Subsystem2>& Context, RowHandle
 	template <TColumnType DynamicColumnTemplate>
 	FSimpleQuery& FSimpleQuery::All(const FName& Identifier)
 	{
-		return All(UE::Editor::DataStorage::FDynamicColumnDescription
+		return All(FDynamicColumnDescription
 			{
 				.TemplateType = DynamicColumnTemplate::StaticStruct(),
 				.Identifier = Identifier
@@ -1252,9 +1252,9 @@ e.g. void(FCachedQueryContext<Subsystem1, const Subsystem2>& Context, RowHandle
 	// External query bindings
 	//
 	template<typename Function>
-	UE::Editor::DataStorage::DirectQueryCallback CreateDirectQueryCallbackBinding(Function&& Callback)
+	DirectQueryCallback CreateDirectQueryCallbackBinding(Function&& Callback)
 	{
-		static_assert(Private::IsValidSelectFunctionSignature<UE::Editor::DataStorage::IDirectQueryContext, Function>(),
+		static_assert(Private::IsValidSelectFunctionSignature<IDirectQueryContext, Function>(),
 			R"(The function provided to the Query Builder's CreateDirectQueryCallbackBinding call wasn't invocable or doesn't contain a supported combination of arguments.
 The following options are supported:
 - void([const]Column&...) 
@@ -1263,13 +1263,13 @@ The following options are supported:
 - void(<Context>&, RowHandle, [const]Column&...) 
 - void(<Context>&, [const]Column*...) 
 - void(<Context>&, const RowHandle*, [const]Column*...) 
-Where <Context> is UE::Editor::DataStorage::IDirectQueryContext
+Where <Context> is IDirectQueryContext
 e.g. void(IDirectQueryContext& Context, RowHandle Row, ColumnType0& ColumnA, const ColumnType1& ColumnB) {...}
 )");
 
 		return [Callback = Forward<Function>(Callback)](
-			const UE::Editor::DataStorage::FQueryDescription& Description,
-			UE::Editor::DataStorage::IDirectQueryContext& Context)
+			const FQueryDescription& Description,
+			IDirectQueryContext& Context)
 		{
 			static constexpr bool ValidateColumns = true;
 			Private::CallQueryFunction<ValidateColumns>(Description, Context, Callback);
@@ -1277,9 +1277,9 @@ e.g. void(IDirectQueryContext& Context, RowHandle Row, ColumnType0& ColumnA, con
 	}
 
 	template<typename Function>
-	UE::Editor::DataStorage::SubqueryCallback CreateSubqueryCallbackBinding(Function&& Callback)
+	SubqueryCallback CreateSubqueryCallbackBinding(Function&& Callback)
 	{
-		static_assert(Private::IsValidSelectFunctionSignature<ITypedElementDataStorageInterface::ISubqueryContext, Function>(),
+		static_assert(Private::IsValidSelectFunctionSignature<ISubqueryContext, Function>(),
 			R"(The function provided to the Query Builder's CreateSubqueryCallbackBinding call wasn't invocable or doesn't contain a supported combination of arguments.
 The following options are supported:
 - void([const]Column&...) 
@@ -1288,13 +1288,13 @@ The following options are supported:
 - void(<Context>&, RowHandle, [const]Column&...) 
 - void(<Context>&, [const]Column*...) 
 - void(<Context>&, const RowHandle*, [const]Column*...) 
-Where <Context> is UE::Editor::DataStorage::ISubqueryContext
+Where <Context> is ISubqueryContext
 e.g. void(ISubqueryContext& Context, RowHandle Row, ColumnType0& ColumnA, const ColumnType1& ColumnB) {...}
 )");
 
 		return [Callback = Forward<Function>(Callback)](
-			const UE::Editor::DataStorage::FQueryDescription& Description,
-			UE::Editor::DataStorage::ISubqueryContext& Context)
+			const FQueryDescription& Description,
+			ISubqueryContext& Context)
 		{
 			static constexpr bool ValidateColumns = true;
 			Private::CallQueryFunction<ValidateColumns>(Description, Context, Callback);
diff --git a/Engine/Source/Runtime/TypedElementFramework/Public/Elements/Interfaces/TypedElementDataStorageInterface.h b/Engine/Source/Runtime/TypedElementFramework/Public/Elements/Interfaces/TypedElementDataStorageInterface.h
--- a/Engine/Source/Runtime/TypedElementFramework/Public/Elements/Interfaces/TypedElementDataStorageInterface.h
+++ b/Engine/Source/Runtime/TypedElementFramework/Public/Elements/Interfaces/TypedElementDataStorageInterface.h
@@ -328,12 +328,12 @@ public:
 	 * therefore recommended to try to simplify the query first before relying on extended query filtering in a processor.
 	 */
 
-	using EQueryTickPhase = TypedElementDataStorage::EQueryTickPhase;
-	using EQueryTickGroups = TypedElementDataStorage::EQueryTickGroups;
-	using EQueryCallbackType = TypedElementDataStorage::EQueryCallbackType;
-	using EQueryAccessType = TypedElementDataStorage::EQueryAccessType;
-	using EQueryDependencyFlags = TypedElementDataStorage::EQueryDependencyFlags;
-	using FQueryResult = TypedElementDataStorage::FQueryResult;
+	using EQueryTickPhase = UE::Editor::DataStorage::EQueryTickPhase;
+	using EQueryTickGroups = UE::Editor::DataStorage::EQueryTickGroups;
+	using EQueryCallbackType = UE::Editor::DataStorage::EQueryCallbackType;
+	using EQueryAccessType = UE::Editor::DataStorage::EQueryAccessType;
+	using EQueryDependencyFlags = UE::Editor::DataStorage::EQueryDependencyFlags;
+	using FQueryResult = UE::Editor::DataStorage::FQueryResult;
 
 	using IQueryContext = UE::Editor::DataStorage::IQueryContext;
 	using IDirectQueryContext = UE::Editor::DataStorage::IDirectQueryContext;
@@ -392,23 +392,23 @@ public:
 	 */
 
 	/** Retrieves the row for an indexed object. Returns an invalid row handle if the hash wasn't found. */
-	virtual RowHandle FindIndexedRow(TypedElementDataStorage::IndexHash Index) const = 0;
+	virtual RowHandle FindIndexedRow(UE::Editor::DataStorage::IndexHash Index) const = 0;
 	/** 
 	 * Registers a row under the index hash. The same row can be registered multiple, but an index hash can only be associated 
 	 * with a single row.
 	 */
-	virtual void IndexRow(TypedElementDataStorage::IndexHash Index, RowHandle Row) = 0;
+	virtual void IndexRow(UE::Editor::DataStorage::IndexHash Index, RowHandle Row) = 0;
 	/**
 	 * Register multiple rows under their index hash. The same row can be registered multiple times,
 	 * but an index hash can only be associated with a single row.
 	 */
 	virtual void BatchIndexRows(
-		TConstArrayView<TPair<TypedElementDataStorage::IndexHash, RowHandle>> IndexRowPairs) = 0;
+		TConstArrayView<TPair<UE::Editor::DataStorage::IndexHash, RowHandle>> IndexRowPairs) = 0;
 	/** Updates the index of a row to a new value. Effectively this is the same as removing an index and adding a new one. */
 	virtual void ReindexRow(
-		TypedElementDataStorage::IndexHash OriginalIndex, TypedElementDataStorage::IndexHash NewIndex, RowHandle Row) = 0;
+		UE::Editor::DataStorage::IndexHash OriginalIndex, UE::Editor::DataStorage::IndexHash NewIndex, RowHandle Row) = 0;
 	/** Removes a previously registered index hash from the index lookup table or does nothing if the hash no longer exists. */
-	virtual void RemoveIndex(TypedElementDataStorage::IndexHash Index) = 0;
+	virtual void RemoveIndex(UE::Editor::DataStorage::IndexHash Index) = 0;
 
 	/**
 	 * @section Miscellaneous
diff --git a/Engine/Source/Runtime/TypedElementFramework/Public/Elements/Interfaces/TypedElementQueryStorageInterfaces.h b/Engine/Source/Runtime/TypedElementFramework/Public/Elements/Interfaces/TypedElementQueryStorageInterfaces.h
--- a/Engine/Source/Runtime/TypedElementFramework/Public/Elements/Interfaces/TypedElementQueryStorageInterfaces.h
+++ b/Engine/Source/Runtime/TypedElementFramework/Public/Elements/Interfaces/TypedElementQueryStorageInterfaces.h
@@ -16,8 +16,6 @@ class UScriptStruct;
 
 namespace UE::Editor::DataStorage
 {
-	using namespace TypedElementDataStorage;
-
 	struct FQueryDescription;
 	struct ISubqueryContext;
 
diff --git a/Engine/Source/Runtime/TypedElementFramework/Tests/Elements/Framework/TypedElementDataStoragePerformanceTestCommands.cpp b/Engine/Source/Runtime/TypedElementFramework/Tests/Elements/Framework/TypedElementDataStoragePerformanceTestCommands.cpp
--- a/Engine/Source/Runtime/TypedElementFramework/Tests/Elements/Framework/TypedElementDataStoragePerformanceTestCommands.cpp
+++ b/Engine/Source/Runtime/TypedElementFramework/Tests/Elements/Framework/TypedElementDataStoragePerformanceTestCommands.cpp
@@ -77,7 +77,6 @@ void UTest_PingPongBetweenPhaseFactory::RegisterTables(ITypedElementDataStorageI
 void UTest_PingPongBetweenPhaseFactory::RegisterQueries(ITypedElementDataStorageInterface& DataStorage)
 {
 	using namespace TypedElementQueryBuilder;
-	using namespace TypedElementDataStorage;
 
 	Super::RegisterQueries(DataStorage);
 
diff --git a/Engine/Source/Runtime/TypedElementFramework/Tests/Elements/Framework/TypedElementDataStorageProcessorTests.cpp b/Engine/Source/Runtime/TypedElementFramework/Tests/Elements/Framework/TypedElementDataStorageProcessorTests.cpp
--- a/Engine/Source/Runtime/TypedElementFramework/Tests/Elements/Framework/TypedElementDataStorageProcessorTests.cpp
+++ b/Engine/Source/Runtime/TypedElementFramework/Tests/Elements/Framework/TypedElementDataStorageProcessorTests.cpp
@@ -5,180 +5,177 @@
 #include "HAL/IConsoleManager.h"
 #include "TypedElementTestColumns.h"
 
-namespace UE::Editor::DataStorage
+namespace UE::Editor::DataStorage::Debug::ProcessorTests
 {
-	namespace Debug::ProcessorTests
-	{
-		void OnProcessorTestsEnabled(IConsoleVariable*);
+	void OnProcessorTestsEnabled(IConsoleVariable*);
+
+	bool GProcessorTestsEnabled = false;
+	FAutoConsoleVariableRef CVarAllowUnversionedContentInEditor(
+		TEXT("TEDS.Tests.ProcessorTestsEnabled"),
+		GProcessorTestsEnabled,
+		TEXT("If true, registers processors and additional commands with TEDS to test processors."),
+		FConsoleVariableDelegate::CreateStatic(&OnProcessorTestsEnabled),
+		ECVF_Default
+	);
 
-		bool GProcessorTestsEnabled = false;
-		FAutoConsoleVariableRef CVarAllowUnversionedContentInEditor(
-			TEXT("TEDS.Tests.ProcessorTestsEnabled"),
-			GProcessorTestsEnabled,
-			TEXT("If true, registers processors and additional commands with TEDS to test processors."),
-			FConsoleVariableDelegate::CreateStatic(&OnProcessorTestsEnabled),
-			ECVF_Default
-		);
+	TableHandle PrimaryTable = InvalidTableHandle;
+	TableHandle SecondaryTable = InvalidTableHandle;
 
-		TableHandle PrimaryTable = InvalidTableHandle;
-		TableHandle SecondaryTable = InvalidTableHandle;
+	TArray<QueryHandle> RegisteredQueries;
+	TArray<IConsoleCommand*> RegisteredCommands;
 
-		TArray<QueryHandle> RegisteredQueries;
-		TArray<IConsoleCommand*> RegisteredCommands;
+	void RegisterProcessors()
+	{
+		ITypedElementDataStorageInterface* DataStorage = UTypedElementRegistry::GetInstance()->GetMutableDataStorage();
+		using namespace TypedElementQueryBuilder;
 
-		void RegisterProcessors()
+		if (PrimaryTable == InvalidTableHandle)
 		{
-			ITypedElementDataStorageInterface* DataStorage = UTypedElementRegistry::GetInstance()->GetMutableDataStorage();
-			using namespace TypedElementQueryBuilder;
+			PrimaryTable =
+				DataStorage->RegisterTable<FTEDSProcessorTestsReferenceColumn, FTEDSProcessorTests_PrimaryTag>(FName(TEXT("ProcessorTests Primary Table")));
+		}
+		if (SecondaryTable == InvalidTableHandle)
+		{
+			SecondaryTable =
+				DataStorage->RegisterTable<FTEDSProcessorTestsReferenceColumn, FTEDSProcessorTests_SecondaryTag>(FName(TEXT("ProcessorTests Secondary Table")));
+		}
 
-			if (PrimaryTable == InvalidTableHandle)
+		// Test creation of a row from within a query processor
+		QueryHandle PrimaryRowQuery = DataStorage->RegisterQuery(
+			Select(TEXT("TEST: Creating a row for primary reference column"),
+			FProcessor(EQueryTickPhase::PrePhysics, DataStorage->GetQueryTickGroupName(EQueryTickGroups::Default)),
+			[](IQueryContext& Context, const RowHandle* Rows, FTEDSProcessorTestsReferenceColumn* ReferenceColumns)
 			{
-				PrimaryTable =
-					DataStorage->RegisterTable<FTEDSProcessorTestsReferenceColumn, FTEDSProcessorTests_PrimaryTag>(FName(TEXT("ProcessorTests Primary Table")));
-			}
-			if (SecondaryTable == InvalidTableHandle)
-			{
-				SecondaryTable =
-					DataStorage->RegisterTable<FTEDSProcessorTestsReferenceColumn, FTEDSProcessorTests_SecondaryTag>(FName(TEXT("ProcessorTests Secondary Table")));
-			}
+				const int32 RowCount = Context.GetRowCount();
+				TConstArrayView<RowHandle> RowsView = MakeArrayView(Rows, RowCount);
+				TArrayView<FTEDSProcessorTestsReferenceColumn> ReferenceColumnsView = MakeArrayView(ReferenceColumns, RowCount);
 
-			// Test creation of a row from within a query processor
-			QueryHandle PrimaryRowQuery = DataStorage->RegisterQuery(
-				Select(TEXT("TEST: Creating a row for primary reference column"),
-				FProcessor(TypedElementDataStorage::EQueryTickPhase::PrePhysics, DataStorage->GetQueryTickGroupName(TypedElementDataStorage::EQueryTickGroups::Default)),
-				[](IQueryContext& Context, const RowHandle* Rows, FTEDSProcessorTestsReferenceColumn* ReferenceColumns)
+				for (int32 Index = 0; Index < RowCount; ++Index)
 				{
-					const int32 RowCount = Context.GetRowCount();
-					TConstArrayView<RowHandle> RowsView = MakeArrayView(Rows, RowCount);
-					TArrayView<FTEDSProcessorTestsReferenceColumn> ReferenceColumnsView = MakeArrayView(ReferenceColumns, RowCount);
-
-					for (int32 Index = 0; Index < RowCount; ++Index)
+					FTEDSProcessorTestsReferenceColumn& PrimaryReferenceColumn = ReferenceColumnsView[Index];
+					// Auto-create a secondary row if this points to no row
+					if (!Context.IsRowAvailable(PrimaryReferenceColumn.Reference))
 					{
-						FTEDSProcessorTestsReferenceColumn& PrimaryReferenceColumn = ReferenceColumnsView[Index];
-						// Auto-create a secondary row if this points to no row
-						if (!Context.IsRowAvailable(PrimaryReferenceColumn.Reference))
-						{
-							RowHandle SecondaryRow = Context.AddRow(SecondaryTable);
-							// Initialize bi-directional row references
-							PrimaryReferenceColumn.Reference = SecondaryRow;
-							Context.AddColumn(SecondaryRow, FTEDSProcessorTestsReferenceColumn{.Reference = RowsView[Index]});
-						}
+						RowHandle SecondaryRow = Context.AddRow(SecondaryTable);
+						// Initialize bi-directional row references
+						PrimaryReferenceColumn.Reference = SecondaryRow;
+						Context.AddColumn(SecondaryRow, FTEDSProcessorTestsReferenceColumn{.Reference = RowsView[Index]});
 					}
 				}
-				)
-				.Where()
-					.All<FTEDSProcessorTests_PrimaryTag>()
-					.None<FTEDSProcessorTests_Linked>()
-				.Compile());
-
-			RegisteredQueries.Emplace(PrimaryRowQuery);
-
-			QueryHandle UpdateTransformWidget = DataStorage->RegisterQuery(
-				Select()
-					.ReadOnly<FTEDSProcessorTestsReferenceColumn>()
-				.Where()
-					.All<FTEDSProcessorTests_PrimaryTag>()
-					.None<FTEDSProcessorTests_Linked>()
-				.Compile());
-		
-			RegisteredQueries.Emplace(UpdateTransformWidget);
-
-			QueryHandle SecondaryRowQuery = DataStorage->RegisterQuery(
-				Select(TEXT("TEST: Creating a row for secondary reference column"),
-				FProcessor(TypedElementDataStorage::EQueryTickPhase::DuringPhysics, DataStorage->GetQueryTickGroupName(TypedElementDataStorage::EQueryTickGroups::Default)),
-				[](IQueryContext& Context, const RowHandle* Rows, FTEDSProcessorTestsReferenceColumn* ReferenceColumns)
-				{
-					const int32 RowCount = Context.GetRowCount();
-					TConstArrayView<RowHandle> RowsView = MakeArrayView(Rows, RowCount);
-					TArrayView<FTEDSProcessorTestsReferenceColumn> ReferenceColumnsView = MakeArrayView(ReferenceColumns, RowCount);
+			}
+			)
+			.Where()
+				.All<FTEDSProcessorTests_PrimaryTag>()
+				.None<FTEDSProcessorTests_Linked>()
+			.Compile());
+
+		RegisteredQueries.Emplace(PrimaryRowQuery);
+
+		QueryHandle UpdateTransformWidget = DataStorage->RegisterQuery(
+			Select()
+				.ReadOnly<FTEDSProcessorTestsReferenceColumn>()
+			.Where()
+				.All<FTEDSProcessorTests_PrimaryTag>()
+				.None<FTEDSProcessorTests_Linked>()
+			.Compile());
+	
+		RegisteredQueries.Emplace(UpdateTransformWidget);
 
-					for (int32 Index = 0; Index < RowCount; ++Index)
-					{
-						const RowHandle& SecondaryRow = RowsView[Index];
-						Context.RunSubquery(
-							0,
-							ReferenceColumnsView[Index].Reference,
-							CreateSubqueryCallbackBinding(
-							[SecondaryRow](ISubqueryContext& SubqueryContext, RowHandle PrimaryRow, const FTEDSProcessorTestsReferenceColumn& ReferenceColumn)
+		QueryHandle SecondaryRowQuery = DataStorage->RegisterQuery(
+			Select(TEXT("TEST: Creating a row for secondary reference column"),
+			FProcessor(EQueryTickPhase::DuringPhysics, DataStorage->GetQueryTickGroupName(EQueryTickGroups::Default)),
+			[](IQueryContext& Context, const RowHandle* Rows, FTEDSProcessorTestsReferenceColumn* ReferenceColumns)
+			{
+				const int32 RowCount = Context.GetRowCount();
+				TConstArrayView<RowHandle> RowsView = MakeArrayView(Rows, RowCount);
+				TArrayView<FTEDSProcessorTestsReferenceColumn> ReferenceColumnsView = MakeArrayView(ReferenceColumns, RowCount);
+
+				for (int32 Index = 0; Index < RowCount; ++Index)
+				{
+					const RowHandle& SecondaryRow = RowsView[Index];
+					Context.RunSubquery(
+						0,
+						ReferenceColumnsView[Index].Reference,
+						CreateSubqueryCallbackBinding(
+						[SecondaryRow](ISubqueryContext& SubqueryContext, RowHandle PrimaryRow, const FTEDSProcessorTestsReferenceColumn& ReferenceColumn)
+						{
+							if (ReferenceColumn.Reference == SecondaryRow)
 							{
-								if (ReferenceColumn.Reference == SecondaryRow)
-								{
-									// Add these tags to prevent further processing
-									SubqueryContext.AddColumns<FTEDSProcessorTests_Linked>(SecondaryRow);
-									SubqueryContext.AddColumns<FTEDSProcessorTests_Linked>(PrimaryRow);
-								}
+								// Add these tags to prevent further processing
+								SubqueryContext.AddColumns<FTEDSProcessorTests_Linked>(SecondaryRow);
+								SubqueryContext.AddColumns<FTEDSProcessorTests_Linked>(PrimaryRow);
 							}
-						));
-					}
+						}
+					));
 				}
-				)
-				.Where()
-					.All<FTEDSProcessorTests_SecondaryTag>()
-					.None<FTEDSProcessorTests_Linked>()
-				.DependsOn()
-					.SubQuery(UpdateTransformWidget)
-				.Compile());
-
-			RegisteredQueries.Emplace(SecondaryRowQuery);
+			}
+			)
+			.Where()
+				.All<FTEDSProcessorTests_SecondaryTag>()
+				.None<FTEDSProcessorTests_Linked>()
+			.DependsOn()
+				.SubQuery(UpdateTransformWidget)
+			.Compile());
+
+		RegisteredQueries.Emplace(SecondaryRowQuery);
+	}
+
+	void UnregisterProcessors()
+	{
+		ITypedElementDataStorageInterface* DataStorage = UTypedElementRegistry::GetInstance()->GetMutableDataStorage();
+		for (const UE::Editor::DataStorage::QueryHandle& Handle : RegisteredQueries)
+		{
+			DataStorage->UnregisterQuery(Handle);
 		}
+		RegisteredQueries.Empty();
+	}
 
-		void UnregisterProcessors()
+	void RegisterCommands()
+	{
+		IConsoleCommand* Command = IConsoleManager::Get().RegisterConsoleCommand(TEXT("TEDS.Tests.ProcessorTests.AddPrimaryRows"), TEXT(""), FConsoleCommandWithArgsDelegate::CreateLambda([](const TArray<FString>& Args)
 		{
-			ITypedElementDataStorageInterface* DataStorage = UTypedElementRegistry::GetInstance()->GetMutableDataStorage();
-			for (const UE::Editor::DataStorage::QueryHandle& Handle : RegisteredQueries)
+			int32 RowsToCreate = 0;
+			if (Args.Num() != 1)
 			{
-				DataStorage->UnregisterQuery(Handle);
+				return;
+			}
+			TTypeFromString<int32>::FromString(RowsToCreate, *Args[0]);
+			if (RowsToCreate <= 0)
+			{
+				return;
 			}
-			RegisteredQueries.Empty();
-		}
 
-		void RegisterCommands()
-		{
-			IConsoleCommand* Command = IConsoleManager::Get().RegisterConsoleCommand(TEXT("TEDS.Tests.ProcessorTests.AddPrimaryRows"), TEXT(""), FConsoleCommandWithArgsDelegate::CreateLambda([](const TArray<FString>& Args)
+			ITypedElementDataStorageInterface* DataStorage = UTypedElementRegistry::GetInstance()->GetMutableDataStorage();
+			DataStorage->BatchAddRow(PrimaryTable, RowsToCreate, [](UE::Editor::DataStorage::RowHandle Row)
 			{
-				int32 RowsToCreate = 0;
-				if (Args.Num() != 1)
-				{
-					return;
-				}
-				TTypeFromString<int32>::FromString(RowsToCreate, *Args[0]);
-				if (RowsToCreate <= 0)
-				{
-					return;
-				}
+			});
+		}), ECVF_Default);
+	
+		RegisteredCommands.Emplace(Command);
+	}
 
-				ITypedElementDataStorageInterface* DataStorage = UTypedElementRegistry::GetInstance()->GetMutableDataStorage();
-				DataStorage->BatchAddRow(PrimaryTable, RowsToCreate, [](UE::Editor::DataStorage::RowHandle Row)
-				{
-				});
-			}), ECVF_Default);
-		
-			RegisteredCommands.Emplace(Command);
+	void UnregisterCommands()
+	{
+		for (IConsoleCommand* Command : RegisteredCommands)
+		{
+			IConsoleManager::Get().UnregisterConsoleObject(Command);
 		}
+		RegisteredCommands.Empty();
+	}
 
-		void UnregisterCommands()
+	void OnProcessorTestsEnabled(IConsoleVariable* Variable)
+	{
+		if (Variable->GetBool())
 		{
-			for (IConsoleCommand* Command : RegisteredCommands)
-			{
-				IConsoleManager::Get().UnregisterConsoleObject(Command);
-			}
-			RegisteredCommands.Empty();
+			RegisterProcessors();
+			RegisterCommands();
 		}
-	
-		void OnProcessorTestsEnabled(IConsoleVariable* Variable)
+		else
 		{
-			if (Variable->GetBool())
-			{
-				RegisterProcessors();
-				RegisterCommands();
-			}
-			else
-			{
-				UnregisterProcessors();
-				UnregisterCommands();
-			}
-		};
-	} // namespace Debug::ProcessorTests
-} // namespace UE::Editor::DataStorage
+			UnregisterProcessors();
+			UnregisterCommands();
+		}
+	};
+} // namespace UE::Editor::DataStorage::Debug::ProcessorTests
 #endif
 
