[Teds] Rename TypedElementQueryBuilder namespace to UE::Editor::DataStorage::Queries and adjust alignment and whitespace of queries.

#rnx

diff --git a/Engine/Plugins/Developer/PerforceSourceControl/Source/PerforceSourceControl/Private/PerforceSourceControlOperations.cpp b/Engine/Plugins/Developer/PerforceSourceControl/Source/PerforceSourceControl/Private/PerforceSourceControlOperations.cpp
--- a/Engine/Plugins/Developer/PerforceSourceControl/Source/PerforceSourceControl/Private/PerforceSourceControlOperations.cpp
+++ b/Engine/Plugins/Developer/PerforceSourceControl/Source/PerforceSourceControl/Private/PerforceSourceControlOperations.cpp
@@ -2020,8 +2020,7 @@ bool FPerforceUpdateStatusWorker::UpdateStates() const
 
 	auto UpdateDataStorage = [](const FPerforceSourceControlState& State)
 	{
-		using namespace TypedElementQueryBuilder;
-		using namespace UE::Editor;
+		using namespace UE::Editor::DataStorage::Queries;
 
 		UTypedElementRegistry* Registry = UTypedElementRegistry::GetInstance();
 		if (!Registry)
@@ -2035,25 +2034,25 @@ bool FPerforceUpdateStatusWorker::UpdateStates() const
 			return;
 		}
 
-		auto GetRevisionControlRow = [DataStorage](const FString& InFilename) -> DataStorage::RowHandle
+		auto GetRevisionControlRow = [DataStorage](const FString& InFilename) -> RowHandle
 		{
 			FString Filename = FPaths::SetExtension(InFilename, "");
 			FPaths::NormalizeFilename(Filename);
 			Filename = FPaths::ConvertRelativePathToFull(Filename);
 			
 			uint64 Index = GenerateIndexHash(Filename);
-			DataStorage::RowHandle Row = DataStorage->FindIndexedRow(Index);
+			RowHandle Row = DataStorage->FindIndexedRow(Index);
 
 			if (!DataStorage->IsRowAvailable(Row))
 			{
-				static DataStorage::RowHandle Table = DataStorage->FindTable(FName("Editor_RevisionControlTable"));
+				static RowHandle Table = DataStorage->FindTable(FName("Editor_RevisionControlTable"));
 				Row = DataStorage->AddRow(Table);
 				DataStorage->IndexRow(Index, Row);
 			}
 			return Row;
 		};
 
-		DataStorage::RowHandle Row = GetRevisionControlRow(State.GetFilename());
+		RowHandle Row = GetRevisionControlRow(State.GetFilename());
 
 		if (!State.IsSourceControlled())
 		{
@@ -2124,10 +2123,10 @@ bool FPerforceUpdateStatusWorker::UpdateStates() const
 		
         if (FTypedElementPackageReference* BackReference = DataStorage->GetColumn<FTypedElementPackageReference>(Row))
         {
-        	if (DataStorage::RowHandle ObjectRow = BackReference->Row; ObjectRow != InvalidRowHandle)
+        	if (RowHandle ObjectRow = BackReference->Row; ObjectRow != InvalidRowHandle)
         	{
         		static TableHandle Table = DataStorage->FindTable(FName("Editor_PackageUpdateTable"));
-        		DataStorage::RowHandle UpdateRow = DataStorage->AddRow(Table);
+        		RowHandle UpdateRow = DataStorage->AddRow(Table);
         		DataStorage->AddColumn(UpdateRow, FTypedElementPackageUpdateColumn{ .ObjectRow = ObjectRow, .PackageRow = Row });
         	}
         }
diff --git a/Engine/Plugins/Experimental/Animation/HierarchyTable/Source/Editor/Private/HierarchyTableEditorToolkit.cpp b/Engine/Plugins/Experimental/Animation/HierarchyTable/Source/Editor/Private/HierarchyTableEditorToolkit.cpp
--- a/Engine/Plugins/Experimental/Animation/HierarchyTable/Source/Editor/Private/HierarchyTableEditorToolkit.cpp
+++ b/Engine/Plugins/Experimental/Animation/HierarchyTable/Source/Editor/Private/HierarchyTableEditorToolkit.cpp
@@ -112,8 +112,7 @@ TSharedRef<SWidget> FHierarchyTableEditorToolkit::CreateTedsOutliner()
 			.Text(INVTEXT("You need to enable the Typed Element Data Storage plugin to see the table viewer!"));
 	}
 
-	using namespace TypedElementQueryBuilder;
-	using namespace UE::Editor::DataStorage;
+	using namespace UE::Editor::DataStorage::Queries;
 
 	if (!ensure(HierarchyTable->TableType))
 	{
@@ -132,7 +131,7 @@ TSharedRef<SWidget> FHierarchyTableEditorToolkit::CreateTedsOutliner()
 
 	FQueryDescription ColumnQueryDescription =
 		Select()
-		.ReadOnly(HierarchyTableTypeColumns)
+			.ReadOnly(HierarchyTableTypeColumns)
 		.Compile();
 
 	InitialColumnQuery = Registry->GetMutableDataStorage()->RegisterQuery(MoveTemp(ColumnQueryDescription));
diff --git a/Engine/Plugins/Experimental/EditorDataStorage/Source/TedsCore/Private/Commands/EditorDataStorageCompatibilityCommands.cpp b/Engine/Plugins/Experimental/EditorDataStorage/Source/TedsCore/Private/Commands/EditorDataStorageCompatibilityCommands.cpp
--- a/Engine/Plugins/Experimental/EditorDataStorage/Source/TedsCore/Private/Commands/EditorDataStorageCompatibilityCommands.cpp
+++ b/Engine/Plugins/Experimental/EditorDataStorage/Source/TedsCore/Private/Commands/EditorDataStorageCompatibilityCommands.cpp
@@ -525,7 +525,7 @@ namespace UE::Editor::DataStorage
 
 	void FCommandProcessor::operator()(FTypeBatchInfoReinstanced& Command)
 	{
-		using namespace TypedElementQueryBuilder;
+		using namespace UE::Editor::DataStorage::Queries;
 
 		StorageCompatibility.Storage->RunQuery(StorageCompatibility.ClassTypeInfoQuery, CreateDirectQueryCallbackBinding(
 			[Range = Command.Batch](IDirectQueryContext& Context, FTypedElementClassTypeInfoColumn& Type)
diff --git a/Engine/Plugins/Experimental/EditorDataStorage/Source/TedsCore/Private/Compatibility/TedsTypedElementActorHandleFactory.cpp b/Engine/Plugins/Experimental/EditorDataStorage/Source/TedsCore/Private/Compatibility/TedsTypedElementActorHandleFactory.cpp
--- a/Engine/Plugins/Experimental/EditorDataStorage/Source/TedsCore/Private/Compatibility/TedsTypedElementActorHandleFactory.cpp
+++ b/Engine/Plugins/Experimental/EditorDataStorage/Source/TedsCore/Private/Compatibility/TedsTypedElementActorHandleFactory.cpp
@@ -33,7 +33,7 @@ void UTypedElementActorHandleDataStorageFactory::RegisterQueries(ITypedElementDa
 		RegisterQuery_ActorHandlePopulate(DataStorage);
 	}
 
-	using namespace TypedElementQueryBuilder;
+	using namespace UE::Editor::DataStorage::Queries;
 	GetAllActorsQuery = DataStorage.RegisterQuery(
 	Select()
 		.ReadOnly<FTypedElementUObjectColumn>()
@@ -44,8 +44,7 @@ void UTypedElementActorHandleDataStorageFactory::RegisterQueries(ITypedElementDa
 
 void UTypedElementActorHandleDataStorageFactory::RegisterQuery_ActorHandlePopulate(ITypedElementDataStorageInterface& DataStorage)
 {
-	using namespace TypedElementQueryBuilder;
-	using namespace UE::Editor::DataStorage;
+	using namespace UE::Editor::DataStorage::Queries;
 
 	if (!ensureMsgf(ActorHandlePopulateQuery == InvalidQueryHandle, TEXT("Already registered query")))
 	{
@@ -66,8 +65,7 @@ void UTypedElementActorHandleDataStorageFactory::RegisterQuery_ActorHandlePopula
 					.Handle = Handle
 				});
 			}
-		}
-	)
+		})
 	.Where()
 		.All<FTypedElementActorTag>()
 	.Compile());
@@ -81,7 +79,7 @@ void UTypedElementActorHandleDataStorageFactory::HandleBridgeEnabled(bool bEnabl
 
 	if (bEnabled)
 	{
-		using namespace TypedElementQueryBuilder;
+		using namespace Queries;
 		
 		// Populate all the rows
 		TArray<RowHandle> CollatedRowHandles;
diff --git a/Engine/Plugins/Experimental/EditorDataStorage/Source/TedsCore/Private/Compatibility/TedsTypedElementBridgeQueries.cpp b/Engine/Plugins/Experimental/EditorDataStorage/Source/TedsCore/Private/Compatibility/TedsTypedElementBridgeQueries.cpp
--- a/Engine/Plugins/Experimental/EditorDataStorage/Source/TedsCore/Private/Compatibility/TedsTypedElementBridgeQueries.cpp
+++ b/Engine/Plugins/Experimental/EditorDataStorage/Source/TedsCore/Private/Compatibility/TedsTypedElementBridgeQueries.cpp
@@ -57,7 +57,7 @@ bool UTypedElementBridgeDataStorageFactory::IsEnabled()
 
 void UTypedElementBridgeDataStorageFactory::RegisterQuery_NewUObject(ITypedElementDataStorageInterface& DataStorage)
 {
-	using namespace TypedElementQueryBuilder;
+	using namespace UE::Editor::DataStorage::Queries;
 	
 	RemoveTypedElementRowHandleQuery = DataStorage.RegisterQuery(
 		Select()
@@ -71,13 +71,11 @@ void UTypedElementBridgeDataStorageFactory::UnregisterQuery_NewUObject(ITypedEle
 
 void UTypedElementBridgeDataStorageFactory::CleanupTypedElementColumns(ITypedElementDataStorageInterface& DataStorage)
 {
-	using namespace TypedElementQueryBuilder;
-	using namespace UE::Editor::DataStorage;
+	using namespace UE::Editor::DataStorage::Queries;
 
 	// Remove any TEv1 handles
 	{
 		TArray<RowHandle> Handles;
-		using namespace TypedElementQueryBuilder;
 		DataStorage.RunQuery(
 			RemoveTypedElementRowHandleQuery,
 			CreateDirectQueryCallbackBinding(
diff --git a/Engine/Plugins/Experimental/EditorDataStorage/Source/TedsCore/Private/Compatibility/TypedElementActorIconOverrideQueries.cpp b/Engine/Plugins/Experimental/EditorDataStorage/Source/TedsCore/Private/Compatibility/TypedElementActorIconOverrideQueries.cpp
--- a/Engine/Plugins/Experimental/EditorDataStorage/Source/TedsCore/Private/Compatibility/TypedElementActorIconOverrideQueries.cpp
+++ b/Engine/Plugins/Experimental/EditorDataStorage/Source/TedsCore/Private/Compatibility/TypedElementActorIconOverrideQueries.cpp
@@ -10,8 +10,7 @@
 
 void UActorIconOverrideDataStorageFactory::RegisterQueries(ITypedElementDataStorageInterface& DataStorage)
 {
-	using namespace UE::Editor::DataStorage;
-	using namespace TypedElementQueryBuilder;
+	using namespace UE::Editor::DataStorage::Queries;
 
 	DataStorage.RegisterQuery(
 		Select(
@@ -27,13 +26,11 @@ void UActorIconOverrideDataStorageFactory::RegisterQueries(ITypedElementDataStor
 						Context.AddColumn(Row, FTypedElementIconOverrideColumn{ .IconName = IconName });
 					}
 				}
-			}
-		)
+			})
 		.Where()
 			.All<FTypedElementSyncFromWorldTag, FTypedElementActorTag>()
 			.None<FTypedElementIconOverrideColumn>()
-		.Compile()
-	);
+		.Compile());
 
 	DataStorage.RegisterQuery(
 		Select(
@@ -56,10 +53,8 @@ void UActorIconOverrideDataStorageFactory::RegisterQueries(ITypedElementDataStor
 						Context.RemoveColumns<FTypedElementIconOverrideColumn>(Row);
 					}
 				}
-			}
-		)
+			})
 		.Where()
 			.All<FTypedElementSyncFromWorldTag, FTypedElementActorTag>()
-		.Compile()
-	);
+		.Compile());
 }
diff --git a/Engine/Plugins/Experimental/EditorDataStorage/Source/TedsCore/Private/Compatibility/TypedElementActorLabelQueries.cpp b/Engine/Plugins/Experimental/EditorDataStorage/Source/TedsCore/Private/Compatibility/TypedElementActorLabelQueries.cpp
--- a/Engine/Plugins/Experimental/EditorDataStorage/Source/TedsCore/Private/Compatibility/TypedElementActorLabelQueries.cpp
+++ b/Engine/Plugins/Experimental/EditorDataStorage/Source/TedsCore/Private/Compatibility/TypedElementActorLabelQueries.cpp
@@ -21,8 +21,7 @@ void UActorLabelDataStorageFactory::RegisterQueries(ITypedElementDataStorageInte
 
 void UActorLabelDataStorageFactory::RegisterActorLabelToColumnQuery(ITypedElementDataStorageInterface& DataStorage) const
 {
-	using namespace TypedElementQueryBuilder;
-	using namespace UE::Editor::DataStorage;
+	using namespace UE::Editor::DataStorage::Queries;
 	
 	DataStorage.RegisterQuery(
 		Select(
@@ -45,14 +44,12 @@ void UActorLabelDataStorageFactory::RegisterActorLabelToColumnQuery(ITypedElemen
 		)
 		.Where()
 			.All<FTypedElementActorTag, FTypedElementSyncFromWorldTag>()
-		.Compile()
-	);
+		.Compile());
 }
 
 void UActorLabelDataStorageFactory::RegisterLabelColumnToActorQuery(ITypedElementDataStorageInterface& DataStorage) const
 {
-	using namespace TypedElementQueryBuilder;
-	using namespace UE::Editor::DataStorage;
+	using namespace UE::Editor::DataStorage::Queries;
 	
 	DataStorage.RegisterQuery(
 		Select(
@@ -75,8 +72,7 @@ void UActorLabelDataStorageFactory::RegisterLabelColumnToActorQuery(ITypedElemen
 		)
 		.Where()
 			.All<FTypedElementActorTag, FTypedElementSyncBackToWorldTag>()
-		.Compile()
-	);
+		.Compile());
 }
 
 #undef LOCTEXT_NAMESPACE
diff --git a/Engine/Plugins/Experimental/EditorDataStorage/Source/TedsCore/Private/Compatibility/TypedElementActorParentQueries.cpp b/Engine/Plugins/Experimental/EditorDataStorage/Source/TedsCore/Private/Compatibility/TypedElementActorParentQueries.cpp
--- a/Engine/Plugins/Experimental/EditorDataStorage/Source/TedsCore/Private/Compatibility/TypedElementActorParentQueries.cpp
+++ b/Engine/Plugins/Experimental/EditorDataStorage/Source/TedsCore/Private/Compatibility/TypedElementActorParentQueries.cpp
@@ -31,8 +31,7 @@ void UActorParentDataStorageFactory::RegisterQueries(ITypedElementDataStorageInt
 
 void UActorParentDataStorageFactory::RegisterAddParentColumn(ITypedElementDataStorageInterface& DataStorage) const
 {
-	using namespace TypedElementQueryBuilder;
-	using namespace UE::Editor::DataStorage;
+	using namespace UE::Editor::DataStorage::Queries;
 	
 	DataStorage.RegisterQuery(
 		Select(
@@ -57,19 +56,16 @@ void UActorParentDataStorageFactory::RegisterAddParentColumn(ITypedElementDataSt
 						}
 					}
 				}
-			}
-		)
+			})
 		.Where()
 			.All<FTypedElementSyncFromWorldTag, FTypedElementActorTag>()
 			.None<FTableRowParentColumn, FUnresolvedTableRowParentColumn>()
-		.Compile()
-	);
+		.Compile());
 }
 
 void UActorParentDataStorageFactory::RegisterUpdateOrRemoveParentColumn(ITypedElementDataStorageInterface& DataStorage) const
 {
-	using namespace TypedElementQueryBuilder;
-	using namespace UE::Editor::DataStorage;
+	using namespace UE::Editor::DataStorage::Queries;
 	
 	DataStorage.RegisterQuery(
 		Select(
@@ -101,10 +97,8 @@ void UActorParentDataStorageFactory::RegisterUpdateOrRemoveParentColumn(ITypedEl
 					}
 				}
 				Context.RemoveColumns<FTableRowParentColumn>(Row);
-			}
-		)
+			})
 		.Where()
 			.All<FTypedElementActorTag, FTypedElementSyncFromWorldTag>()
-		.Compile()
-	);
+		.Compile());
 }
diff --git a/Engine/Plugins/Experimental/EditorDataStorage/Source/TedsCore/Private/Compatibility/TypedElementActorTransformQueries.cpp b/Engine/Plugins/Experimental/EditorDataStorage/Source/TedsCore/Private/Compatibility/TypedElementActorTransformQueries.cpp
--- a/Engine/Plugins/Experimental/EditorDataStorage/Source/TedsCore/Private/Compatibility/TypedElementActorTransformQueries.cpp
+++ b/Engine/Plugins/Experimental/EditorDataStorage/Source/TedsCore/Private/Compatibility/TypedElementActorTransformQueries.cpp
@@ -17,14 +17,13 @@ void UActorTransformDataStorageFactory::RegisterQueries(ITypedElementDataStorage
 
 void UActorTransformDataStorageFactory::RegisterActorAddTransformColumn(ITypedElementDataStorageInterface& DataStorage) const
 {
-	using namespace TypedElementQueryBuilder;
-	using namespace UE::Editor::DataStorage;
+	using namespace UE::Editor::DataStorage::Queries;
+
 	
 	DataStorage.RegisterQuery(
 		Select(
 			TEXT("Add transform column to actor"),
-			FProcessor(EQueryTickPhase::PrePhysics,
-				DataStorage.GetQueryTickGroupName(EQueryTickGroups::SyncExternalToDataStorage))
+			FProcessor(EQueryTickPhase::PrePhysics, DataStorage.GetQueryTickGroupName(EQueryTickGroups::SyncExternalToDataStorage))
 				.SetExecutionMode(EExecutionMode::GameThread),
 			[](IQueryContext& Context, RowHandle Row, const FTypedElementUObjectColumn& Actor)
 			{
@@ -32,19 +31,17 @@ void UActorTransformDataStorageFactory::RegisterActorAddTransformColumn(ITypedEl
 				{
 					Context.AddColumn(Row, FTypedElementLocalTransformColumn{ .Transform = ActorInstance->GetActorTransform() });
 				}
-			}
-		)
+			})
 		.Where()
 			.All<FTypedElementSyncFromWorldTag, FTypedElementActorTag>()
 			.None<FTypedElementLocalTransformColumn>()
-		.Compile()
-	);
+		.Compile());
 }
 
 void UActorTransformDataStorageFactory::RegisterActorLocalTransformToColumn(ITypedElementDataStorageInterface& DataStorage) const
 {
-	using namespace TypedElementQueryBuilder;
-	using namespace UE::Editor::DataStorage;
+	using namespace UE::Editor::DataStorage::Queries;
+
 	
 	DataStorage.RegisterQuery(
 		Select(
@@ -61,19 +58,16 @@ void UActorTransformDataStorageFactory::RegisterActorLocalTransformToColumn(ITyp
 				{
 					Context.RemoveColumns<FTypedElementLocalTransformColumn>(Row);
 				}
-			}
-		)
+			})
 		.Where()
 			.All<FTypedElementActorTag>()
 			.Any<FTypedElementSyncFromWorldTag, FTypedElementSyncFromWorldInteractiveTag>()
-		.Compile()
-	);
+		.Compile());
 }
 
 void UActorTransformDataStorageFactory::RegisterLocalTransformColumnToActor(ITypedElementDataStorageInterface& DataStorage) const
 {
-	using namespace TypedElementQueryBuilder;
-	using namespace UE::Editor::DataStorage;
+	using namespace UE::Editor::DataStorage::Queries;
 	
 	DataStorage.RegisterQuery(
 		Select(
@@ -86,10 +80,8 @@ void UActorTransformDataStorageFactory::RegisterLocalTransformColumnToActor(ITyp
 				{
 					ActorInstance->SetActorTransform(Transform.Transform);
 				}
-			}
-		)
+			})
 		.Where()
 			.All<FTypedElementActorTag, FTypedElementSyncBackToWorldTag>()
-		.Compile()
-	);
+		.Compile());
 }
diff --git a/Engine/Plugins/Experimental/EditorDataStorage/Source/TedsCore/Private/Compatibility/TypedElementActorViewportProcessors.cpp b/Engine/Plugins/Experimental/EditorDataStorage/Source/TedsCore/Private/Compatibility/TypedElementActorViewportProcessors.cpp
--- a/Engine/Plugins/Experimental/EditorDataStorage/Source/TedsCore/Private/Compatibility/TypedElementActorViewportProcessors.cpp
+++ b/Engine/Plugins/Experimental/EditorDataStorage/Source/TedsCore/Private/Compatibility/TypedElementActorViewportProcessors.cpp
@@ -19,8 +19,7 @@ namespace UE::Editor::DataStorage::Private
 		TEXT("Adds an outline color to selected objects."),
 		FConsoleCommandWithArgsAndOutputDeviceDelegate::CreateLambda([](const TArray<FString>& Args, FOutputDevice& Output)
 			{
-				using namespace TypedElementQueryBuilder;
-				using namespace UE::Editor::DataStorage;
+				using namespace UE::Editor::DataStorage::Queries;
 
 				TRACE_CPUPROFILER_EVENT_SCOPE(TEDS.Debug.AddOverlayColorToSelectionCommand);
 
@@ -78,8 +77,7 @@ namespace UE::Editor::DataStorage::Private
 		TEXT("Adds an overlay color to selected objects."),
 		FConsoleCommandWithArgsAndOutputDeviceDelegate::CreateLambda([](const TArray<FString>& Args, FOutputDevice& Output)
 			{
-				using namespace TypedElementQueryBuilder;
-				using namespace UE::Editor::DataStorage;
+				using namespace UE::Editor::DataStorage::Queries;
 
 				TRACE_CPUPROFILER_EVENT_SCOPE(TEDS.Debug.AddOverlayColorToSelectionCommand);
 
@@ -130,8 +128,7 @@ namespace UE::Editor::DataStorage::Private
 		TEXT("Removes an overlay color to selected objects."),
 		FConsoleCommandWithArgsAndOutputDeviceDelegate::CreateLambda([](const TArray<FString>& Args, FOutputDevice& Output)
 			{
-				using namespace TypedElementQueryBuilder;
-				using namespace UE::Editor::DataStorage;
+				using namespace UE::Editor::DataStorage::Queries;
 
 				TRACE_CPUPROFILER_EVENT_SCOPE(TEDS.Debug.AddOverlayColorToSelectionCommand);
 
@@ -175,8 +172,7 @@ void UActorViewportDataStorageFactory::RegisterQueries(ITypedElementDataStorageI
 
 void UActorViewportDataStorageFactory::RegisterOutlineColorColumnToActor(ITypedElementDataStorageInterface& DataStorage)
 {
-	using namespace TypedElementQueryBuilder;
-	using namespace UE::Editor::DataStorage;
+	using namespace UE::Editor::DataStorage::Queries;
 
 	DataStorage.RegisterQuery(
 		Select(
@@ -193,18 +189,15 @@ void UActorViewportDataStorageFactory::RegisterOutlineColorColumnToActor(ITypedE
 						PrimitiveComponent->SetSelectionOutlineColorIndex(ViewportColor.SelectionOutlineColorIndex);
 					});
 				}
-			}
-		)
+			})
 		.Where()
 			.All<FTypedElementActorTag, FTypedElementSyncBackToWorldTag>()
-		.Compile()
-	);
+		.Compile());
 }
 
 void UActorViewportDataStorageFactory::RegisterOverlayColorColumnToActor(ITypedElementDataStorageInterface& DataStorage)
 {
-	using namespace TypedElementQueryBuilder;
-	using namespace UE::Editor::DataStorage;
+	using namespace UE::Editor::DataStorage::Queries;
 	
 	DataStorage.RegisterQuery(
 		Select(
@@ -222,12 +215,10 @@ void UActorViewportDataStorageFactory::RegisterOverlayColorColumnToActor(ITypedE
 						PrimitiveComponent->MarkRenderStateDirty();
 					});
 				}
-			}
-		)
+			})
 		.Where()
 			.All<FTypedElementActorTag>()
-		.Compile()
-	);
+		.Compile());
 
 	DataStorage.RegisterQuery(
 		Select(
@@ -244,10 +235,8 @@ void UActorViewportDataStorageFactory::RegisterOverlayColorColumnToActor(ITypedE
 						PrimitiveComponent->MarkRenderStateDirty();
 					});
 				}
-			}
-		)
+			})
 		.Where()
 			.All<FTypedElementActorTag>()
-		.Compile()
-	);
+		.Compile());
 }
diff --git a/Engine/Plugins/Experimental/EditorDataStorage/Source/TedsCore/Private/Compatibility/TypedElementUObjectWorldQueries.cpp b/Engine/Plugins/Experimental/EditorDataStorage/Source/TedsCore/Private/Compatibility/TypedElementUObjectWorldQueries.cpp
--- a/Engine/Plugins/Experimental/EditorDataStorage/Source/TedsCore/Private/Compatibility/TypedElementUObjectWorldQueries.cpp
+++ b/Engine/Plugins/Experimental/EditorDataStorage/Source/TedsCore/Private/Compatibility/TypedElementUObjectWorldQueries.cpp
@@ -15,8 +15,7 @@ void UObjectWorldDataStorageFactory::RegisterQueries(ITypedElementDataStorageInt
 
 void UObjectWorldDataStorageFactory::RegisterAddWorldColumn(ITypedElementDataStorageInterface& DataStorage) const
 {
-	using namespace TypedElementQueryBuilder;
-	using namespace UE::Editor::DataStorage;
+	using namespace UE::Editor::DataStorage::Queries;
 	
 	DataStorage.RegisterQuery(
 		Select(
@@ -34,19 +33,16 @@ void UObjectWorldDataStorageFactory::RegisterAddWorldColumn(ITypedElementDataSto
 						Context.AddColumn(Row, FTypedElementWorldColumn{ .World = World });
 					}
 				}
-			}
-		)
+			})
 		.Where()
 			.All<FTypedElementSyncFromWorldTag>()
 			.None<FTypedElementWorldColumn, FTypedElementClassDefaultObjectTag>()
-		.Compile()
-	);
+		.Compile());
 }
 
 void UObjectWorldDataStorageFactory::RegisterUpdateOrRemoveWorldColumn(ITypedElementDataStorageInterface& DataStorage) const
 {
-	using namespace TypedElementQueryBuilder;
-	using namespace UE::Editor::DataStorage;
+	using namespace UE::Editor::DataStorage::Queries;
 	
 	DataStorage.RegisterQuery(
 		Select(
@@ -66,10 +62,8 @@ void UObjectWorldDataStorageFactory::RegisterUpdateOrRemoveWorldColumn(ITypedEle
 					}
 				}
 				Context.RemoveColumns<FTypedElementWorldColumn>(Row);
-			}
-		)
+			})
 		.Where()
 			.All<FTypedElementSyncFromWorldTag>()
-		.Compile()
-	);
+		.Compile());
 }
diff --git a/Engine/Plugins/Experimental/EditorDataStorage/Source/TedsCore/Private/Debug/TypedElementDatabaseDebugCommands.cpp b/Engine/Plugins/Experimental/EditorDataStorage/Source/TedsCore/Private/Debug/TypedElementDatabaseDebugCommands.cpp
--- a/Engine/Plugins/Experimental/EditorDataStorage/Source/TedsCore/Private/Debug/TypedElementDatabaseDebugCommands.cpp
+++ b/Engine/Plugins/Experimental/EditorDataStorage/Source/TedsCore/Private/Debug/TypedElementDatabaseDebugCommands.cpp
@@ -25,15 +25,12 @@ namespace UE::Editor::DataStorage::Private
 	template<typename TypeInfoType>
 	void PrintObjectTypeInformation(ITypedElementDataStorageInterface* DataStorage, FString Message, FOutputDevice& Output)
 	{
-		using namespace TypedElementQueryBuilder;
+		using namespace UE::Editor::DataStorage::Queries;
 
-		static QueryHandle Query = [DataStorage]
-		{
-			return DataStorage->RegisterQuery(
-				Select()
-					.ReadOnly<TypeInfoType>()
-				.Compile());
-		}();
+		static QueryHandle Query = DataStorage->RegisterQuery(
+			Select()
+				.ReadOnly<TypeInfoType>()
+			.Compile());
 
 		if (Query != InvalidQueryHandle)
 		{
@@ -66,7 +63,7 @@ namespace UE::Editor::DataStorage::Private
 	template<typename... Conditions>
 	void PrintObjectLabels(FOutputDevice& Output)
 	{
-		using namespace TypedElementQueryBuilder;
+		using namespace UE::Editor::DataStorage::Queries;
 
 		if (ITypedElementDataStorageInterface* DataStorage = UTypedElementRegistry::GetInstance()->GetMutableDataStorage())
 		{
@@ -407,8 +404,7 @@ static FAutoConsoleCommand CVarPrintDynamicColumnWithQuery(
 	TEXT("Argument: Identifier\n"),
 	FConsoleCommandWithArgsDelegate::CreateLambda([](const TArray<FString>& Args)
 	{
-		using namespace TypedElementQueryBuilder;
-		using namespace UE::Editor::DataStorage;
+		using namespace UE::Editor::DataStorage::Queries;
 		
 		// Print column using query
 		if (Args.Num() != 1)
@@ -468,8 +464,7 @@ static FAutoConsoleCommand CVarCountDynamicTagWithQuery(
 	TEXT("Argument: Identifier\n"),
 	FConsoleCommandWithArgsDelegate::CreateLambda([](const TArray<FString>& Args)
 	{
-		using namespace TypedElementQueryBuilder;
-		using namespace UE::Editor::DataStorage;
+		using namespace UE::Editor::DataStorage::Queries;
 				
 		// Print column using query
 		if (Args.Num() != 1)
@@ -510,8 +505,7 @@ static FAutoConsoleCommand CVarRegisterListDynamicColumnQuery(
 	TEXT("Argument: Identifier ActivationGroup\n"),
 	FConsoleCommandWithArgsDelegate::CreateLambda([](const TArray<FString>& Args)
 	{
-		using namespace TypedElementQueryBuilder;
-		using namespace UE::Editor::DataStorage;
+		using namespace UE::Editor::DataStorage::Queries;
 				
 		// Print column using query
 		if (Args.Num() != 2)
@@ -527,21 +521,21 @@ static FAutoConsoleCommand CVarRegisterListDynamicColumnQuery(
 
 		// Lists the rows processed that have 
 		DataStorage->RegisterQuery(
-				Select(
-					TEXT("ProcessDynamicTagColumns"),
-					FProcessor(EQueryTickPhase::FrameEnd, DataStorage->GetQueryTickGroupName(EQueryTickGroups::Default))
-						.MakeActivatable(ActivationGroup),
-					[](IQueryContext& Context, const RowHandle* Rows)
+			Select(
+				TEXT("ProcessDynamicTagColumns"),
+				FProcessor(EQueryTickPhase::FrameEnd, DataStorage->GetQueryTickGroupName(EQueryTickGroups::Default))
+					.MakeActivatable(ActivationGroup),
+				[](IQueryContext& Context, const RowHandle* Rows)
+				{
+					auto RowView = MakeConstArrayView(Rows, Context.GetRowCount());
+					for (RowHandle Row : RowView)
 					{
-						auto RowView = MakeConstArrayView(Rows, Context.GetRowCount());
-						for (RowHandle Row : RowView)
-						{
-							UE_LOG(LogEditorDataStorage, Warning, TEXT("- '%llu'\n"), Row);
-						}
-					})
-					.Where()
-						.All<FTestDynamicTag>(Identifier)
-				.Compile());
+						UE_LOG(LogEditorDataStorage, Warning, TEXT("- '%llu'\n"), Row);
+					}
+				})
+			.Where()
+				.All<FTestDynamicTag>(Identifier)
+			.Compile());
 
 		UE_LOG(LogEditorDataStorage, Warning, TEXT("Query registered for Dynamic Column FTestDynamicTag::%s with activation group '%s'"), *Identifier.ToString(), *ActivationGroup.ToString());
 	}),
@@ -637,8 +631,7 @@ static FAutoConsoleCommand CVarMatchValueTag(
 	TEXT("Argument: Tag, [optional] Value\n"),
 	FConsoleCommandWithArgsDelegate::CreateLambda([](const TArray<FString>& Args)
 	{
-		using namespace TypedElementQueryBuilder;
-		using namespace UE::Editor::DataStorage;
+		using namespace UE::Editor::DataStorage::Queries;
 		
 		ITypedElementDataStorageInterface* DataStorage = UTypedElementRegistry::GetInstance()->GetMutableDataStorage();
 
@@ -654,12 +647,12 @@ static FAutoConsoleCommand CVarMatchValueTag(
 			{
 				// Matches all rows with the 
 				return DataStorage->RegisterQuery(
-					Select().
-						Where().
-							// Match all rows with a value tag of type Tag (ie. all rows with a value tag of "Color")
-							All<FValueTag>(Tag).
-							All<FTestColumnA>().
-						Compile());
+					Select()
+					.Where()
+						// Match all rows with a value tag of type Tag (ie. all rows with a value tag of "Color")
+						.All<FValueTag>(Tag)
+						.All<FTestColumnA>()
+					.Compile());
 			}
 			else
 			{
@@ -752,8 +745,7 @@ static FAutoConsoleCommand CVarMatchValueTagFromEnum(
 	TEXT("Argument: [optional] EnumValue\n"),
 	FConsoleCommandWithArgsDelegate::CreateLambda([](const TArray<FString>& Args)
 	{
-		using namespace TypedElementQueryBuilder;
-		using namespace UE::Editor::DataStorage;
+		using namespace UE::Editor::DataStorage::Queries;
 		
 		ITypedElementDataStorageInterface* DataStorage = UTypedElementRegistry::GetInstance()->GetMutableDataStorage();
 
@@ -780,11 +772,11 @@ static FAutoConsoleCommand CVarMatchValueTagFromEnum(
 			{
 				// Matches all rows with the 
 				return DataStorage->RegisterQuery(
-					Select().
-						Where().
-							// Match all rows with an enum value tag of the hardcoded enum type
-							All<ETedsDebugEnum>().
-						Compile());
+					Select()
+					.Where()
+						// Match all rows with an enum value tag of the hardcoded enum type
+						.All<ETedsDebugEnum>()
+					.Compile());
 			}
 			else if (Args.Num() == 1)
 			{
diff --git a/Engine/Plugins/Experimental/EditorDataStorage/Source/TedsCore/Private/Queries/TypedElementAlertQueries.cpp b/Engine/Plugins/Experimental/EditorDataStorage/Source/TedsCore/Private/Queries/TypedElementAlertQueries.cpp
--- a/Engine/Plugins/Experimental/EditorDataStorage/Source/TedsCore/Private/Queries/TypedElementAlertQueries.cpp
+++ b/Engine/Plugins/Experimental/EditorDataStorage/Source/TedsCore/Private/Queries/TypedElementAlertQueries.cpp
@@ -17,8 +17,7 @@ FAutoConsoleCommand AddRandomAlertToRowConsoleCommand(
 	FConsoleCommandDelegate::CreateLambda(
 		[]()
 		{
-			using namespace TypedElementQueryBuilder;
-			using namespace UE::Editor::DataStorage;
+			using namespace UE::Editor::DataStorage::Queries;
 
 			TRACE_CPUPROFILER_EVENT_SCOPE(TEDS.Debug.AddRandomAlertToSelectedRows);
 
@@ -76,8 +75,7 @@ FAutoConsoleCommand ClearAllAlertsConsoleCommand(
 	FConsoleCommandDelegate::CreateLambda(
 		[]()
 		{
-			using namespace TypedElementQueryBuilder;
-			using namespace UE::Editor::DataStorage;
+			using namespace UE::Editor::DataStorage::Queries;
 
 			TRACE_CPUPROFILER_EVENT_SCOPE(TEDS.Debug.ClearAllAlertInfo);
 
@@ -90,9 +88,9 @@ FAutoConsoleCommand ClearAllAlertsConsoleCommand(
 				{
 					return DataStorage->RegisterQuery(
 						Select()
-							.Where()
-						.Any<FTypedElementAlertColumn, FTypedElementChildAlertColumn>()
-							.Compile());
+						.Where()
+							.Any<FTypedElementAlertColumn, FTypedElementChildAlertColumn>()
+						.Compile());
 				}();
 				TArray<RowHandle> Rows;
 				DataStorage->RunQuery(AlertInfoQuery, CreateDirectQueryCallbackBinding(
@@ -116,8 +114,7 @@ FAutoConsoleCommand ClearSelectedAlertsConsoleCommand(
 	FConsoleCommandDelegate::CreateLambda(
 		[]()
 		{
-			using namespace TypedElementQueryBuilder;
-			using namespace UE::Editor::DataStorage;
+			using namespace UE::Editor::DataStorage::Queries;
 
 			TRACE_CPUPROFILER_EVENT_SCOPE(TEDS.Debug.ClearSelectedAlerts);
 
@@ -159,7 +156,7 @@ void UTypedElementAlertQueriesFactory::RegisterQueries(ITypedElementDataStorageI
 
 void UTypedElementAlertQueriesFactory::RegisterSubQueries(ITypedElementDataStorageInterface& DataStorage)
 {
-	using namespace TypedElementQueryBuilder;
+	using namespace UE::Editor::DataStorage::Queries;
 	
 	ChildAlertColumnReadWriteQuery = DataStorage.RegisterQuery(
 		Select()
@@ -174,8 +171,7 @@ void UTypedElementAlertQueriesFactory::RegisterSubQueries(ITypedElementDataStora
 
 void UTypedElementAlertQueriesFactory::RegisterParentUpdatesQueries(ITypedElementDataStorageInterface& DataStorage)
 {
-	using namespace UE::Editor::DataStorage;
-	using namespace TypedElementQueryBuilder;
+	using namespace UE::Editor::DataStorage::Queries;
 
 	DataStorage.RegisterQuery(
 		Select(
@@ -212,8 +208,7 @@ void UTypedElementAlertQueriesFactory::RegisterParentUpdatesQueries(ITypedElemen
 
 void UTypedElementAlertQueriesFactory::RegisterChildAlertUpdatesQueries(ITypedElementDataStorageInterface& DataStorage)
 {
-	using namespace UE::Editor::DataStorage;
-	using namespace TypedElementQueryBuilder;
+	using namespace UE::Editor::DataStorage::Queries;
 
 	DataStorage.RegisterQuery(
 		Select(
@@ -273,8 +268,7 @@ void UTypedElementAlertQueriesFactory::RegisterChildAlertUpdatesQueries(ITypedEl
 
 void UTypedElementAlertQueriesFactory::RegisterOnAddQueries(ITypedElementDataStorageInterface& DataStorage)
 {
-	using namespace TypedElementQueryBuilder;
-	using namespace UE::Editor::DataStorage;
+	using namespace UE::Editor::DataStorage::Queries;
 	
 	DataStorage.RegisterQuery(
 		Select(
@@ -303,7 +297,7 @@ void UTypedElementAlertQueriesFactory::RegisterOnAddQueries(ITypedElementDataSto
 
 void UTypedElementAlertQueriesFactory::RegisterOnRemoveQueries(ITypedElementDataStorageInterface& DataStorage)
 {
-	using namespace TypedElementQueryBuilder;
+	using namespace UE::Editor::DataStorage::Queries;
 	using namespace UE::Editor::DataStorage;
 	
 	DataStorage.RegisterQuery(
@@ -356,8 +350,7 @@ void UTypedElementAlertQueriesFactory::IncrementParents(
 	UE::Editor::DataStorage::IQueryContext& Context, UE::Editor::DataStorage::RowHandle Row, FTypedElementAlertColumnType AlertType,
 	int32 ChildAlertQueryIndex)
 {
-	using namespace TypedElementQueryBuilder;
-	using namespace UE::Editor::DataStorage;
+	using namespace UE::Editor::DataStorage::Queries;
 	
 	while (Context.IsRowAvailable(Row))
 	{
@@ -384,8 +377,7 @@ void UTypedElementAlertQueriesFactory::ResetChildAlertCounters(FTypedElementChil
 bool UTypedElementAlertQueriesFactory::MoveToNextParent(
 	UE::Editor::DataStorage::RowHandle& Parent, UE::Editor::DataStorage::IQueryContext& Context, int32 SubQueryIndex)
 {
-	using namespace UE::Editor::DataStorage;
-	using namespace TypedElementQueryBuilder;
+	using namespace UE::Editor::DataStorage::Queries;
 	
 	FQueryResult Result = Context.RunSubquery(SubQueryIndex, Parent, CreateSubqueryCallbackBinding(
 		[&Parent](const FTableRowParentColumn& NextParent)
diff --git a/Engine/Plugins/Experimental/EditorDataStorage/Source/TedsCore/Private/Queries/TypedElementHierarchyQueries.cpp b/Engine/Plugins/Experimental/EditorDataStorage/Source/TedsCore/Private/Queries/TypedElementHierarchyQueries.cpp
--- a/Engine/Plugins/Experimental/EditorDataStorage/Source/TedsCore/Private/Queries/TypedElementHierarchyQueries.cpp
+++ b/Engine/Plugins/Experimental/EditorDataStorage/Source/TedsCore/Private/Queries/TypedElementHierarchyQueries.cpp
@@ -7,8 +7,7 @@
 
 void UTypedElementHiearchyQueriesFactory::RegisterQueries(ITypedElementDataStorageInterface& DataStorage)
 {
-	using namespace UE::Editor::DataStorage;
-	using namespace TypedElementQueryBuilder;
+	using namespace UE::Editor::DataStorage::Queries;
 
 	DataStorage.RegisterQuery(
 		Select(
@@ -22,8 +21,6 @@ void UTypedElementHiearchyQueriesFactory::RegisterQueries(ITypedElementDataStora
 					Context.RemoveColumns<FUnresolvedTableRowParentColumn>(Row);
 					Context.AddColumn(Row, FTableRowParentColumn{ .Parent = ParentRow });
 				}
-			}
-		)
-		.Compile()
-	);
+			})
+		.Compile());
 }
diff --git a/Engine/Plugins/Experimental/EditorDataStorage/Source/TedsCore/Private/Queries/TypedElementMiscQueries.cpp b/Engine/Plugins/Experimental/EditorDataStorage/Source/TedsCore/Private/Queries/TypedElementMiscQueries.cpp
--- a/Engine/Plugins/Experimental/EditorDataStorage/Source/TedsCore/Private/Queries/TypedElementMiscQueries.cpp
+++ b/Engine/Plugins/Experimental/EditorDataStorage/Source/TedsCore/Private/Queries/TypedElementMiscQueries.cpp
@@ -7,8 +7,7 @@
 
 void UTypedElementRemoveSyncToWorldTagFactory::RegisterQueries(ITypedElementDataStorageInterface& DataStorage)
 {
-	using namespace TypedElementQueryBuilder;
-	using namespace UE::Editor::DataStorage;
+	using namespace UE::Editor::DataStorage::Queries;
 
 	DataStorage.RegisterQuery(
 		Select(
@@ -21,20 +20,18 @@ void UTypedElementRemoveSyncToWorldTagFactory::RegisterQueries(ITypedElementData
 		)
 		.Where()
 			.All<FTypedElementSyncBackToWorldTag>()
-		.Compile()
-	);
+		.Compile());
 
 	DataStorage.RegisterQuery(
-	Select(
-		TEXT("Remove 'sync from world' tag"),
-		FPhaseAmble(FPhaseAmble::ELocation::Postamble, EQueryTickPhase::FrameEnd),
-		[](IQueryContext& Context, const RowHandle* Rows)
-		{
-			Context.RemoveColumns<FTypedElementSyncFromWorldTag>(TConstArrayView<RowHandle>(Rows, Context.GetRowCount()));
-		}
-	)
-	.Where()
-		.All<FTypedElementSyncFromWorldTag>()
-	.Compile()
-);
+		Select(
+			TEXT("Remove 'sync from world' tag"),
+			FPhaseAmble(FPhaseAmble::ELocation::Postamble, EQueryTickPhase::FrameEnd),
+			[](IQueryContext& Context, const RowHandle* Rows)
+			{
+				Context.RemoveColumns<FTypedElementSyncFromWorldTag>(TConstArrayView<RowHandle>(Rows, Context.GetRowCount()));
+			}
+		)
+		.Where()
+			.All<FTypedElementSyncFromWorldTag>()
+		.Compile());
 }
diff --git a/Engine/Plugins/Experimental/EditorDataStorage/Source/TedsCore/Private/TypedElementDatabaseCompatibility.cpp b/Engine/Plugins/Experimental/EditorDataStorage/Source/TedsCore/Private/TypedElementDatabaseCompatibility.cpp
--- a/Engine/Plugins/Experimental/EditorDataStorage/Source/TedsCore/Private/TypedElementDatabaseCompatibility.cpp
+++ b/Engine/Plugins/Experimental/EditorDataStorage/Source/TedsCore/Private/TypedElementDatabaseCompatibility.cpp
@@ -376,7 +376,7 @@ void UEditorDataStorageCompatibility::RegisterTypeInformationQueries()
 {
 	// Thread-safe as this is only called from a function that has an exclusive lock.
 
-	using namespace TypedElementQueryBuilder;
+	using namespace UE::Editor::DataStorage::Queries;
 
 	ClassTypeInfoQuery = Storage->RegisterQuery(
 		Select()
@@ -603,8 +603,7 @@ void UEditorDataStorageCompatibility::FPendingTypeInformationUpdate::AddTypeInfo
 
 void UEditorDataStorageCompatibility::FPendingTypeInformationUpdate::Process(UEditorDataStorageCompatibility& Compatibility)
 {
-	using namespace TypedElementQueryBuilder;
-	using namespace UE::Editor::DataStorage;
+	using namespace UE::Editor::DataStorage::Queries;
 
 	if (bHasPendingUpdate)
 	{
@@ -1104,8 +1103,7 @@ void UEditorDataStorageCompatibility::OnObjectReinstanced(const FCoreUObjectDele
 
 void UEditorDataStorageCompatibility::OnPostGcUnreachableAnalysis()
 {
-	using namespace TypedElementQueryBuilder;
-	using namespace UE::Editor::DataStorage;
+	using namespace UE::Editor::DataStorage::Queries;
 	using namespace UE::Editor::DataStorage::Private;
 
 	if (bIntegrateWithGC)
diff --git a/Engine/Plugins/Experimental/EditorDataStorage/Source/TedsCore/Tests/DynamicColumnTests.cpp b/Engine/Plugins/Experimental/EditorDataStorage/Source/TedsCore/Tests/DynamicColumnTests.cpp
--- a/Engine/Plugins/Experimental/EditorDataStorage/Source/TedsCore/Tests/DynamicColumnTests.cpp
+++ b/Engine/Plugins/Experimental/EditorDataStorage/Source/TedsCore/Tests/DynamicColumnTests.cpp
@@ -151,7 +151,7 @@ namespace UE::Editor::DataStorage::Tests
 				{
 					Async(EAsyncExecution::TaskGraphMainThread,[this]()
 					{
-						using namespace TypedElementQueryBuilder;
+						using namespace UE::Editor::DataStorage::Queries;
 
 						TArray<RowHandle> RowsToMatch;
 						TBitArray WasMatched;
@@ -295,8 +295,8 @@ namespace UE::Editor::DataStorage::Tests
 								}
 							}
 						};
-						
-						using namespace TypedElementQueryBuilder;
+
+						using namespace UE::Editor::DataStorage::Queries;
 						// Setup an activatable processor for the test
 						ActivationKeys.Emplace(TEXT("TEST: Match Row[0]"));
 						QueryExpectedMatchRows.Emplace(TArray({Rows[0]}));
@@ -484,7 +484,7 @@ namespace UE::Editor::DataStorage::Tests
 					
 					Async(EAsyncExecution::TaskGraphMainThread,[this, &ActivationKeys, &AllTestExpectations, &UnexpectedRowCount]()
 					{
-						using namespace TypedElementQueryBuilder;
+						using namespace UE::Editor::DataStorage::Queries;
 
 						auto RunTest = [this, &AllTestExpectations, &UnexpectedRowCount](IQueryContext& Context, const RowHandle* RowsPtr, int32 TestIndex)
 						{
@@ -547,11 +547,10 @@ namespace UE::Editor::DataStorage::Tests
 								ActivationKeys.Last(),
 								FProcessor(EQueryTickPhase::FrameEnd, TedsInterface->GetQueryTickGroupName(EQueryTickGroups::SyncDataStorageToExternal))
 									.MakeActivatable(ActivationKeys.Last()),
-									[this, TestIndex = AllTestExpectations.Num() - 1, &RunTest](IQueryContext& Context, const RowHandle* RowsPtr)
-									{
-										RunTest(Context, RowsPtr, TestIndex);
-									}
-							)
+								[this, TestIndex = AllTestExpectations.Num() - 1, &RunTest](IQueryContext& Context, const RowHandle* RowsPtr)
+								{
+									RunTest(Context, RowsPtr, TestIndex);
+								})
 							.ReadOnly<FTestDynamicColumn>(Identifiers[0])
 							.Compile());
 					});
diff --git a/Engine/Plugins/Experimental/EditorDataStorage/Source/TedsCore/Tests/TypedElementDataStorageInterfaceTests.cpp b/Engine/Plugins/Experimental/EditorDataStorage/Source/TedsCore/Tests/TypedElementDataStorageInterfaceTests.cpp
--- a/Engine/Plugins/Experimental/EditorDataStorage/Source/TedsCore/Tests/TypedElementDataStorageInterfaceTests.cpp
+++ b/Engine/Plugins/Experimental/EditorDataStorage/Source/TedsCore/Tests/TypedElementDataStorageInterfaceTests.cpp
@@ -56,9 +56,9 @@ namespace UE::Editor::DataStorage
 		template<typename ColumnType>
 		void RegisterTestObserver(
 			FName ObserverName,
-			ITypedElementDataStorageInterface* TedsInterface, TypedElementQueryBuilder::FObserver::EEvent EventType, int64* QueryCallCountPtr, QueryHandle* QueryHandleInOut)
+			ITypedElementDataStorageInterface* TedsInterface, Queries::FObserver::EEvent EventType, int64* QueryCallCountPtr, QueryHandle* QueryHandleInOut)
 		{
-			using namespace TypedElementQueryBuilder;
+			using namespace UE::Editor::DataStorage::Queries;
 			FQueryDescription QueryDescription = TedsInterface->GetQueryDescription(*QueryHandleInOut);
 
 			// Check if it is empty... If not then we need to register
@@ -72,7 +72,8 @@ namespace UE::Editor::DataStorage
 						[QueryCallCountPtr](IQueryContext& Context, RowHandle Row)
 						{
 							++(*QueryCallCountPtr);
-						}).Compile());
+						})
+					.Compile());
 			}
 		}
 	
@@ -211,13 +212,15 @@ namespace UE::Editor::DataStorage
 						 */
 						BeforeEach([this]()
 						{
+							using namespace Queries;
+
 							TestTableHandleA = RegisterTestTableA(TedsInterface);
 							TestNotEqual("Expecting valid table handle", TestTableHandleA, InvalidTableHandle);
 
 							RegisterTestObserver<FTestColumnA>(
 								TEXT("Increment CallCount when FTestColumnA added"),
 								TedsInterface,
-								TypedElementQueryBuilder::FObserver::EEvent::Add,
+								FObserver::EEvent::Add,
 								&DataColumnA_AddObserverCallCount,
 								&DataColumnA_AddObserverHandle);
 							TestTrue("Expect valid query observer handle", DataColumnA_AddObserverHandle != InvalidQueryHandle);
@@ -225,7 +228,7 @@ namespace UE::Editor::DataStorage
 							RegisterTestObserver<FTestColumnC>(
 								TEXT("Increment CallCount when FTestColumnC added"),
 								TedsInterface,
-								TypedElementQueryBuilder::FObserver::EEvent::Add,
+								FObserver::EEvent::Add,
 								&DataColumnC_AddObserverCallCount,
 								&DataColumnC_AddObserverHandle);
 							TestTrue("Expect valid query observer handle", DataColumnC_AddObserverHandle != InvalidQueryHandle);
@@ -233,7 +236,7 @@ namespace UE::Editor::DataStorage
 							RegisterTestObserver<FTestColumnB>(
 								TEXT("Increment CallCount when FTestColumnB removed"),
 								TedsInterface,
-								TypedElementQueryBuilder::FObserver::EEvent::Remove,
+								FObserver::EEvent::Remove,
 								&DataColumnB_RemoveObserverCallCount,
 								&DataColumnB_RemoveObserverHandle);
 							TestTrue("Expect valid query observer handle", DataColumnB_RemoveObserverHandle != InvalidQueryHandle);
@@ -241,7 +244,7 @@ namespace UE::Editor::DataStorage
 							RegisterTestObserver<FTestTagColumnA>(
 								TEXT("Increment CallCount when FTestColumnA added"),
 								TedsInterface,
-								TypedElementQueryBuilder::FObserver::EEvent::Add,
+								FObserver::EEvent::Add,
 								&TagColumnA_AddObserverCallCount,
 								&TagColumnA_AddObserverHandle);
 							TestTrue("Expect valid query observer handle", TagColumnA_AddObserverHandle != InvalidQueryHandle);
@@ -249,14 +252,14 @@ namespace UE::Editor::DataStorage
 							RegisterTestObserver<FTestTagColumnC>(
 								TEXT("Increment CallCount when FTestTagColumnC added"),
 								TedsInterface,
-								TypedElementQueryBuilder::FObserver::EEvent::Add,
+								FObserver::EEvent::Add,
 								&TagColumnC_AddObserverCallCount, 
 								&TagColumnC_AddObserverHandle);
 							TestTrue("Expect valid query observer handle", TagColumnC_AddObserverHandle != InvalidQueryHandle);
 
 							RegisterTestObserver<FTestTagColumnB>(
 								TEXT("Increment CallCount when FTestTagColumnB removed"),
-								TedsInterface, TypedElementQueryBuilder::FObserver::EEvent::Remove,
+								TedsInterface, FObserver::EEvent::Remove,
 								&TagColumnB_RemoveObserverCallCount,
 								&TagColumnB_RemoveObserverHandle);
 							TestTrue("Expect valid query observer handle", TagColumnB_RemoveObserverHandle != InvalidQueryHandle);					
diff --git a/Engine/Plugins/Experimental/EditorDataStorage/Source/TedsUI/Private/Processors/WidgetReferenceColumnUpdateProcessor.cpp b/Engine/Plugins/Experimental/EditorDataStorage/Source/TedsUI/Private/Processors/WidgetReferenceColumnUpdateProcessor.cpp
--- a/Engine/Plugins/Experimental/EditorDataStorage/Source/TedsUI/Private/Processors/WidgetReferenceColumnUpdateProcessor.cpp
+++ b/Engine/Plugins/Experimental/EditorDataStorage/Source/TedsUI/Private/Processors/WidgetReferenceColumnUpdateProcessor.cpp
@@ -13,8 +13,7 @@ void UWidgetReferenceColumnUpdateFactory::RegisterQueries(ITypedElementDataStora
 
 void UWidgetReferenceColumnUpdateFactory::RegisterDeleteRowOnWidgetDeleteQuery(ITypedElementDataStorageInterface& DataStorage) const
 {
-	using namespace TypedElementQueryBuilder;
-	using namespace UE::Editor::DataStorage;
+	using namespace UE::Editor::DataStorage::Queries;
 
 	DataStorage.RegisterQuery(
     	Select(
@@ -26,18 +25,15 @@ void UWidgetReferenceColumnUpdateFactory::RegisterDeleteRowOnWidgetDeleteQuery(I
     			{
     				Context.RemoveRow(Row);
     			}
-    		}
-    	)
+    		})
     	.Where()
     		.All<FTypedElementSlateWidgetReferenceDeletesRowTag>()
-    	.Compile()
-    	);
+    	.Compile());
 }
 
 void UWidgetReferenceColumnUpdateFactory::RegisterDeleteColumnOnWidgetDeleteQuery(ITypedElementDataStorageInterface& DataStorage) const
 {
-	using namespace TypedElementQueryBuilder;
-	using namespace UE::Editor::DataStorage;
+	using namespace UE::Editor::DataStorage::Queries;
 
 	DataStorage.RegisterQuery(
 		Select(
@@ -49,10 +45,8 @@ void UWidgetReferenceColumnUpdateFactory::RegisterDeleteColumnOnWidgetDeleteQuer
 				{
 					Context.RemoveColumns<FTypedElementSlateWidgetReferenceColumn>(Row);
 				}
-			}
-		)
+			})
 		.Where()
 			.None<FTypedElementSlateWidgetReferenceDeletesRowTag>()
-		.Compile()
-	);
+		.Compile());
 }
diff --git a/Engine/Plugins/Experimental/EditorDataStorage/Source/TedsUI/Private/Widgets/CounterWidget.cpp b/Engine/Plugins/Experimental/EditorDataStorage/Source/TedsUI/Private/Widgets/CounterWidget.cpp
--- a/Engine/Plugins/Experimental/EditorDataStorage/Source/TedsUI/Private/Widgets/CounterWidget.cpp
+++ b/Engine/Plugins/Experimental/EditorDataStorage/Source/TedsUI/Private/Widgets/CounterWidget.cpp
@@ -45,8 +45,7 @@ UCounterWidgetFactory::UCounterWidgetFactory()
 
 void UCounterWidgetFactory::RegisterQueries(ITypedElementDataStorageInterface& DataStorage)
 {
-	using namespace TypedElementQueryBuilder;
-	using namespace UE::Editor::DataStorage;
+	using namespace UE::Editor::DataStorage::Queries;
 	
 	DataStorage.RegisterQuery(Select(TEXT("Sync counter widgets"), 
 		FProcessor(EQueryTickPhase::FrameEnd, DataStorage.GetQueryTickGroupName(EQueryTickGroups::SyncWidgets))
@@ -74,8 +73,8 @@ void UCounterWidgetFactory::RegisterQueries(ITypedElementDataStorageInterface& D
 				WidgetInstance->SetText(FText::Format(Counter.LabelTextFormatter, Result.Count));
 				Comparison.Value = Result.Count;
 			}
-		}
-	).Compile());
+		})
+	.Compile());
 }
 
 void UCounterWidgetFactory::RegisterWidgetPurposes(ITypedElementDataStorageUiInterface& DataStorageUi) const
@@ -87,7 +86,7 @@ void UCounterWidgetFactory::RegisterWidgetPurposes(ITypedElementDataStorageUiInt
 void UCounterWidgetFactory::RegisterWidgetConstructors(ITypedElementDataStorageInterface& DataStorage,
 	ITypedElementDataStorageUiInterface& DataStorageUi) const
 {
-	using namespace TypedElementQueryBuilder;
+	using namespace UE::Editor::DataStorage::Queries;
 
 	TUniquePtr<FCounterWidgetConstructor> WidgetCounter = MakeUnique<FCounterWidgetConstructor>();
 	WidgetCounter->LabelText = LOCTEXT("WidgetCounterStatusBarLabel", "{0} {0}|plural(one=Widget, other=Widgets)");
diff --git a/Engine/Plugins/Experimental/EditorDataStorage/Source/TedsUI/Private/Widgets/ExportedTextWidget.cpp b/Engine/Plugins/Experimental/EditorDataStorage/Source/TedsUI/Private/Widgets/ExportedTextWidget.cpp
--- a/Engine/Plugins/Experimental/EditorDataStorage/Source/TedsUI/Private/Widgets/ExportedTextWidget.cpp
+++ b/Engine/Plugins/Experimental/EditorDataStorage/Source/TedsUI/Private/Widgets/ExportedTextWidget.cpp
@@ -54,8 +54,8 @@ namespace UE::Editor
 		
 		static DataStorage::RowHandle RegisterUpdateCallback(ITypedElementDataStorageInterface& DataStorage, const UScriptStruct* Target)
 		{
-			using namespace TypedElementQueryBuilder;
-			using namespace UE::Editor::DataStorage;
+			using namespace UE::Editor::DataStorage::Queries;
+
 			
 			QueryHandle TypeDataQuery = DataStorage.RegisterQuery(
 				Select()
diff --git a/Engine/Plugins/Experimental/EditorDataStorageFeatures/Source/TedsAssetData/Private/CB/TedsAssetDataCBDataSource.cpp b/Engine/Plugins/Experimental/EditorDataStorageFeatures/Source/TedsAssetData/Private/CB/TedsAssetDataCBDataSource.cpp
--- a/Engine/Plugins/Experimental/EditorDataStorageFeatures/Source/TedsAssetData/Private/CB/TedsAssetDataCBDataSource.cpp
+++ b/Engine/Plugins/Experimental/EditorDataStorageFeatures/Source/TedsAssetData/Private/CB/TedsAssetDataCBDataSource.cpp
@@ -23,7 +23,7 @@ namespace UE::Editor::AssetData::Private
 FTedsAssetDataCBDataSource::FTedsAssetDataCBDataSource(ITypedElementDataStorageInterface& InDatabase)
 	: Database(InDatabase)
 {
-	using namespace TypedElementQueryBuilder;
+	using namespace Queries;
 
 	InitVirtualPathProcessor();
 
@@ -61,12 +61,10 @@ FTedsAssetDataCBDataSource::FTedsAssetDataCBDataSource(ITypedElementDataStorageI
 						Context.AddColumn(Rows[Index], MoveTemp(VirtualPathColumn));
 					}
 				}
-			}
-		)
+			})
 		.Where()
 			.All<FUpdatedPathTag>()
-		.Compile()
-		);
+		.Compile());
 
 	ProcessAssetDataPathUpdateQuery = Database.RegisterQuery(
 		Select(
@@ -88,13 +86,11 @@ FTedsAssetDataCBDataSource::FTedsAssetDataCBDataSource(ITypedElementDataStorageI
 					// Todo investigate for a batch add Maybe?
 					Context.AddColumn(Row, MoveTemp(VirtualPathColumn));
 				}
-			}
-		)
+			})
 		.Where()
 				.All<FUpdatedPathTag>()
 			.None<FUpdatedAssetDataTag>()
-		.Compile()
-		);
+		.Compile());
 
 
 	// For now just add the columns one by one but this should be rework to work in batch
@@ -176,8 +172,7 @@ FTedsAssetDataCBDataSource::FTedsAssetDataCBDataSource(ITypedElementDataStorageI
 		)
 		.Where()
 			.All<FUpdatedAssetDataTag, FUpdatedPathTag>()
-			.Compile()
-		);
+		.Compile());
 
 	ProcessAssetDataUpdateQuery = Database.RegisterQuery(
 		Select(
@@ -205,8 +200,7 @@ FTedsAssetDataCBDataSource::FTedsAssetDataCBDataSource(ITypedElementDataStorageI
 		.Where()
 			.All<FUpdatedAssetDataTag>()
 			.None<FUpdatedPathTag>()
-		.Compile()
-		);
+		.Compile());
 
 }
 
diff --git a/Engine/Plugins/Experimental/EditorDataStorageFeatures/Source/TedsAssetData/Private/Debug/TedsAssetDataDebugCommands.cpp b/Engine/Plugins/Experimental/EditorDataStorageFeatures/Source/TedsAssetData/Private/Debug/TedsAssetDataDebugCommands.cpp
--- a/Engine/Plugins/Experimental/EditorDataStorageFeatures/Source/TedsAssetData/Private/Debug/TedsAssetDataDebugCommands.cpp
+++ b/Engine/Plugins/Experimental/EditorDataStorageFeatures/Source/TedsAssetData/Private/Debug/TedsAssetDataDebugCommands.cpp
@@ -219,7 +219,6 @@ namespace UE::Editor::DataStorage::Debug::Private
 			{
 				TRACE_CPUPROFILER_EVENT_SCOPE(TEDS.Debug.CreateDebugAssetRows);
 
-				using namespace TypedElementQueryBuilder;
 
 				ITypedElementDataStorageInterface* DataStorage = UTypedElementRegistry::GetInstance()->GetMutableDataStorage();
 
@@ -251,7 +250,7 @@ namespace UE::Editor::DataStorage::Debug::Private
 		FConsoleCommandWithArgsDelegate::CreateLambda(
 			[](const TArray<FString>& Args)
 			{
-				using namespace TypedElementQueryBuilder;
+				using namespace UE::Editor::DataStorage::Queries;
 
 				ITypedElementDataStorageInterface* DataStorage = UTypedElementRegistry::GetInstance()->GetMutableDataStorage();
 
@@ -271,7 +270,7 @@ namespace UE::Editor::DataStorage::Debug::Private
 						Rows.Append(Context.GetRowHandles());
 					}));
 
-				for(const RowHandle Row : Rows)
+				for (const RowHandle Row : Rows)
 				{
 					DataStorage->RemoveRow(Row);
 				}
diff --git a/Engine/Plugins/Experimental/EditorDataStorageFeatures/Source/TedsAssetData/Private/TedsAssetData.cpp b/Engine/Plugins/Experimental/EditorDataStorageFeatures/Source/TedsAssetData/Private/TedsAssetData.cpp
--- a/Engine/Plugins/Experimental/EditorDataStorageFeatures/Source/TedsAssetData/Private/TedsAssetData.cpp
+++ b/Engine/Plugins/Experimental/EditorDataStorageFeatures/Source/TedsAssetData/Private/TedsAssetData.cpp
@@ -160,6 +160,7 @@ void PopulatePathDataTableRow(FPopulatePathRowArgs&& InPopulatePathRowArgs, ITyp
 FTedsAssetData::FTedsAssetData(ITypedElementDataStorageInterface& InDatabase)
 	: Database(InDatabase)
 {
+	using namespace UE::Editor::DataStorage::Queries;
 	TRACE_CPUPROFILER_EVENT_SCOPE(FTedsAssetData::FTedsAssetData);
 
 #if TRACK_TEDSASSETDATA_MEMORY
@@ -192,43 +193,38 @@ FTedsAssetData::FTedsAssetData(ITypedElementDataStorageInterface& InDatabase)
 
 
 	RemoveUpdatedPathTagQuery = Database.RegisterQuery(
-			TypedElementQueryBuilder::Select(
-				TEXT("FTedsAssetData: Remove Updated Path Tag"),
-				TypedElementQueryBuilder::FPhaseAmble(TypedElementQueryBuilder::FPhaseAmble::ELocation::Postamble, EQueryTickPhase::FrameEnd),
-				[](IQueryContext& Context, const RowHandle* Rows)
-				{
-					Context.RemoveColumns<FUpdatedPathTag>(TConstArrayView<RowHandle>(Rows, Context.GetRowCount()));
-				}
-			)
-			.Where()
-				.All<FUpdatedPathTag>()
-			.Compile()
-		);
+		Select(
+			TEXT("FTedsAssetData: Remove Updated Path Tag"),
+			FPhaseAmble(FPhaseAmble::ELocation::Postamble, EQueryTickPhase::FrameEnd),
+			[](IQueryContext& Context, const RowHandle* Rows)
+			{
+				Context.RemoveColumns<FUpdatedPathTag>(TConstArrayView<RowHandle>(Rows, Context.GetRowCount()));
+			})
+		.Where()
+			.All<FUpdatedPathTag>()
+		.Compile());
 
 	RemoveUpdatedAssetDataTagQuery = Database.RegisterQuery(
-			TypedElementQueryBuilder::Select(
-				TEXT("FTedsAssetData: Remove Updated Asset Data Tag"),
-				TypedElementQueryBuilder::FPhaseAmble(TypedElementQueryBuilder::FPhaseAmble::ELocation::Postamble, EQueryTickPhase::FrameEnd),
-				[](IQueryContext& Context, const RowHandle* Rows)
-				{
-					Context.RemoveColumns<FUpdatedAssetDataTag>(TConstArrayView<RowHandle>(Rows, Context.GetRowCount()));
-				}
-			)
-			.Where()
-				.All<FUpdatedAssetDataTag>()
-				.Compile()
-		);
+		Select(
+			TEXT("FTedsAssetData: Remove Updated Asset Data Tag"),
+			FPhaseAmble(FPhaseAmble::ELocation::Postamble, EQueryTickPhase::FrameEnd),
+			[](IQueryContext& Context, const RowHandle* Rows)
+			{
+				Context.RemoveColumns<FUpdatedAssetDataTag>(TConstArrayView<RowHandle>(Rows, Context.GetRowCount()));
+			})
+		.Where()
+			.All<FUpdatedAssetDataTag>()
+		.Compile());
 
 	UpdateAssetsInPathQuery = Database.RegisterQuery(
-		TypedElementQueryBuilder::Select()
+		Select()
 			.ReadWrite<FAssetsInPathColumn_Experimental>()
-			.Compile()
-		);
+		.Compile());
 
 	ResolveMissingAssetInPathQuery = Database.RegisterQuery(
-		TypedElementQueryBuilder::Select(
+		Select(
 			TEXT("FTedsAssetData: Resolve Missing Asset In Path"),
-			TypedElementQueryBuilder::FProcessor(EQueryTickPhase::FrameEnd, Database.GetQueryTickGroupName(EQueryTickGroups::Default)),
+			FProcessor(EQueryTickPhase::FrameEnd, Database.GetQueryTickGroupName(EQueryTickGroups::Default)),
 			[this](IQueryContext& Context, RowHandle Row, const FUnresolvedAssetsInPathColumn_Experimental& UnresolvedAssetPath)
 			{
 #if TRACK_TEDSASSETDATA_MEMORY
@@ -240,33 +236,29 @@ FTedsAssetData::FTedsAssetData(ITypedElementDataStorageInterface& InDatabase)
 				{
 					Context.RemoveColumns<FUnresolvedAssetsInPathColumn_Experimental>(Row);
 					Context.RunSubquery(0, PathRow, 
-							TypedElementQueryBuilder::CreateSubqueryCallbackBinding([Row](FAssetsInPathColumn_Experimental& AssetsInPath)
-							{
+						CreateSubqueryCallbackBinding([Row](FAssetsInPathColumn_Experimental& AssetsInPath)
+						{
 #if TRACK_TEDSASSETDATA_MEMORY
-								LLM_SCOPE_BYNAME(TEXT("FTedsAssetData"))
+							LLM_SCOPE_BYNAME(TEXT("FTedsAssetData"))
 #endif
 
-								AssetsInPath.AssetsRow.Add(Row);
-							})
-						);
+							AssetsInPath.AssetsRow.Add(Row);
+						}));
 				}
-			}
-		)
+			})
 		.DependsOn()
 			.SubQuery(UpdateAssetsInPathQuery)
-		.Compile()
-	);
+		.Compile());
 
 	UpdateParentToChildrenAssetPathQuery = Database.RegisterQuery(
-		TypedElementQueryBuilder::Select()
+		Select()
 			.ReadWrite<FChildrenAssetPathColumn_Experimental>()
-			.Compile()
-		);
+		.Compile());
 
 	ResolveMissingParentPathQuery = Database.RegisterQuery(
-		TypedElementQueryBuilder::Select(
+		Select(
 			TEXT("FTedsAssetData: Resolve Missing Parent Path Row"),
-			TypedElementQueryBuilder::FProcessor(EQueryTickPhase::FrameEnd, Database.GetQueryTickGroupName(EQueryTickGroups::Default)),
+			FProcessor(EQueryTickPhase::FrameEnd, Database.GetQueryTickGroupName(EQueryTickGroups::Default)),
 			[this](IQueryContext& Context, RowHandle Row, const FUnresolvedParentAssetPathColumn_Experimental& UnresolvedParentAssetPath, FParentAssetPathColumn_Experimental& ParentAssetPathColumn)
 			{
 #if TRACK_TEDSASSETDATA_MEMORY
@@ -280,21 +272,18 @@ FTedsAssetData::FTedsAssetData(ITypedElementDataStorageInterface& InDatabase)
 					ParentAssetPathColumn.ParentRow = ParentPathRow;
 
 					Context.RunSubquery(0, ParentPathRow, 
-							TypedElementQueryBuilder::CreateSubqueryCallbackBinding([Row](FChildrenAssetPathColumn_Experimental& ChildrenPathColumn)
-							{
+						CreateSubqueryCallbackBinding([Row](FChildrenAssetPathColumn_Experimental& ChildrenPathColumn)
+						{
 #if TRACK_TEDSASSETDATA_MEMORY
-								LLM_SCOPE_BYNAME(TEXT("FTedsAssetData"))
+							LLM_SCOPE_BYNAME(TEXT("FTedsAssetData"))
 #endif
-								ChildrenPathColumn.ChildrenRows.Add(Row);
-							})
-						);
+							ChildrenPathColumn.ChildrenRows.Add(Row);
+						}));
 				}
-			}
-		)
+			})
 		.DependsOn()
 			.SubQuery(UpdateParentToChildrenAssetPathQuery)
-		.Compile()
-	);
+		.Compile());
 
 	// Init with the data existing at moment in asset registry
 
diff --git a/Engine/Plugins/Experimental/EditorDataStorageFeatures/Source/TedsDebugger/Private/Widgets/QueryEditor/TedsQueryEditorResultsView.cpp b/Engine/Plugins/Experimental/EditorDataStorageFeatures/Source/TedsDebugger/Private/Widgets/QueryEditor/TedsQueryEditorResultsView.cpp
--- a/Engine/Plugins/Experimental/EditorDataStorageFeatures/Source/TedsDebugger/Private/Widgets/QueryEditor/TedsQueryEditorResultsView.cpp
+++ b/Engine/Plugins/Experimental/EditorDataStorageFeatures/Source/TedsDebugger/Private/Widgets/QueryEditor/TedsQueryEditorResultsView.cpp
@@ -100,7 +100,7 @@ namespace UE::Editor::DataStorage::Debug::QueryEditor
 
 	void SResultsView::Tick(const FGeometry& AllottedGeometry, const double InCurrentTime, const float InDeltaTime)
 	{
-		using namespace TypedElementQueryBuilder;
+		using namespace UE::Editor::DataStorage::Queries;
 
 		ITypedElementDataStorageInterface& TedsInterface = Model->GetTedsInterface();
 
@@ -154,9 +154,9 @@ namespace UE::Editor::DataStorage::Debug::QueryEditor
 
 			FQueryResult QueryResult = Model->GetTedsInterface().RunQuery(TableViewerQueryHandle,
 				CreateDirectQueryCallbackBinding([&NewTableViewerRows_Set](const ITypedElementDataStorageInterface::IDirectQueryContext& Context, const RowHandle*)
-			{
-				NewTableViewerRows_Set.Append(Context.GetRowHandles());
-			}));
+				{
+					NewTableViewerRows_Set.Append(Context.GetRowHandles());
+				}));
 		}
 	
 		// Check if the two sets are equal, i.e not changes and no need to update the table viewer
diff --git a/Engine/Plugins/Experimental/EditorDataStorageFeatures/Source/TedsDebugger/Private/Widgets/RowReferenceWidget.cpp b/Engine/Plugins/Experimental/EditorDataStorageFeatures/Source/TedsDebugger/Private/Widgets/RowReferenceWidget.cpp
--- a/Engine/Plugins/Experimental/EditorDataStorageFeatures/Source/TedsDebugger/Private/Widgets/RowReferenceWidget.cpp
+++ b/Engine/Plugins/Experimental/EditorDataStorageFeatures/Source/TedsDebugger/Private/Widgets/RowReferenceWidget.cpp
@@ -106,8 +106,7 @@ void URowReferenceWidgetFactory::RegisterWidgetConstructors(ITypedElementDataSto
 
 void URowReferenceWidgetFactory::RegisterQueries(ITypedElementDataStorageInterface& DataStorage)
 {
-	using namespace TypedElementQueryBuilder;
-	using namespace UE::Editor::DataStorage;
+	using namespace UE::Editor::DataStorage::Queries;
 	
 	const QueryHandle UpdateRowReferenceWidget = DataStorage.RegisterQuery(
 		Select()
@@ -135,8 +134,7 @@ void URowReferenceWidgetFactory::RegisterQueries(ITypedElementDataStorageInterfa
 			})
 		.DependsOn()
 			.SubQuery( UpdateRowReferenceWidget )
-		.Compile()
-	);
+		.Compile());
 }
 
 FRowReferenceWidgetConstructor::FRowReferenceWidgetConstructor()
diff --git a/Engine/Plugins/Experimental/EditorDataStorageFeatures/Source/TedsOutliner/Private/Compatibility/SceneOutlinerTedsBridge.cpp b/Engine/Plugins/Experimental/EditorDataStorageFeatures/Source/TedsOutliner/Private/Compatibility/SceneOutlinerTedsBridge.cpp
--- a/Engine/Plugins/Experimental/EditorDataStorageFeatures/Source/TedsOutliner/Private/Compatibility/SceneOutlinerTedsBridge.cpp
+++ b/Engine/Plugins/Experimental/EditorDataStorageFeatures/Source/TedsOutliner/Private/Compatibility/SceneOutlinerTedsBridge.cpp
@@ -34,7 +34,7 @@ FAutoConsoleCommand BindColumnsToSceneOutlinerConsoleCommand(
 	TEXT("An example input to show a label column is 'TEDS.UI.BindColumnsToSceneOutliner /Script/TypedElementFramework.TypedElementLabelColumn'."),
 	FConsoleCommandWithArgsDelegate::CreateLambda([](const TArray<FString>& Args)
 		{
-			using namespace TypedElementQueryBuilder;
+			using namespace UE::Editor::DataStorage::Queries;
 			
 		    const FName WidgetPurposes[] = {TEXT("SceneOutliner.Cell"), TEXT("General.Cell")};
 
diff --git a/Engine/Plugins/Experimental/EditorDataStorageFeatures/Source/TedsOutliner/Private/TedsOutlinerImpl.cpp b/Engine/Plugins/Experimental/EditorDataStorageFeatures/Source/TedsOutliner/Private/TedsOutlinerImpl.cpp
--- a/Engine/Plugins/Experimental/EditorDataStorageFeatures/Source/TedsOutliner/Private/TedsOutlinerImpl.cpp
+++ b/Engine/Plugins/Experimental/EditorDataStorageFeatures/Source/TedsOutliner/Private/TedsOutlinerImpl.cpp
@@ -48,7 +48,7 @@ FTedsOutlinerImpl::FTedsOutlinerImpl(const FTedsOutlinerParams& InParams, IScene
 
 void FTedsOutlinerImpl::CreateLabelWidgetConstructors()
 {
-	using namespace TypedElementQueryBuilder;
+	using namespace UE::Editor::DataStorage::Queries;
 
 	// Create and assign the widget constructors for the label column
 	
@@ -95,10 +95,9 @@ void FTedsOutlinerImpl::CreateLabelWidgetConstructors()
 	};
 	
 	UE::Editor::DataStorage::QueryHandle LabelColumnQueryHandle = Storage->RegisterQuery(
-																	Select()
-																		.ReadWrite<FTypedElementLabelColumn, FTypedElementClassTypeInfoColumn>()
-																	.Compile()
-																	);
+		Select()
+			.ReadWrite<FTypedElementLabelColumn, FTypedElementClassTypeInfoColumn>()
+		.Compile());
 
 	if (TSharedPtr<FTypedElementWidgetConstructor> LabelColumnWidgetConstructor = CreateWidgetConstructorForQuery(Storage->GetQueryDescription(LabelColumnQueryHandle)))
 	{
@@ -108,7 +107,7 @@ void FTedsOutlinerImpl::CreateLabelWidgetConstructors()
 
 void FTedsOutlinerImpl::CreateFilterQueries()
 {
-	using namespace TypedElementQueryBuilder;
+	using namespace UE::Editor::DataStorage::Queries;
 
 	if (CreationParams.bUseDefaultTedsFilters)
 	{
@@ -363,8 +362,7 @@ bool FTedsOutlinerImpl::CanDisplayRow(UE::Editor::DataStorage::RowHandle ItemRow
 
 void FTedsOutlinerImpl::CreateItemsFromQuery(TArray<FSceneOutlinerTreeItemPtr>& OutItems, ISceneOutlinerMode* InMode) const
 {
-	using namespace TypedElementQueryBuilder;
-	using namespace UE::Editor::DataStorage;
+	using namespace UE::Editor::DataStorage::Queries;
 	
 	TArray<RowHandle> Rows;
 	
@@ -408,8 +406,7 @@ void FTedsOutlinerImpl::CreateChildren(const FSceneOutlinerTreeItemPtr& Item, TA
 		return;
 	}
 
-	using namespace TypedElementQueryBuilder;
-	using namespace UE::Editor::DataStorage;
+	using namespace UE::Editor::DataStorage::Queries;
 	
 	const FTedsOutlinerTreeItem* TedsTreeItem = Item->CastTo<FTedsOutlinerTreeItem>();
 
@@ -552,7 +549,7 @@ void FTedsOutlinerImpl::OnItemMoved(UE::Editor::DataStorage::RowHandle ItemRowHa
 
 void FTedsOutlinerImpl::RecompileQueries()
 {
-	using namespace TypedElementQueryBuilder;
+	using namespace UE::Editor::DataStorage::Queries;
 
 	UnregisterQueries();
 
@@ -570,13 +567,13 @@ void FTedsOutlinerImpl::RecompileQueries()
 	// Row to track addition of rows to the Outliner
 	FQueryDescription RowAdditionQueryDescription =
 		Select(
-				TEXT("Add Row to Outliner"),
-				FObserver::OnAdd<FTypedElementLabelColumn>().SetExecutionMode(EExecutionMode::GameThread),
-				[this](IQueryContext& Context, UE::Editor::DataStorage::RowHandle Row)
-				{
-					OnItemAdded(Row);
-				})
-			.Compile();
+			TEXT("Add Row to Outliner"),
+			FObserver::OnAdd<FTypedElementLabelColumn>().SetExecutionMode(EExecutionMode::GameThread),
+			[this](IQueryContext& Context, UE::Editor::DataStorage::RowHandle Row)
+			{
+				OnItemAdded(Row);
+			})
+		.Compile();
 
 	// Add the conditions from FinalQueryDescription to ensure we are tracking addition of the rows the user requested
 	AppendQuery(RowAdditionQueryDescription, FinalQueryDescription);
@@ -710,9 +707,8 @@ void FTedsOutlinerImpl::ClearSelection() const
 	{
 		return;
 	}
-	
-	using namespace TypedElementQueryBuilder;
-	using namespace UE::Editor::DataStorage;
+
+	using namespace UE::Editor::DataStorage::Queries;
 
 	TArray<RowHandle> RowsToRemoveSelectionColumn;
 
diff --git a/Engine/Plugins/Experimental/EditorDataStorageFeatures/Source/TedsOutliner/Private/TedsOutlinerMode.cpp b/Engine/Plugins/Experimental/EditorDataStorageFeatures/Source/TedsOutliner/Private/TedsOutlinerMode.cpp
--- a/Engine/Plugins/Experimental/EditorDataStorageFeatures/Source/TedsOutliner/Private/TedsOutlinerMode.cpp
+++ b/Engine/Plugins/Experimental/EditorDataStorageFeatures/Source/TedsOutliner/Private/TedsOutlinerMode.cpp
@@ -27,9 +27,7 @@ namespace UE::EditorDataStorage::Outliner::Private
 
 FTedsOutlinerMode::FTedsOutlinerMode(const FTedsOutlinerParams& InParams)
 	: ISceneOutlinerMode(InParams.SceneOutliner)
-{
-	using namespace TypedElementQueryBuilder;
-	
+{	
 	TedsOutlinerImpl = MakeShared<FTedsOutlinerImpl>(InParams, this);
 	TedsOutlinerImpl->Init();
 	TedsOutlinerImpl->OnSelectionChanged().AddRaw(this, &FTedsOutlinerMode::OnSelectionChanged);
diff --git a/Engine/Plugins/Experimental/EditorDataStorageFeatures/Source/TedsOutliner/Private/TedsOutlinerModule.cpp b/Engine/Plugins/Experimental/EditorDataStorageFeatures/Source/TedsOutliner/Private/TedsOutlinerModule.cpp
--- a/Engine/Plugins/Experimental/EditorDataStorageFeatures/Source/TedsOutliner/Private/TedsOutlinerModule.cpp
+++ b/Engine/Plugins/Experimental/EditorDataStorageFeatures/Source/TedsOutliner/Private/TedsOutlinerModule.cpp
@@ -72,7 +72,6 @@ TSharedRef<ISceneOutliner> FTedsOutlinerModule::CreateTedsOutliner(const FSceneO
 
 	InitOptions.ModeFactory = FCreateSceneOutlinerMode::CreateLambda([&InitTedsOptions](SSceneOutliner* Outliner)
 	{
-		using namespace TypedElementQueryBuilder;
 		InitTedsOptions.SceneOutliner = Outliner;
 		
 		return new FTedsOutlinerMode(InitTedsOptions);
@@ -115,7 +114,7 @@ UE::Editor::DataStorage::QueryHandle FTedsOutlinerModule::GetLevelEditorTedsOutl
 	UTypedElementRegistry* Registry = UTypedElementRegistry::GetInstance();
 	ITypedElementDataStorageInterface* Storage = Registry->GetMutableDataStorage();
 		
-	using namespace TypedElementQueryBuilder;
+	using namespace UE::Editor::DataStorage::Queries;
 
 	static UE::Editor::DataStorage::QueryHandle ColumnQuery = Storage->RegisterQuery(
 		Select()
@@ -143,14 +142,14 @@ TSharedRef<SWidget> FTedsOutlinerModule::CreateLevelEditorTedsOutliner()
 		.Text(LOCTEXT("TEDSPluginNotEnabledText", "You need to enable the Typed Element Data Storage plugin to see the table viewer!"));
 	}
 
-	using namespace TypedElementQueryBuilder;
+	using namespace UE::Editor::DataStorage::Queries;
 
 	// The Outliner is populated with Actors and Entities
 	UE::Editor::DataStorage::FQueryDescription OutlinerQueryDescription =
-						Select()
-						.Where()
-							.All<FTypedElementClassTypeInfoColumn>() // TEDS-Outliner TODO: Currently looking at all entries with type info in TEDS
-						.Compile();
+		Select()
+		.Where()
+			.All<FTypedElementClassTypeInfoColumn>() // TEDS-Outliner TODO: Currently looking at all entries with type info in TEDS
+		.Compile();
 
 	FSceneOutlinerInitializationOptions InitOptions;
 	InitOptions.bShowHeaderRow = true;
@@ -164,10 +163,10 @@ TSharedRef<SWidget> FTedsOutlinerModule::CreateLevelEditorTedsOutliner()
 
 	// Example Query to filter for actors
 	UE::Editor::DataStorage::FQueryDescription ActorFilterQuery =
-					Select()
-					.Where()
-						.All<FTypedElementActorTag>()
-					.Compile();
+		Select()
+		.Where()
+			.All<FTypedElementActorTag>()
+		.Compile();
 	Params.FilterQueries.Emplace("Actors", ActorFilterQuery);
 		
 	// Empty selection set name is currently the level editor
diff --git a/Engine/Plugins/Experimental/EditorDataStorageFeatures/Source/TedsOutliner/Private/Widgets/AlertWidget.cpp b/Engine/Plugins/Experimental/EditorDataStorageFeatures/Source/TedsOutliner/Private/Widgets/AlertWidget.cpp
--- a/Engine/Plugins/Experimental/EditorDataStorageFeatures/Source/TedsOutliner/Private/Widgets/AlertWidget.cpp
+++ b/Engine/Plugins/Experimental/EditorDataStorageFeatures/Source/TedsOutliner/Private/Widgets/AlertWidget.cpp
@@ -168,8 +168,7 @@ void UAlertWidgetFactory::RegisterQueries(ITypedElementDataStorageInterface& Dat
 
 void UAlertWidgetFactory::RegisterAlertQueries(ITypedElementDataStorageInterface& DataStorage)
 {
-	using namespace TypedElementQueryBuilder;
-	using namespace UE::Editor::DataStorage;
+	using namespace UE::Editor::DataStorage::Queries;
 	
 	QueryHandle UpdateWidget_OnlyAlert = DataStorage.RegisterQuery(
 		Select()
@@ -250,8 +249,7 @@ void UAlertWidgetFactory::RegisterAlertQueries(ITypedElementDataStorageInterface
 
 void UAlertWidgetFactory::RegisterAlertHeaderQueries(ITypedElementDataStorageInterface& DataStorage)
 {
-	using namespace TypedElementQueryBuilder;
-	using namespace UE::Editor::DataStorage;
+	using namespace UE::Editor::DataStorage::Queries;
 	
 	QueryHandle AlertCount = DataStorage.RegisterQuery(
 		Count()
@@ -274,8 +272,7 @@ void UAlertWidgetFactory::RegisterAlertHeaderQueries(ITypedElementDataStorageInt
 						Context.AddColumns<FAlertHeaderActiveWidgetTag>(Row);
 					}
 				}
-			}
-		)
+			})
 		.Where()
 			.All<FAlertHeaderWidgetTag>()
 			.None<FAlertHeaderActiveWidgetTag>()
diff --git a/Engine/Plugins/Experimental/EditorDataStorageFeatures/Source/TedsOutliner/Private/Widgets/TransformHeadsUpWidget.cpp b/Engine/Plugins/Experimental/EditorDataStorageFeatures/Source/TedsOutliner/Private/Widgets/TransformHeadsUpWidget.cpp
--- a/Engine/Plugins/Experimental/EditorDataStorageFeatures/Source/TedsOutliner/Private/Widgets/TransformHeadsUpWidget.cpp
+++ b/Engine/Plugins/Experimental/EditorDataStorageFeatures/Source/TedsOutliner/Private/Widgets/TransformHeadsUpWidget.cpp
@@ -154,8 +154,7 @@ static void UpdateTransformHeadsUpDisplay(FTypedElementSlateWidgetReferenceColum
 
 void UTransformHeadsUpWidgetFactory::RegisterQueries(ITypedElementDataStorageInterface& DataStorage)
 {
-	using namespace TypedElementQueryBuilder;
-	using namespace UE::Editor::DataStorage;
+	using namespace UE::Editor::DataStorage::Queries;
 		
 	UE::Editor::DataStorage::QueryHandle UpdateTransformWidget = DataStorage.RegisterQuery(
 		Select()
@@ -165,20 +164,18 @@ void UTransformHeadsUpWidgetFactory::RegisterQueries(ITypedElementDataStorageInt
 		.Compile());
 
 	DataStorage.RegisterQuery(
-		Select(TEXT("Sync Transform column to heads up display"),
-		FProcessor(EQueryTickPhase::FrameEnd, DataStorage.GetQueryTickGroupName(EQueryTickGroups::SyncWidgets))
-			.SetExecutionMode(EExecutionMode::GameThread),
-			[](IQueryContext& Context,
-				FTypedElementSlateWidgetReferenceColumn& Widget,
-				const FTypedElementRowReferenceColumn& ReferenceColumn)
+		Select(
+			TEXT("Sync Transform column to heads up display"),
+			FProcessor(EQueryTickPhase::FrameEnd, DataStorage.GetQueryTickGroupName(EQueryTickGroups::SyncWidgets))
+				.SetExecutionMode(EExecutionMode::GameThread),
+			[](IQueryContext& Context, FTypedElementSlateWidgetReferenceColumn& Widget, const FTypedElementRowReferenceColumn& ReferenceColumn)
 			{
 				Context.RunSubquery(0, ReferenceColumn.Row, CreateSubqueryCallbackBinding(
 					[&Widget](const FTypedElementLocalTransformColumn& Transform)
 					{
 						UpdateTransformHeadsUpDisplay(Widget, UE::Editor::DataStorage::Widgets::Private::GetAbnormalTransformTypes(Transform.Transform));
 					}));
-			}
-		)
+			})
 	.Where()
 		.All<FTransformHeadsUpWidgetTag>()
 	.DependsOn()
diff --git a/Engine/Plugins/Experimental/EditorDataStorageFeatures/Source/TedsRevisionControl/Private/Processors/RevisionControlProcessors.cpp b/Engine/Plugins/Experimental/EditorDataStorageFeatures/Source/TedsRevisionControl/Private/Processors/RevisionControlProcessors.cpp
--- a/Engine/Plugins/Experimental/EditorDataStorageFeatures/Source/TedsRevisionControl/Private/Processors/RevisionControlProcessors.cpp
+++ b/Engine/Plugins/Experimental/EditorDataStorageFeatures/Source/TedsRevisionControl/Private/Processors/RevisionControlProcessors.cpp
@@ -237,8 +237,7 @@ void URevisionControlDataStorageFactory::RegisterQueries(ITypedElementDataStorag
 
 void URevisionControlDataStorageFactory::RegisterFetchUpdates(ITypedElementDataStorageInterface& DataStorage)
 {
-	using namespace TypedElementQueryBuilder;
-	using namespace UE::Editor::DataStorage;
+	using namespace UE::Editor::DataStorage::Queries;
 	
 	FSourceControlFileStatusMonitor& FileStatusMonitor = ISourceControlModule::Get().GetSourceControlFileStatusMonitor();
 
@@ -259,17 +258,14 @@ void URevisionControlDataStorageFactory::RegisterFetchUpdates(ITypedElementDataS
 						EmptyDelegate
 						
 					);
-				}
-			)
-			.Compile()
-		);
+				})
+			.Compile());
 	}
 }
 
 void URevisionControlDataStorageFactory::RegisterApplyOverlays(ITypedElementDataStorageInterface& DataStorage)
 {
-	using namespace TypedElementQueryBuilder;
-	using namespace UE::Editor::DataStorage;
+	using namespace UE::Editor::DataStorage::Queries;
 	using namespace UE::Editor::RevisionControl::Private;
 	
 	if (ApplyOverlaysObjectToSCC == InvalidQueryHandle)
@@ -279,7 +275,7 @@ void URevisionControlDataStorageFactory::RegisterApplyOverlays(ITypedElementData
 				.ReadOnly<FTypedElementPackagePathColumn>()
 				.ReadOnly<FSCCStatusColumn>(EOptional::Yes)
 			.Compile());
-		}
+	}
 
 	if (ApplyNewOverlays == InvalidQueryHandle)
 	{
@@ -289,8 +285,7 @@ void URevisionControlDataStorageFactory::RegisterApplyOverlays(ITypedElementData
 			.Where()
 				.All<FTypedElementActorTag>()
 				.None<FTypedElementViewportOverlayColorColumn>()
-			.Compile()
-		);
+			.Compile());
 	}
 
 	if (ChangeOverlay == InvalidQueryHandle)
@@ -300,8 +295,7 @@ void URevisionControlDataStorageFactory::RegisterApplyOverlays(ITypedElementData
 			.ReadOnly<FTypedElementUObjectColumn, FTypedElementPackageReference, FTypedElementViewportOverlayColorColumn>()
 			.Where()
 				.All<FTypedElementActorTag>()
-			.Compile()
-		);
+			.Compile());
 	}
 
 	if (FlushPackageUpdates == InvalidQueryHandle)
@@ -376,19 +370,16 @@ void URevisionControlDataStorageFactory::RegisterApplyOverlays(ITypedElementData
 						}
 					));
 					Context.RemoveRow(Row);
-				}
-			)
+				})
 			.DependsOn()
 				.SubQuery(Subqueries)
-			.Compile()
-		);
+			.Compile());
 	}
 }
 
 void URevisionControlDataStorageFactory::RegisterRemoveOverlays(ITypedElementDataStorageInterface& DataStorage)
 {
-	using namespace TypedElementQueryBuilder;
-	using namespace UE::Editor::DataStorage;
+	using namespace UE::Editor::DataStorage::Queries;
 	
 	if (RemoveOverlays == InvalidQueryHandle)
 	{
@@ -406,11 +397,9 @@ void URevisionControlDataStorageFactory::RegisterRemoveOverlays(ITypedElementDat
 				[](IQueryContext& Context, RowHandle ObjectRow, FTypedElementUObjectColumn& Actor, const FTypedElementViewportOverlayColorColumn& ViewportColor)
 				{
 					Context.RemoveColumns<FTypedElementViewportOverlayColorColumn>(ObjectRow);
-				}
-			)
+				})
 			.Where()
 				.All<FTypedElementActorTag>()
-			.Compile()
-		);
+			.Compile());
 	}
 }
\ No newline at end of file
diff --git a/Engine/Plugins/Experimental/EditorDataStorageFeatures/Source/TedsRevisionControl/Private/Queries/ObjectPackagePathToColumnQueries.cpp b/Engine/Plugins/Experimental/EditorDataStorageFeatures/Source/TedsRevisionControl/Private/Queries/ObjectPackagePathToColumnQueries.cpp
--- a/Engine/Plugins/Experimental/EditorDataStorageFeatures/Source/TedsRevisionControl/Private/Queries/ObjectPackagePathToColumnQueries.cpp
+++ b/Engine/Plugins/Experimental/EditorDataStorageFeatures/Source/TedsRevisionControl/Private/Queries/ObjectPackagePathToColumnQueries.cpp
@@ -45,8 +45,7 @@ namespace UE::Editor::RevisionControl::Private
 
 void UTypedElementUObjectPackagePathFactory::RegisterQueries(ITypedElementDataStorageInterface& DataStorage)
 {
-	using namespace TypedElementQueryBuilder;
-	using namespace UE::Editor::DataStorage;
+	using namespace UE::Editor::DataStorage::Queries;
 	
 	UE::Editor::RevisionControl::Private::CVarAutoPopulateState->AsVariable()->OnChangedDelegate().AddLambda(
 		[this, &DataStorage](IConsoleVariable* AutoPopulate)
@@ -81,10 +80,8 @@ void UTypedElementUObjectPackagePathFactory::RegisterQueries(ITypedElementDataSt
 				Context.RemoveColumns(Row, { FTypedElementPackageUnresolvedReference::StaticStruct() });
 
 				UE::Editor::RevisionControl::Private::ResolvePackageReference(Context, Package, Row, PackageRow);
-			}
-		)
-		.Compile()
-	);
+			})
+		.Compile());
 
 	if (UE::Editor::RevisionControl::Private::CVarAutoPopulateState->GetBool())
 	{
@@ -94,8 +91,7 @@ void UTypedElementUObjectPackagePathFactory::RegisterQueries(ITypedElementDataSt
 
 void UTypedElementUObjectPackagePathFactory::RegisterTryAddPackageRef(ITypedElementDataStorageInterface& DataStorage)
 {
-	using namespace TypedElementQueryBuilder;
-	using namespace UE::Editor::DataStorage;
+	using namespace UE::Editor::DataStorage::Queries;
 	
 	TryAddPackageRef = DataStorage.RegisterQuery(
 		Select(
@@ -129,8 +125,6 @@ void UTypedElementUObjectPackagePathFactory::RegisterTryAddPackageRef(ITypedElem
 						}
 					}
 				}
-			}
-		)
-		.Compile()
-	);
+			})
+		.Compile());
 }
diff --git a/Engine/Plugins/Experimental/EditorDataStorageFeatures/Source/TedsSettings/Private/TedsSettingsManager.cpp b/Engine/Plugins/Experimental/EditorDataStorageFeatures/Source/TedsSettings/Private/TedsSettingsManager.cpp
--- a/Engine/Plugins/Experimental/EditorDataStorageFeatures/Source/TedsSettings/Private/TedsSettingsManager.cpp
+++ b/Engine/Plugins/Experimental/EditorDataStorageFeatures/Source/TedsSettings/Private/TedsSettingsManager.cpp
@@ -100,7 +100,7 @@ void FTedsSettingsManager::RegisterTables(ITypedElementDataStorageInterface& Dat
 
 void FTedsSettingsManager::RegisterQueries(ITypedElementDataStorageInterface& DataStorage)
 {
-	using namespace TypedElementQueryBuilder;
+	using namespace UE::Editor::DataStorage::Queries;
 
 	if (SelectAllSettingsQuery == UE::Editor::DataStorage::InvalidQueryHandle)
 	{
@@ -274,7 +274,7 @@ void FTedsSettingsManager::UpdateSettingsCategory(TSharedPtr<ISettingsCategory>
 	// Gather all existing rows for the given { ContainerName, CategoryName } pair.
 	if (bQueryExistingRows)
 	{
-		using namespace TypedElementQueryBuilder;
+		using namespace UE::Editor::DataStorage::Queries;
 
 		DataStorage->RunQuery(SelectAllSettingsQuery, CreateDirectQueryCallbackBinding(
 			[&OldRowHandles, &OldSectionNames, &ContainerName, &CategoryName](
diff --git a/Engine/Plugins/Experimental/EditorDataStorageFeatures/Source/TedsSettings/Tests/TedsSettingsTests.cpp b/Engine/Plugins/Experimental/EditorDataStorageFeatures/Source/TedsSettings/Tests/TedsSettingsTests.cpp
--- a/Engine/Plugins/Experimental/EditorDataStorageFeatures/Source/TedsSettings/Tests/TedsSettingsTests.cpp
+++ b/Engine/Plugins/Experimental/EditorDataStorageFeatures/Source/TedsSettings/Tests/TedsSettingsTests.cpp
@@ -87,11 +87,11 @@ namespace UE::Editor::Settings::Tests
 			check(DataStorageCompatibility != nullptr);
 
 			{
-				using namespace TypedElementQueryBuilder;
+				using namespace UE::Editor::DataStorage::Queries;
 				CountAllSettingsQuery = DataStorage->RegisterQuery(
 					Count()
 					.Where()
-					.All<FSettingsSectionTag>()
+						.All<FSettingsSectionTag>()
 					.Compile());
 			}
 
diff --git a/Engine/Plugins/Experimental/EditorDataStorageFeatures/Source/TedsTableViewer/Private/TedsTableViewerColumn.cpp b/Engine/Plugins/Experimental/EditorDataStorageFeatures/Source/TedsTableViewer/Private/TedsTableViewerColumn.cpp
--- a/Engine/Plugins/Experimental/EditorDataStorageFeatures/Source/TedsTableViewer/Private/TedsTableViewerColumn.cpp
+++ b/Engine/Plugins/Experimental/EditorDataStorageFeatures/Source/TedsTableViewer/Private/TedsTableViewerColumn.cpp
@@ -156,7 +156,7 @@ namespace UE::Editor::DataStorage
 
 	void FTedsTableViewerColumn::RegisterQueries()
 	{
-		using namespace TypedElementQueryBuilder;
+		using namespace UE::Editor::DataStorage::Queries;
 
 		const TConstArrayView<TWeakObjectPtr<const UScriptStruct>> ColumnTypes = GetMatchedColumns();
 		
@@ -179,8 +179,7 @@ namespace UE::Editor::DataStorage
 						})
 				.Where()
 					.All(ColumnType.Get())
-				.Compile()
-				);
+				.Compile());
 			
 			InternalObserverQueries.Add(AddQueryHandle);
 
@@ -200,8 +199,7 @@ namespace UE::Editor::DataStorage
 						})
 				.Where()
 					.All(ColumnType.Get())
-				.Compile()
-			);
+				.Compile());
 			
 			InternalObserverQueries.Add(RemoveQueryHandle);
 		}
@@ -217,11 +215,11 @@ namespace UE::Editor::DataStorage
 		
 		// Query to get all widgets that were created by this column
 		WidgetQuery = Storage->RegisterQuery(
-				Select()
-					.ReadOnly(SelectionColumns)
-				.Where()
-					.All(AdditionalWidgetColumns)
-				.Compile());
+			Select()
+				.ReadOnly(SelectionColumns)
+			.Where()
+				.All(AdditionalWidgetColumns)
+			.Compile());
 	}
 
 	void FTedsTableViewerColumn::UnRegisterQueries()
@@ -265,8 +263,8 @@ namespace UE::Editor::DataStorage
 			// stops matching, there is a potential need for widget update
 			return (bMatchesQueryConditions && Pair.Value) || (!bMatchesQueryConditions && !Pair.Value);
 		});
-		
-		using namespace TypedElementQueryBuilder;
+
+		using namespace UE::Editor::DataStorage::Queries;
 
 		// Query to find all widgets that belong to the row handles that need updates
 		DirectQueryCallback RowCollector = CreateDirectQueryCallbackBinding(
diff --git a/Engine/Plugins/Experimental/EditorDataStorageFeatures/Source/TedsTableViewer/Private/TedsTableViewerModule.cpp b/Engine/Plugins/Experimental/EditorDataStorageFeatures/Source/TedsTableViewer/Private/TedsTableViewerModule.cpp
--- a/Engine/Plugins/Experimental/EditorDataStorageFeatures/Source/TedsTableViewer/Private/TedsTableViewerModule.cpp
+++ b/Engine/Plugins/Experimental/EditorDataStorageFeatures/Source/TedsTableViewer/Private/TedsTableViewerModule.cpp
@@ -71,30 +71,29 @@ public:
 		}
 
 		ITypedElementDataStorageInterface* DataStorage = Registry->GetMutableDataStorage();
-		
-		using namespace TypedElementQueryBuilder;
+
+		using namespace UE::Editor::DataStorage::Queries;
 
 		// We'll just create a test table viewer that views all actor rows without actually updating it dynamically for now
 		static QueryHandle QueryHandle =
 			DataStorage->RegisterQuery(
-						Select()
-						.Where()
-							.All<FTypedElementActorTag>()
-						.Compile()
-				);
+				Select()
+				.Where()
+					.All<FTypedElementActorTag>()
+				.Compile());
 
 		Rows.Empty();
 		
 		UE::Editor::DataStorage::FQueryResult QueryResult = DataStorage->RunQuery(QueryHandle,
 			CreateDirectQueryCallbackBinding([this](const ITypedElementDataStorageInterface::IDirectQueryContext& Context, const UE::Editor::DataStorage::RowHandle* RowHandles)
-		{
-			Rows.Append(RowHandles, Context.GetRowCount());
-		}));
+			{
+				Rows.Append(RowHandles, Context.GetRowCount());
+			}));
 
 		return SNew(UE::Editor::DataStorage::STedsTableViewer)
-				.QueryStack(MakeShared<UE::Editor::DataStorage::FQueryStackNode_RowView>(&Rows))
-				.Columns({FTypedElementLabelColumn::StaticStruct(), FTypedElementSelectionColumn::StaticStruct(),
-					FTypedElementAlertColumn::StaticStruct(), FTypedElementChildAlertColumn::StaticStruct()});
+			.QueryStack(MakeShared<UE::Editor::DataStorage::FQueryStackNode_RowView>(&Rows))
+			.Columns({FTypedElementLabelColumn::StaticStruct(), FTypedElementSelectionColumn::StaticStruct(),
+				FTypedElementAlertColumn::StaticStruct(), FTypedElementChildAlertColumn::StaticStruct()});
 	}
 	
 	FDelegateHandle LevelEditorTabManagerChangedHandle;
diff --git a/Engine/Plugins/Experimental/InstanceDataObjectFixupTool/Source/Private/InstanceDataObjectFixupToolTedsQueries.cpp b/Engine/Plugins/Experimental/InstanceDataObjectFixupTool/Source/Private/InstanceDataObjectFixupToolTedsQueries.cpp
--- a/Engine/Plugins/Experimental/InstanceDataObjectFixupTool/Source/Private/InstanceDataObjectFixupToolTedsQueries.cpp
+++ b/Engine/Plugins/Experimental/InstanceDataObjectFixupTool/Source/Private/InstanceDataObjectFixupToolTedsQueries.cpp
@@ -17,9 +17,7 @@
 
 void UInstanceDataObjectFixupToolTedsQueryFactory::RegisterQueries(ITypedElementDataStorageInterface& DataStorage)
 {
-	using namespace UE;
-	using namespace TypedElementQueryBuilder;
-	using namespace UE::Editor::DataStorage;
+	using namespace UE::Editor::DataStorage::Queries;
 	
 	DataStorage.RegisterQuery(
 		Select(
@@ -28,8 +26,7 @@ void UInstanceDataObjectFixupToolTedsQueryFactory::RegisterQueries(ITypedElement
 			[](IQueryContext& Context, RowHandle Row, const FTypedElementUObjectColumn& Object)
 			{
 				Context.AddColumn(Row, FTypedElementAlertActionColumn{ .Action = ShowFixUpToolForPlaceholders });
-			}
-		)
+			})
 		.Where()
 			.All<FTypedElementSyncFromWorldTag, FTypedElementAlertColumn, FTypedElementPropertyBagPlaceholderTag>()
 			.None<FTypedElementAlertActionColumn>()
@@ -42,8 +39,7 @@ void UInstanceDataObjectFixupToolTedsQueryFactory::RegisterQueries(ITypedElement
 			[](IQueryContext& Context, RowHandle Row, const FTypedElementUObjectColumn& Object)
 			{
 				Context.AddColumn(Row, FTypedElementAlertActionColumn{ .Action = ShowFixUpToolForLooseProperties });
-			}
-		)
+			})
 		.Where()
 			.All<FTypedElementSyncFromWorldTag, FTypedElementAlertColumn, FTypedElementLoosePropertyTag>()
 			.None<FTypedElementAlertActionColumn>()
diff --git a/Engine/Source/Runtime/TypedElementFramework/Private/Elements/Framework/TypedElementQueryBuilder.cpp b/Engine/Source/Runtime/TypedElementFramework/Private/Elements/Framework/TypedElementQueryBuilder.cpp
--- a/Engine/Source/Runtime/TypedElementFramework/Private/Elements/Framework/TypedElementQueryBuilder.cpp
+++ b/Engine/Source/Runtime/TypedElementFramework/Private/Elements/Framework/TypedElementQueryBuilder.cpp
@@ -7,7 +7,7 @@
 #include "Elements/Framework/TypedElementMetaData.h"
 #include "GenericPlatform/GenericPlatformMath.h"
 
-namespace TypedElementQueryBuilder
+namespace UE::Editor::DataStorage::Queries
 {
 	using namespace UE::Editor::DataStorage;
 
@@ -521,4 +521,4 @@ namespace TypedElementQueryBuilder
 	{
 		return FDependency{ &Query };
 	}
-}
\ No newline at end of file
+} // namespace UE::Editor::DataStorage::Queries
diff --git a/Engine/Source/Runtime/TypedElementFramework/Public/Elements/Framework/TypedElementQueryBuilder.h b/Engine/Source/Runtime/TypedElementFramework/Public/Elements/Framework/TypedElementQueryBuilder.h
--- a/Engine/Source/Runtime/TypedElementFramework/Public/Elements/Framework/TypedElementQueryBuilder.h
+++ b/Engine/Source/Runtime/TypedElementFramework/Public/Elements/Framework/TypedElementQueryBuilder.h
@@ -2,7 +2,6 @@
 
 #pragma once
 
-#include "CoreMinimal.h"
 #include "Elements/Common/TypedElementCommonTypes.h"
 #include "Elements/Interfaces/TypedElementDataStorageInterface.h"
 #include "Templates/SubclassOf.h"
@@ -104,7 +103,7 @@ class USubsystem;
  * and store its compiled form for repeated use instead of rebuilding the query on every update.
  */
 
-namespace TypedElementQueryBuilder
+namespace UE::Editor::DataStorage::Queries
 {
 	using namespace UE::Editor::DataStorage;
 
@@ -433,6 +432,6 @@ namespace TypedElementQueryBuilder
 	template<typename Function>
 	SubqueryCallback CreateSubqueryCallbackBinding(Function&& Callback);
 
-} // namespace TypedElementQueryBuilder
+} // namespace UE::Editor::DataStorage::Queries
 
 #include "Elements/Framework/TypedElementQueryBuilder.inl"
diff --git a/Engine/Source/Runtime/TypedElementFramework/Public/Elements/Framework/TypedElementQueryBuilder.inl b/Engine/Source/Runtime/TypedElementFramework/Public/Elements/Framework/TypedElementQueryBuilder.inl
--- a/Engine/Source/Runtime/TypedElementFramework/Public/Elements/Framework/TypedElementQueryBuilder.inl
+++ b/Engine/Source/Runtime/TypedElementFramework/Public/Elements/Framework/TypedElementQueryBuilder.inl
@@ -10,7 +10,7 @@
 #include "Elements/Common/TypedElementQueryDescription.h"
 #include "Elements/Common/TypedElementQueryTypes.h"
 
-namespace TypedElementQueryBuilder
+namespace UE::Editor::DataStorage::Queries
 {
 	using namespace UE::Editor::DataStorage;
 
@@ -1300,4 +1300,4 @@ e.g. void(ISubqueryContext& Context, RowHandle Row, ColumnType0& ColumnA, const
 			Private::CallQueryFunction<ValidateColumns>(Description, Context, Callback);
 		};
 	}
-} // namespace TypedElementQueryBuilder
+} // namespace UE::Editor::DataStorage::Queries
diff --git a/Engine/Source/Runtime/TypedElementFramework/Tests/Elements/Framework/TypedElementDataStoragePerformanceTestCommands.cpp b/Engine/Source/Runtime/TypedElementFramework/Tests/Elements/Framework/TypedElementDataStoragePerformanceTestCommands.cpp
--- a/Engine/Source/Runtime/TypedElementFramework/Tests/Elements/Framework/TypedElementDataStoragePerformanceTestCommands.cpp
+++ b/Engine/Source/Runtime/TypedElementFramework/Tests/Elements/Framework/TypedElementDataStoragePerformanceTestCommands.cpp
@@ -40,7 +40,7 @@ namespace UE::Editor::DataStorage
 		TEXT("Removes all added rows for the performance test"),
 		FConsoleCommandWithArgsDelegate::CreateLambda([](const TArray<FString>& Args)
 		{
-			using namespace TypedElementQueryBuilder;
+			using namespace UE::Editor::DataStorage::Queries;
 
 			ITypedElementDataStorageInterface* DataStorageInterface = UTypedElementRegistry::GetInstance()->GetMutableDataStorage();
 		
@@ -76,7 +76,7 @@ void UTest_PingPongBetweenPhaseFactory::RegisterTables(ITypedElementDataStorageI
 // the next time A is run.
 void UTest_PingPongBetweenPhaseFactory::RegisterQueries(ITypedElementDataStorageInterface& DataStorage)
 {
-	using namespace TypedElementQueryBuilder;
+	using namespace UE::Editor::DataStorage::Queries;
 
 	Super::RegisterQueries(DataStorage);
 
@@ -99,50 +99,50 @@ void UTest_PingPongBetweenPhaseFactory::RegisterQueries(ITypedElementDataStorage
 					Context.RemoveColumns(Rows[Index], {RemovedColumnStruct});
 				}
 			})
-			.Where()
-			.Compile());
+		.Where()
+		.Compile());
+
+	DataStorage.RegisterQuery(
+		Select(TEXT("PingPong DurPhysics->PostPhysics"),
+			FProcessor(EQueryTickPhase::DuringPhysics, DataStorage.GetQueryTickGroupName(EQueryTickGroups::Default))
+				.SetExecutionMode(EExecutionMode::GameThread),
+			[](IQueryContext& Context, const RowHandle* RowPtr, FTest_PingPongDurPhys* PingPongBPtr)
+			{
+				QUICK_SCOPE_CYCLE_COUNTER(PingPong_During_Post);
+				TArrayView<const RowHandle> Rows(RowPtr, Context.GetRowCount());
+				TArrayView<FTest_PingPongDurPhys> PingPongAs(PingPongBPtr, Context.GetRowCount());
+	
+				UScriptStruct* RemovedColumnStruct = FTest_PingPongDurPhys::StaticStruct();
+				for (int32 Index = 0, End = Context.GetRowCount(); Index < End; ++Index)
+				{
+					uint64 Value = PingPongAs[Index].Value;
+					++Value;
+					Context.AddColumn(Rows[Index], FTest_PingPongPostPhys{.Value = Value});
+					Context.RemoveColumns(Rows[Index], {RemovedColumnStruct});
+				}
+			})
+		.Where()
+		.Compile());
 
 	DataStorage.RegisterQuery(
-	Select(TEXT("PingPong DurPhysics->PostPhysics"),
-		FProcessor(EQueryTickPhase::DuringPhysics, DataStorage.GetQueryTickGroupName(EQueryTickGroups::Default))
+	Select(TEXT("PingPong PostPhysics->PrePhysics"),
+		FProcessor(EQueryTickPhase::PostPhysics, DataStorage.GetQueryTickGroupName(EQueryTickGroups::Default))
 			.SetExecutionMode(EExecutionMode::GameThread),
-		[](IQueryContext& Context, const RowHandle* RowPtr, FTest_PingPongDurPhys* PingPongBPtr)
+		[](IQueryContext& Context, const RowHandle* RowPtr, FTest_PingPongPostPhys* PingPongBPtr)
 		{
-			QUICK_SCOPE_CYCLE_COUNTER(PingPong_During_Post);
+			QUICK_SCOPE_CYCLE_COUNTER(PingPong_Post_Pre);
 			TArrayView<const RowHandle> Rows(RowPtr, Context.GetRowCount());
-			TArrayView<FTest_PingPongDurPhys> PingPongAs(PingPongBPtr, Context.GetRowCount());
+			TArrayView<FTest_PingPongPostPhys> PingPongAs(PingPongBPtr, Context.GetRowCount());
 
-			UScriptStruct* RemovedColumnStruct = FTest_PingPongDurPhys::StaticStruct();
+			UScriptStruct* RemovedColumnStruct = FTest_PingPongPostPhys::StaticStruct();
 			for (int32 Index = 0, End = Context.GetRowCount(); Index < End; ++Index)
 			{
 				uint64 Value = PingPongAs[Index].Value;
 				++Value;
-				Context.AddColumn(Rows[Index], FTest_PingPongPostPhys{.Value = Value});
+				Context.AddColumn(Rows[Index], FTest_PingPongPrePhys{.Value = Value});
 				Context.RemoveColumns(Rows[Index], {RemovedColumnStruct});
 			}
 		})
-		.Where()
-		.Compile());
-
-	DataStorage.RegisterQuery(
-	Select(TEXT("PingPong PostPhysics->PrePhysics"),
-	FProcessor(EQueryTickPhase::PostPhysics, DataStorage.GetQueryTickGroupName(EQueryTickGroups::Default))
-		.SetExecutionMode(EExecutionMode::GameThread),
-	[](IQueryContext& Context, const RowHandle* RowPtr, FTest_PingPongPostPhys* PingPongBPtr)
-	{
-		QUICK_SCOPE_CYCLE_COUNTER(PingPong_Post_Pre);
-		TArrayView<const RowHandle> Rows(RowPtr, Context.GetRowCount());
-		TArrayView<FTest_PingPongPostPhys> PingPongAs(PingPongBPtr, Context.GetRowCount());
-
-		UScriptStruct* RemovedColumnStruct = FTest_PingPongPostPhys::StaticStruct();
-		for (int32 Index = 0, End = Context.GetRowCount(); Index < End; ++Index)
-		{
-			uint64 Value = PingPongAs[Index].Value;
-			++Value;
-			Context.AddColumn(Rows[Index], FTest_PingPongPrePhys{.Value = Value});
-			Context.RemoveColumns(Rows[Index], {RemovedColumnStruct});
-		}
-	})
 	.Where()
 	.Compile());
 }
\ No newline at end of file
diff --git a/Engine/Source/Runtime/TypedElementFramework/Tests/Elements/Framework/TypedElementDataStorageProcessorTests.cpp b/Engine/Source/Runtime/TypedElementFramework/Tests/Elements/Framework/TypedElementDataStorageProcessorTests.cpp
--- a/Engine/Source/Runtime/TypedElementFramework/Tests/Elements/Framework/TypedElementDataStorageProcessorTests.cpp
+++ b/Engine/Source/Runtime/TypedElementFramework/Tests/Elements/Framework/TypedElementDataStorageProcessorTests.cpp
@@ -26,8 +26,9 @@ namespace UE::Editor::DataStorage::Debug::ProcessorTests
 
 	void RegisterProcessors()
 	{
+		using namespace UE::Editor::DataStorage::Queries;
+
 		ITypedElementDataStorageInterface* DataStorage = UTypedElementRegistry::GetInstance()->GetMutableDataStorage();
-		using namespace TypedElementQueryBuilder;
 
 		if (PrimaryTable == InvalidTableHandle)
 		{
@@ -42,28 +43,28 @@ namespace UE::Editor::DataStorage::Debug::ProcessorTests
 
 		// Test creation of a row from within a query processor
 		QueryHandle PrimaryRowQuery = DataStorage->RegisterQuery(
-			Select(TEXT("TEST: Creating a row for primary reference column"),
-			FProcessor(EQueryTickPhase::PrePhysics, DataStorage->GetQueryTickGroupName(EQueryTickGroups::Default)),
-			[](IQueryContext& Context, const RowHandle* Rows, FTEDSProcessorTestsReferenceColumn* ReferenceColumns)
-			{
-				const int32 RowCount = Context.GetRowCount();
-				TConstArrayView<RowHandle> RowsView = MakeArrayView(Rows, RowCount);
-				TArrayView<FTEDSProcessorTestsReferenceColumn> ReferenceColumnsView = MakeArrayView(ReferenceColumns, RowCount);
-
-				for (int32 Index = 0; Index < RowCount; ++Index)
+			Select(
+				TEXT("TEST: Creating a row for primary reference column"),
+				FProcessor(EQueryTickPhase::PrePhysics, DataStorage->GetQueryTickGroupName(EQueryTickGroups::Default)),
+				[](IQueryContext& Context, const RowHandle* Rows, FTEDSProcessorTestsReferenceColumn* ReferenceColumns)
 				{
-					FTEDSProcessorTestsReferenceColumn& PrimaryReferenceColumn = ReferenceColumnsView[Index];
-					// Auto-create a secondary row if this points to no row
-					if (!Context.IsRowAvailable(PrimaryReferenceColumn.Reference))
+					const int32 RowCount = Context.GetRowCount();
+					TConstArrayView<RowHandle> RowsView = MakeArrayView(Rows, RowCount);
+					TArrayView<FTEDSProcessorTestsReferenceColumn> ReferenceColumnsView = MakeArrayView(ReferenceColumns, RowCount);
+
+					for (int32 Index = 0; Index < RowCount; ++Index)
 					{
-						RowHandle SecondaryRow = Context.AddRow(SecondaryTable);
-						// Initialize bi-directional row references
-						PrimaryReferenceColumn.Reference = SecondaryRow;
-						Context.AddColumn(SecondaryRow, FTEDSProcessorTestsReferenceColumn{.Reference = RowsView[Index]});
+						FTEDSProcessorTestsReferenceColumn& PrimaryReferenceColumn = ReferenceColumnsView[Index];
+						// Auto-create a secondary row if this points to no row
+						if (!Context.IsRowAvailable(PrimaryReferenceColumn.Reference))
+						{
+							RowHandle SecondaryRow = Context.AddRow(SecondaryTable);
+							// Initialize bi-directional row references
+							PrimaryReferenceColumn.Reference = SecondaryRow;
+							Context.AddColumn(SecondaryRow, FTEDSProcessorTestsReferenceColumn{.Reference = RowsView[Index]});
+						}
 					}
-				}
-			}
-			)
+				})
 			.Where()
 				.All<FTEDSProcessorTests_PrimaryTag>()
 				.None<FTEDSProcessorTests_Linked>()
@@ -83,33 +84,32 @@ namespace UE::Editor::DataStorage::Debug::ProcessorTests
 
 		QueryHandle SecondaryRowQuery = DataStorage->RegisterQuery(
 			Select(TEXT("TEST: Creating a row for secondary reference column"),
-			FProcessor(EQueryTickPhase::DuringPhysics, DataStorage->GetQueryTickGroupName(EQueryTickGroups::Default)),
-			[](IQueryContext& Context, const RowHandle* Rows, FTEDSProcessorTestsReferenceColumn* ReferenceColumns)
-			{
-				const int32 RowCount = Context.GetRowCount();
-				TConstArrayView<RowHandle> RowsView = MakeArrayView(Rows, RowCount);
-				TArrayView<FTEDSProcessorTestsReferenceColumn> ReferenceColumnsView = MakeArrayView(ReferenceColumns, RowCount);
-
-				for (int32 Index = 0; Index < RowCount; ++Index)
+				FProcessor(EQueryTickPhase::DuringPhysics, DataStorage->GetQueryTickGroupName(EQueryTickGroups::Default)),
+				[](IQueryContext& Context, const RowHandle* Rows, FTEDSProcessorTestsReferenceColumn* ReferenceColumns)
 				{
-					const RowHandle& SecondaryRow = RowsView[Index];
-					Context.RunSubquery(
-						0,
-						ReferenceColumnsView[Index].Reference,
-						CreateSubqueryCallbackBinding(
-						[SecondaryRow](ISubqueryContext& SubqueryContext, RowHandle PrimaryRow, const FTEDSProcessorTestsReferenceColumn& ReferenceColumn)
-						{
-							if (ReferenceColumn.Reference == SecondaryRow)
+					const int32 RowCount = Context.GetRowCount();
+					TConstArrayView<RowHandle> RowsView = MakeArrayView(Rows, RowCount);
+					TArrayView<FTEDSProcessorTestsReferenceColumn> ReferenceColumnsView = MakeArrayView(ReferenceColumns, RowCount);
+
+					for (int32 Index = 0; Index < RowCount; ++Index)
+					{
+						const RowHandle& SecondaryRow = RowsView[Index];
+						Context.RunSubquery(
+							0,
+							ReferenceColumnsView[Index].Reference,
+							CreateSubqueryCallbackBinding(
+							[SecondaryRow](ISubqueryContext& SubqueryContext, RowHandle PrimaryRow, const FTEDSProcessorTestsReferenceColumn& ReferenceColumn)
 							{
-								// Add these tags to prevent further processing
-								SubqueryContext.AddColumns<FTEDSProcessorTests_Linked>(SecondaryRow);
-								SubqueryContext.AddColumns<FTEDSProcessorTests_Linked>(PrimaryRow);
+								if (ReferenceColumn.Reference == SecondaryRow)
+								{
+									// Add these tags to prevent further processing
+									SubqueryContext.AddColumns<FTEDSProcessorTests_Linked>(SecondaryRow);
+									SubqueryContext.AddColumns<FTEDSProcessorTests_Linked>(PrimaryRow);
+								}
 							}
-						}
-					));
-				}
-			}
-			)
+						));
+					}
+				})
 			.Where()
 				.All<FTEDSProcessorTests_SecondaryTag>()
 				.None<FTEDSProcessorTests_Linked>()
diff --git a/Engine/Source/Runtime/TypedElementFramework/Tests/Elements/Framework/TypedElementQueryBuilderTests.cpp b/Engine/Source/Runtime/TypedElementFramework/Tests/Elements/Framework/TypedElementQueryBuilderTests.cpp
--- a/Engine/Source/Runtime/TypedElementFramework/Tests/Elements/Framework/TypedElementQueryBuilderTests.cpp
+++ b/Engine/Source/Runtime/TypedElementFramework/Tests/Elements/Framework/TypedElementQueryBuilderTests.cpp
@@ -11,11 +11,10 @@
 #include "Tests/TestHarnessAdapter.h"
 #include "UObject/Class.h"
 
-namespace UE::Editor::DataStorage
-{
-	using namespace Queries;
-namespace Tests
+namespace UE::Editor::DataStorage::Tests
 {
+	using namespace UE::Editor::DataStorage::Queries;
+
 static void AppendColumnName(FString& Output, TWeakObjectPtr<const UScriptStruct> TypeInfo)
 {
 #if WITH_EDITORONLY_DATA
@@ -101,7 +100,7 @@ static bool TestMatching(const FConditions& TestQuery, const TArray<FColumnBase>
 	return (Result == Expected);
 }
 
-TEST_CASE_NAMED(FTypedElementQueryConditions_NoColumn, "TypedElementQueryBuilder::FTypedElementQueryConditions_NoColumn", "[ApplicationContextMask][EngineFilter]")
+TEST_CASE_NAMED(FTypedElementQueryConditions_NoColumn, "Editor::DataStorage::QueryBuilder::FTypedElementQueryConditions_NoColumn", "[ApplicationContextMask][EngineFilter]")
 {
 	FConditions Example;
 	
@@ -110,7 +109,7 @@ TEST_CASE_NAMED(FTypedElementQueryConditions_NoColumn, "TypedElementQueryBuilder
 	CHECK(TestMatching(Example, { TColumn<FTestColumnA>() }, true));
 }
 
-TEST_CASE_NAMED(FTypedElementQueryConditions_OneColumn, "TypedElementQueryBuilder::FTypedElementQueryConditions_OneColumn", "[ApplicationContextMask][EngineFilter]")
+TEST_CASE_NAMED(FTypedElementQueryConditions_OneColumn, "Editor::DataStorage::Queries::FTypedElementQueryConditions_OneColumn", "[ApplicationContextMask][EngineFilter]")
 {
 	FConditions Example{ TColumn<FTestColumnA>() };
 
@@ -118,7 +117,7 @@ TEST_CASE_NAMED(FTypedElementQueryConditions_OneColumn, "TypedElementQueryBuilde
 	CHECK(TestMatching(Example, { TColumn<FTestColumnA>() }, true));
 }
 
-TEST_CASE_NAMED(FTypedElementQueryConditions1, "TypedElementQueryBuilder::FTypedElementQueryConditions A && B && C", "[ApplicationContextMask][EngineFilter]")
+TEST_CASE_NAMED(FTypedElementQueryConditions1, "Editor::DataStorage::QueryBuilder::FTypedElementQueryConditions A && B && C", "[ApplicationContextMask][EngineFilter]")
 {
 	TColumn<FTestColumnA> TestA;
 	
@@ -129,7 +128,7 @@ TEST_CASE_NAMED(FTypedElementQueryConditions1, "TypedElementQueryBuilder::FTyped
 	CHECK(TestMatching(Example, { TColumn<FTestColumnA>(), TColumn<FTestColumnB>(), TColumn<FTestColumnD>() }, false));
 }
 
-TEST_CASE_NAMED(FTypedElementQueryConditions2, "TypedElementQueryBuilder::FTypedElementQueryConditions A || B || C", "[ApplicationContextMask][EngineFilter]")
+TEST_CASE_NAMED(FTypedElementQueryConditions2, "Editor::DataStorage::QueryBuilder::FTypedElementQueryConditions A || B || C", "[ApplicationContextMask][EngineFilter]")
 {
 	FConditions Example = TColumn<FTestColumnA>() || TColumn<FTestColumnB>() || TColumn<FTestColumnC>();
 	
@@ -139,7 +138,7 @@ TEST_CASE_NAMED(FTypedElementQueryConditions2, "TypedElementQueryBuilder::FTyped
 	CHECK(TestMatching(Example, { TColumn<FTestColumnD>() }, false));
 }
 
-TEST_CASE_NAMED(FTypedElementQueryConditions3, "TypedElementQueryBuilder::FTypedElementQueryConditions A && (B || C)", "[ApplicationContextMask][EngineFilter]")
+TEST_CASE_NAMED(FTypedElementQueryConditions3, "Editor::DataStorage::QueryBuilder::FTypedElementQueryConditions A && (B || C)", "[ApplicationContextMask][EngineFilter]")
 {
 	FConditions Example = TColumn<FTestColumnA>() && (TColumn<FTestColumnB>() || TColumn<FTestColumnC>());
 	
@@ -150,7 +149,7 @@ TEST_CASE_NAMED(FTypedElementQueryConditions3, "TypedElementQueryBuilder::FTyped
 	CHECK(TestMatching(Example, { TColumn<FTestColumnD>(), TColumn<FTestColumnB>() }, false));
 }
 
-TEST_CASE_NAMED(FTypedElementQueryConditions4, "TypedElementQueryBuilder::FTypedElementQueryConditions A && (B || C) && (D || E)", "[ApplicationContextMask][EngineFilter]")
+TEST_CASE_NAMED(FTypedElementQueryConditions4, "Editor::DataStorage::QueryBuilder::FTypedElementQueryConditions A && (B || C) && (D || E)", "[ApplicationContextMask][EngineFilter]")
 {
 	FConditions Example = 
 		TColumn<FTestColumnA>() && 
@@ -171,7 +170,7 @@ TEST_CASE_NAMED(FTypedElementQueryConditions4, "TypedElementQueryBuilder::FTyped
 	CHECK(TestMatching(Example, { TColumn<FTestColumnB>(), TColumn<FTestColumnC>(), TColumn<FTestColumnD>() }, false));
 }
 
-TEST_CASE_NAMED(FTypedElementQueryConditions5, "TypedElementQueryBuilder::FTypedElementQueryConditions (A || B) && (C || D) && (E || F)", "[ApplicationContextMask][EngineFilter]")
+TEST_CASE_NAMED(FTypedElementQueryConditions5, "Editor::DataStorage::QueryBuilder::FTypedElementQueryConditions (A || B) && (C || D) && (E || F)", "[ApplicationContextMask][EngineFilter]")
 {
 	FConditions Example =
 		(TColumn<FTestColumnA>() || TColumn<FTestColumnB>()) &&
@@ -191,7 +190,7 @@ TEST_CASE_NAMED(FTypedElementQueryConditions5, "TypedElementQueryBuilder::FTyped
 	CHECK(TestMatching(Example, { TColumn<FTestColumnG>(), TColumn<FTestColumnC>(), TColumn<FTestColumnD>() }, false));
 }
 
-TEST_CASE_NAMED(FTypedElementQueryConditions6, "TypedElementQueryBuilder::FTypedElementQueryConditions ((A || B) && (C || D)) || (E && F)", "[ApplicationContextMask][EngineFilter]")
+TEST_CASE_NAMED(FTypedElementQueryConditions6, "Editor::DataStorage::QueryBuilder::FTypedElementQueryConditions ((A || B) && (C || D)) || (E && F)", "[ApplicationContextMask][EngineFilter]")
 {
 	FConditions Example =
 		(
@@ -209,7 +208,7 @@ TEST_CASE_NAMED(FTypedElementQueryConditions6, "TypedElementQueryBuilder::FTyped
 	CHECK(TestMatching(Example, { TColumn<FTestColumnG>() }, false));
 }
 
-TEST_CASE_NAMED(FTypedElementQueryConditions7, "TypedElementQueryBuilder::FTypedElementQueryConditions (A && B) || (C && D) || (E && F)", "[ApplicationContextMask][EngineFilter]")
+TEST_CASE_NAMED(FTypedElementQueryConditions7, "Editor::DataStorage::QueryBuilder::FTypedElementQueryConditions (A && B) || (C && D) || (E && F)", "[ApplicationContextMask][EngineFilter]")
 {
 	FConditions Example =
 		(TColumn<FTestColumnA>() && TColumn<FTestColumnB>()) ||
@@ -230,7 +229,7 @@ TEST_CASE_NAMED(FTypedElementQueryConditions7, "TypedElementQueryBuilder::FTyped
 	CHECK(TestMatching(Example, { TColumn<FTestColumnG>() }, false));
 }
 
-TEST_CASE_NAMED(FTypedElementQueryConditions_MultiMatch, "TypedElementQueryBuilder::FTypedElementQueryConditions_MultiMatch", "[ApplicationContextMask][EngineFilter]")
+TEST_CASE_NAMED(FTypedElementQueryConditions_MultiMatch, "Editor::DataStorage::QueryBuilder::FTypedElementQueryConditions_MultiMatch", "[ApplicationContextMask][EngineFilter]")
 {
 	FConditions Example =
 		(TColumn<FTestColumnA>() || TColumn<FTestColumnB>()) &&
@@ -256,7 +255,7 @@ TEST_CASE_NAMED(FTypedElementQueryConditions_MultiMatch, "TypedElementQueryBuild
 		}, true));
 }
 
-TEST_CASE_NAMED(FTypedElementQueryConditions_Sorted, "TypedElementQueryBuilder::FTypedElementQueryConditions_Sorted", "[ApplicationContextMask][EngineFilter]")
+TEST_CASE_NAMED(FTypedElementQueryConditions_Sorted, "Editor::DataStorage::QueryBuilder::FTypedElementQueryConditions_Sorted", "[ApplicationContextMask][EngineFilter]")
 {
 	FConditions Example =
 		(TColumn<FTestColumnA>() && TColumn<FTestColumnB>()) ||
@@ -271,7 +270,6 @@ TEST_CASE_NAMED(FTypedElementQueryConditions_Sorted, "TypedElementQueryBuilder::
 
 	CHECK(TestMatching(Example, { TColumn<FTestColumnC>(), TColumn<FTestColumnD>(), TColumn<FTestColumnE>(), TColumn<FTestColumnF>() }, true, true));
 }
-} // namespace Tests
-} // namespace UE::Editor::DataStorage
+} // namespace UE::Editor::DataStorage::Tests
 
 #endif // WITH_TESTS
