[TEDS] Remove unnecessary using in header that snuck in

#rnx

diff --git a/Engine/Source/Runtime/TypedElementFramework/Public/Elements/Framework/TypedElementQueryBuilder.h b/Engine/Source/Runtime/TypedElementFramework/Public/Elements/Framework/TypedElementQueryBuilder.h
--- a/Engine/Source/Runtime/TypedElementFramework/Public/Elements/Framework/TypedElementQueryBuilder.h
+++ b/Engine/Source/Runtime/TypedElementFramework/Public/Elements/Framework/TypedElementQueryBuilder.h
@@ -105,8 +105,6 @@ class USubsystem;
 
 namespace UE::Editor::DataStorage::Queries
 {
-	using namespace UE::Editor::DataStorage;
-
 	TYPEDELEMENTFRAMEWORK_API const UScriptStruct* Type(FTopLevelAssetPath Name);
 	TYPEDELEMENTFRAMEWORK_API const UScriptStruct* TypeOptional(FTopLevelAssetPath Name);
 	TYPEDELEMENTFRAMEWORK_API const UScriptStruct* operator""_Type(const char* Name, std::size_t NameSize);
