Fixed an issue in the template function: ColumnTypeTemplate* ICommonQueryWithEnvironmentContext::AddColumn(RowHandle Row, const FName& Identifier, ColumnTypeTemplate&& Column)

The type returned didn't match the signature.

diff --git a/Engine/Source/Runtime/TypedElementFramework/Public/Elements/Interfaces/TypedElementQueryStorageInterfaces.h b/Engine/Source/Runtime/TypedElementFramework/Public/Elements/Interfaces/TypedElementQueryStorageInterfaces.h
--- a/Engine/Source/Runtime/TypedElementFramework/Public/Elements/Interfaces/TypedElementQueryStorageInterfaces.h
+++ b/Engine/Source/Runtime/TypedElementFramework/Public/Elements/Interfaces/TypedElementQueryStorageInterfaces.h
@@ -204,7 +204,7 @@ namespace UE::Editor::DataStorage
 		 * @return An ephemeral reference to the data of the column.  The returned column will have been constructed.
 		 */
 		template<typename ColumnTypeTemplate>
-		ColumnTypeTemplate* AddColumn(RowHandle Row, const FName& Identifier, ColumnTypeTemplate&& Column);
+		ColumnTypeTemplate& AddColumn(RowHandle Row, const FName& Identifier, ColumnTypeTemplate&& Column);
 		/**
 		 * Adds new empty columns to a row of the provided type. The addition will not be immediately done but delayed until the end of the
 		 * tick group.
@@ -542,7 +542,7 @@ namespace UE::Editor::DataStorage
 	}
 
 	template <typename ColumnTypeTemplate>
-	ColumnTypeTemplate* ICommonQueryWithEnvironmentContext::AddColumn(RowHandle Row, const FName& Identifier, ColumnTypeTemplate&& Column)
+	ColumnTypeTemplate& ICommonQueryWithEnvironmentContext::AddColumn(RowHandle Row, const FName& Identifier, ColumnTypeTemplate&& Column)
 	{
 		UScriptStruct* TemplateType = ColumnTypeTemplate::StaticStruct();
 		const FDynamicColumnDescription Description
