[Teds] Rename interface types and base factory class

diff --git a/Engine/Plugins/Experimental/EditorDataStorageFeatures/Source/TedsStyling/Private/TedsStylingFactory.h b/Engine/Plugins/Experimental/EditorDataStorageFeatures/Source/TedsStyling/Private/TedsStylingFactory.h
--- a/Engine/Plugins/Experimental/EditorDataStorageFeatures/Source/TedsStyling/Private/TedsStylingFactory.h
+++ b/Engine/Plugins/Experimental/EditorDataStorageFeatures/Source/TedsStyling/Private/TedsStylingFactory.h
@@ -12,20 +12,20 @@ struct FSlateColor;
 class ISlateStyle;
 
 UCLASS()
-class UTedsStylingFactory : public UTypedElementDataStorageFactory
+class UTedsStylingFactory : public UEditorDataStorageFactory
 {
 	GENERATED_BODY()
 
 public:
 	~UTedsStylingFactory() override = default;
 
-	void RegisterTables(ITypedElementDataStorageInterface& DataStorage) override;
-	virtual void RegisterQueries(ITypedElementDataStorageInterface& DataStorage) override;
+	void RegisterTables(IEditorDataStorageProvider& DataStorage) override;
+	virtual void RegisterQueries(IEditorDataStorageProvider& DataStorage) override;
 
 	static void RegisterAllKnownStyles();
 	
 private:
-	static void RegisterBrush(ITypedElementDataStorageInterface* DataStorage, const FName& StyleName, const FSlateBrush* Brush, const ISlateStyle& OwnerStyle);
-	static void RegisterColor(ITypedElementDataStorageInterface* DataStorage, const FName& StyleName, const FSlateColor& Color, const ISlateStyle& OwnerStyle);
-	static UE::Editor::DataStorage::RowHandle AddOrGetStyleRow(ITypedElementDataStorageInterface* DataStorage, const FName& StyleName, const ISlateStyle& OwnerStyle);
+	static void RegisterBrush(IEditorDataStorageProvider* DataStorage, const FName& StyleName, const FSlateBrush* Brush, const ISlateStyle& OwnerStyle);
+	static void RegisterColor(IEditorDataStorageProvider* DataStorage, const FName& StyleName, const FSlateColor& Color, const ISlateStyle& OwnerStyle);
+	static UE::Editor::DataStorage::RowHandle AddOrGetStyleRow(IEditorDataStorageProvider* DataStorage, const FName& StyleName, const ISlateStyle& OwnerStyle);
 };
diff --git a/Engine/Plugins/Experimental/EditorDataStorageFeatures/Source/TedsStyling/Private/Widgets/SlateBrushWidget.cpp b/Engine/Plugins/Experimental/EditorDataStorageFeatures/Source/TedsStyling/Private/Widgets/SlateBrushWidget.cpp
--- a/Engine/Plugins/Experimental/EditorDataStorageFeatures/Source/TedsStyling/Private/Widgets/SlateBrushWidget.cpp
+++ b/Engine/Plugins/Experimental/EditorDataStorageFeatures/Source/TedsStyling/Private/Widgets/SlateBrushWidget.cpp
@@ -11,8 +11,8 @@
 #include "Widgets/SBoxPanel.h"
 #include "Widgets/Text/STextBlock.h"
 
-void USlateStylePreviewWidget::RegisterWidgetConstructors(ITypedElementDataStorageInterface& DataStorage,
-                                                          ITypedElementDataStorageUiInterface& DataStorageUi) const
+void USlateStylePreviewWidget::RegisterWidgetConstructors(IEditorDataStorageProvider& DataStorage,
+	IEditorDataStorageUiProvider& DataStorageUi) const
 {
 	using namespace UE::Editor::DataStorage::Queries;
 	DataStorageUi.RegisterWidgetFactory<FSlateStylePreviewWidgetConstructor>(TEXT("General.RowLabel"), TColumn<FNameColumn>() && TColumn<FSlateStyleTag>());
@@ -24,8 +24,8 @@ FSlateStylePreviewWidgetConstructor::FSlateStylePreviewWidgetConstructor()
 {
 }
 
-TSharedPtr<SWidget> FSlateStylePreviewWidgetConstructor::CreateWidget(ITypedElementDataStorageInterface* DataStorage,
-	ITypedElementDataStorageUiInterface* DataStorageUi, UE::Editor::DataStorage::RowHandle TargetRow, UE::Editor::DataStorage::RowHandle WidgetRow,
+TSharedPtr<SWidget> FSlateStylePreviewWidgetConstructor::CreateWidget(IEditorDataStorageProvider* DataStorage,
+	IEditorDataStorageUiProvider* DataStorageUi, UE::Editor::DataStorage::RowHandle TargetRow, UE::Editor::DataStorage::RowHandle WidgetRow,
 	const UE::Editor::DataStorage::FMetaDataView& Arguments)
 {
 	UE::Editor::DataStorage::FAttributeBinder Binder(TargetRow, DataStorage);
diff --git a/Engine/Plugins/Experimental/EditorDataStorageFeatures/Source/TedsStyling/Private/Widgets/SlateBrushWidget.h b/Engine/Plugins/Experimental/EditorDataStorageFeatures/Source/TedsStyling/Private/Widgets/SlateBrushWidget.h
--- a/Engine/Plugins/Experimental/EditorDataStorageFeatures/Source/TedsStyling/Private/Widgets/SlateBrushWidget.h
+++ b/Engine/Plugins/Experimental/EditorDataStorageFeatures/Source/TedsStyling/Private/Widgets/SlateBrushWidget.h
@@ -11,15 +11,15 @@ class ITypedElementDataStorageInterface;
 class UScriptStruct;
 
 UCLASS()
-class USlateStylePreviewWidget : public UTypedElementDataStorageFactory
+class USlateStylePreviewWidget : public UEditorDataStorageFactory
 {
 	GENERATED_BODY()
 
 public:
 	~USlateStylePreviewWidget() override = default;
 
-	void RegisterWidgetConstructors(ITypedElementDataStorageInterface& DataStorage,
-		ITypedElementDataStorageUiInterface& DataStorageUi) const override;
+	void RegisterWidgetConstructors(IEditorDataStorageProvider& DataStorage,
+		IEditorDataStorageUiProvider& DataStorageUi) const override;
 };
 
 // Widget to show a slate brush drawn as an SImage
@@ -33,8 +33,8 @@ public:
 	~FSlateStylePreviewWidgetConstructor() override = default;
 
 	virtual TSharedPtr<SWidget> CreateWidget(
-		ITypedElementDataStorageInterface* DataStorage,
-		ITypedElementDataStorageUiInterface* DataStorageUi,
+		IEditorDataStorageProvider* DataStorage,
+		IEditorDataStorageUiProvider* DataStorageUi,
 		UE::Editor::DataStorage::RowHandle TargetRow,
 		UE::Editor::DataStorage::RowHandle WidgetRow,
 		const UE::Editor::DataStorage::FMetaDataView& Arguments) override;
diff --git a/Engine/Plugins/Experimental/EditorDataStorageFeatures/Source/TedsTableViewer/Private/TedsTableViewerModel.cpp b/Engine/Plugins/Experimental/EditorDataStorageFeatures/Source/TedsTableViewer/Private/TedsTableViewerModel.cpp
--- a/Engine/Plugins/Experimental/EditorDataStorageFeatures/Source/TedsTableViewer/Private/TedsTableViewerModel.cpp
+++ b/Engine/Plugins/Experimental/EditorDataStorageFeatures/Source/TedsTableViewer/Private/TedsTableViewerModel.cpp
@@ -137,14 +137,14 @@ namespace UE::Editor::DataStorage
 		ColumnsView.Insert(InColumn, 0);
 	}
 
-	ITypedElementDataStorageInterface* FTedsTableViewerModel::GetDataStorageInterface() const
+	IEditorDataStorageProvider* FTedsTableViewerModel::GetDataStorageInterface() const
 	{
 		return Storage;
 	}
 
 	void FTedsTableViewerModel::GenerateColumns()
 	{
-		using MatchApproach = ITypedElementDataStorageUiInterface::EMatchApproach;
+		using MatchApproach = IEditorDataStorageUiProvider::EMatchApproach;
 		int32 IndexOffset = 0;
 
 		ColumnsView.Empty();
diff --git a/Engine/Plugins/Experimental/EditorDataStorageFeatures/Source/TedsTableViewer/Private/TedsTableViewerModule.cpp b/Engine/Plugins/Experimental/EditorDataStorageFeatures/Source/TedsTableViewer/Private/TedsTableViewerModule.cpp
--- a/Engine/Plugins/Experimental/EditorDataStorageFeatures/Source/TedsTableViewer/Private/TedsTableViewerModule.cpp
+++ b/Engine/Plugins/Experimental/EditorDataStorageFeatures/Source/TedsTableViewer/Private/TedsTableViewerModule.cpp
@@ -70,7 +70,7 @@ public:
 			return SNullWidget::NullWidget;
 		}
 
-		ITypedElementDataStorageInterface* DataStorage = Registry->GetMutableDataStorage();
+		IEditorDataStorageProvider* DataStorage = Registry->GetMutableDataStorage();
 
 		using namespace UE::Editor::DataStorage::Queries;
 
@@ -85,7 +85,7 @@ public:
 		Rows.Empty();
 		
 		UE::Editor::DataStorage::FQueryResult QueryResult = DataStorage->RunQuery(QueryHandle,
-			CreateDirectQueryCallbackBinding([this](const ITypedElementDataStorageInterface::IDirectQueryContext& Context, const UE::Editor::DataStorage::RowHandle* RowHandles)
+			CreateDirectQueryCallbackBinding([this](const IEditorDataStorageProvider::IDirectQueryContext& Context, const UE::Editor::DataStorage::RowHandle* RowHandles)
 			{
 				Rows.Append(RowHandles, Context.GetRowCount());
 			}));
diff --git a/Engine/Plugins/Experimental/EditorDataStorageFeatures/Source/TedsTableViewer/Private/TedsTableViewerUtils.cpp b/Engine/Plugins/Experimental/EditorDataStorageFeatures/Source/TedsTableViewer/Private/TedsTableViewerUtils.cpp
--- a/Engine/Plugins/Experimental/EditorDataStorageFeatures/Source/TedsTableViewer/Private/TedsTableViewerUtils.cpp
+++ b/Engine/Plugins/Experimental/EditorDataStorageFeatures/Source/TedsTableViewer/Private/TedsTableViewerUtils.cpp
@@ -90,9 +90,9 @@ namespace UE::Editor::DataStorage::TableViewerUtils
 		return VerifiedColumnTypes;
 	}
 
-	TSharedPtr<FTypedElementWidgetConstructor> CreateHeaderWidgetConstructor(ITypedElementDataStorageUiInterface& StorageUi, const UE::Editor::DataStorage::FMetaDataView& InMetaData, TConstArrayView<TWeakObjectPtr<const UScriptStruct>> ColumnTypes, const TConstArrayView<FName> CellWidgetPurposes)
+	TSharedPtr<FTypedElementWidgetConstructor> CreateHeaderWidgetConstructor(IEditorDataStorageUiProvider& StorageUi, const UE::Editor::DataStorage::FMetaDataView& InMetaData, TConstArrayView<TWeakObjectPtr<const UScriptStruct>> ColumnTypes, const TConstArrayView<FName> CellWidgetPurposes)
 	{
-		using MatchApproach = ITypedElementDataStorageUiInterface::EMatchApproach;
+		using MatchApproach = IEditorDataStorageUiProvider::EMatchApproach;
 
 		TArray<TWeakObjectPtr<const UScriptStruct>> VerifiedColumnTypes = CreateVerifiedColumnTypeArray(ColumnTypes);
 		TSharedPtr<FTypedElementWidgetConstructor> Constructor;
@@ -145,7 +145,7 @@ namespace UE::Editor::DataStorage::TableViewerUtils
 		return nullptr;
 	}
 
-	const FSlateBrush* GetIconForRow(ITypedElementDataStorageInterface* DataStorage, RowHandle Row)
+	const FSlateBrush* GetIconForRow(IEditorDataStorageProvider* DataStorage, RowHandle Row)
 	{
 		static TMap<FName, const FSlateBrush*> CachedIconMap;
 
@@ -201,7 +201,7 @@ namespace UE::Editor::DataStorage::TableViewerUtils
 	}
 }
 
-void UTypedElementTableViewerFactory::RegisterTables(ITypedElementDataStorageInterface& DataStorage)
+void UTypedElementTableViewerFactory::RegisterTables(IEditorDataStorageProvider& DataStorage)
 {
 	using namespace UE::Editor::DataStorage;
 	const TableHandle BaseWidgetTable = DataStorage.FindTable(FName(TEXT("Editor_WidgetTable")));
diff --git a/Engine/Plugins/Experimental/EditorDataStorageFeatures/Source/TedsTableViewer/Private/Widgets/RowHandleWidget.cpp b/Engine/Plugins/Experimental/EditorDataStorageFeatures/Source/TedsTableViewer/Private/Widgets/RowHandleWidget.cpp
--- a/Engine/Plugins/Experimental/EditorDataStorageFeatures/Source/TedsTableViewer/Private/Widgets/RowHandleWidget.cpp
+++ b/Engine/Plugins/Experimental/EditorDataStorageFeatures/Source/TedsTableViewer/Private/Widgets/RowHandleWidget.cpp
@@ -9,20 +9,20 @@
 #define LOCTEXT_NAMESPACE "RowHandleWidget"
 
 
-void URowHandleWidgetFactory::RegisterWidgetConstructors(ITypedElementDataStorageInterface& DataStorage, ITypedElementDataStorageUiInterface& DataStorageUi) const
+void URowHandleWidgetFactory::RegisterWidgetConstructors(IEditorDataStorageProvider& DataStorage, IEditorDataStorageUiProvider& DataStorageUi) const
 {
 	DataStorageUi.RegisterWidgetFactory(FName(TEXT("General.Cell.RowHandle")), FRowHandleWidgetConstructor::StaticStruct());
 }
 
-void URowHandleWidgetFactory::RegisterWidgetPurposes(ITypedElementDataStorageUiInterface& DataStorageUi) const
+void URowHandleWidgetFactory::RegisterWidgetPurposes(IEditorDataStorageUiProvider& DataStorageUi) const
 {
-	DataStorageUi.RegisterWidgetPurpose(FName(TEXT("General.Cell.RowHandle")), ITypedElementDataStorageUiInterface::EPurposeType::UniqueByName,
+	DataStorageUi.RegisterWidgetPurpose(FName(TEXT("General.Cell.RowHandle")), IEditorDataStorageUiProvider::EPurposeType::UniqueByName,
 	LOCTEXT("GeneralRowHandlePurpose", "Specific purpose to request a widget to display row handles."));
 
-	DataStorageUi.RegisterWidgetPurpose(FName(TEXT("RowDetails.Cell")), ITypedElementDataStorageUiInterface::EPurposeType::UniqueByNameAndColumn,
+	DataStorageUi.RegisterWidgetPurpose(FName(TEXT("RowDetails.Cell")), IEditorDataStorageUiProvider::EPurposeType::UniqueByNameAndColumn,
 	LOCTEXT("DetailsRowHandlePurpose", "Specific purpose to request a widget to display the details on a row (e.g SRowDetails)."));
 
-	DataStorageUi.RegisterWidgetPurpose(FName(TEXT("RowDetails.Cell.Large")), ITypedElementDataStorageUiInterface::EPurposeType::UniqueByNameAndColumn,
+	DataStorageUi.RegisterWidgetPurpose(FName(TEXT("RowDetails.Cell.Large")), IEditorDataStorageUiProvider::EPurposeType::UniqueByNameAndColumn,
 	LOCTEXT("LargeDetailsRowHandlePurpose", "Specific purpose to request a widget that is larger than a single cell to display the details on a row (e.g SRowDetails)"));
 
 
@@ -41,7 +41,7 @@ TSharedPtr<SWidget> FRowHandleWidgetConstructor::CreateWidget(const UE::Editor::
 		.VAlign(VAlign_Center)
 		.Padding(8, 0, 0, 0);}
 
-bool FRowHandleWidgetConstructor::FinalizeWidget(ITypedElementDataStorageInterface* DataStorage, ITypedElementDataStorageUiInterface* DataStorageUi, UE::Editor::DataStorage::RowHandle Row, const TSharedPtr<SWidget>& Widget)
+bool FRowHandleWidgetConstructor::FinalizeWidget(IEditorDataStorageProvider* DataStorage, IEditorDataStorageUiProvider* DataStorageUi, UE::Editor::DataStorage::RowHandle Row, const TSharedPtr<SWidget>& Widget)
 {
 
 	checkf(Widget->GetType() == SBox::StaticWidgetClass().GetWidgetType(),
diff --git a/Engine/Plugins/Experimental/EditorDataStorageFeatures/Source/TedsTableViewer/Private/Widgets/SRowDetails.cpp b/Engine/Plugins/Experimental/EditorDataStorageFeatures/Source/TedsTableViewer/Private/Widgets/SRowDetails.cpp
--- a/Engine/Plugins/Experimental/EditorDataStorageFeatures/Source/TedsTableViewer/Private/Widgets/SRowDetails.cpp
+++ b/Engine/Plugins/Experimental/EditorDataStorageFeatures/Source/TedsTableViewer/Private/Widgets/SRowDetails.cpp
@@ -81,7 +81,7 @@ namespace UE::Editor::DataStorage
 			
 			for (const FName& Purpose : WidgetPurposes)
 			{
-				DataStorageUi->CreateWidgetConstructors(Purpose, ITypedElementDataStorageUiInterface::EMatchApproach::LongestMatch,
+				DataStorageUi->CreateWidgetConstructors(Purpose, IEditorDataStorageUiProvider::EMatchApproach::LongestMatch,
 					Columns, {}, [this, Row](
 						TUniquePtr<FTypedElementWidgetConstructor> Constructor, TConstArrayView<TWeakObjectPtr<const UScriptStruct>> Columns)
 					{
@@ -130,8 +130,8 @@ namespace UE::Editor::DataStorage
 	//
 	// SRowDetailsRow
 	//
-	void SRowDetailsRow::Construct(const FArguments& Args, const TSharedRef<STableViewBase>& OwnerTableView,ITypedElementDataStorageInterface* InDataStorage,
-			ITypedElementDataStorageUiInterface* InDataStorageUi)
+	void SRowDetailsRow::Construct(const FArguments& Args, const TSharedRef<STableViewBase>& OwnerTableView,IEditorDataStorageProvider* InDataStorage,
+			IEditorDataStorageUiProvider* InDataStorageUi)
 	{
 		Item = Args._Item;
 
diff --git a/Engine/Plugins/Experimental/EditorDataStorageFeatures/Source/TedsTableViewer/Private/Widgets/STedsTableViewer.cpp b/Engine/Plugins/Experimental/EditorDataStorageFeatures/Source/TedsTableViewer/Private/Widgets/STedsTableViewer.cpp
--- a/Engine/Plugins/Experimental/EditorDataStorageFeatures/Source/TedsTableViewer/Private/Widgets/STedsTableViewer.cpp
+++ b/Engine/Plugins/Experimental/EditorDataStorageFeatures/Source/TedsTableViewer/Private/Widgets/STedsTableViewer.cpp
@@ -62,7 +62,7 @@ namespace UE::Editor::DataStorage
 	{
 		TedsWidget = SNew(STedsWidget);
 
-		if(ITypedElementDataStorageInterface* DataStorage = Model->GetDataStorageInterface())
+		if(IEditorDataStorageProvider* DataStorage = Model->GetDataStorageInterface())
 		{
 			const RowHandle WidgetRowHandle = TedsWidget->GetRowHandle();
 		
diff --git a/Engine/Plugins/Experimental/EditorDataStorageFeatures/Source/TedsTableViewer/Private/Widgets/UObjectLabelWidget.cpp b/Engine/Plugins/Experimental/EditorDataStorageFeatures/Source/TedsTableViewer/Private/Widgets/UObjectLabelWidget.cpp
--- a/Engine/Plugins/Experimental/EditorDataStorageFeatures/Source/TedsTableViewer/Private/Widgets/UObjectLabelWidget.cpp
+++ b/Engine/Plugins/Experimental/EditorDataStorageFeatures/Source/TedsTableViewer/Private/Widgets/UObjectLabelWidget.cpp
@@ -14,8 +14,8 @@
 
 #define LOCTEXT_NAMESPACE "FUObjectLabelWidgetConstructor"
 
-void UUObjectLabelWidgetFactory::RegisterWidgetConstructors(ITypedElementDataStorageInterface& DataStorage,
-	ITypedElementDataStorageUiInterface& DataStorageUi) const
+void UUObjectLabelWidgetFactory::RegisterWidgetConstructors(IEditorDataStorageProvider& DataStorage,
+	IEditorDataStorageUiProvider& DataStorageUi) const
 {
 	using namespace UE::Editor::DataStorage::Queries;
 	DataStorageUi.RegisterWidgetFactory<FUObjectLabelWidgetConstructor>(
@@ -33,8 +33,8 @@ FUObjectLabelWidgetConstructor::FUObjectLabelWidgetConstructor(const UScriptStru
 {
 }
 
-TSharedPtr<SWidget> FUObjectLabelWidgetConstructor::CreateWidget(ITypedElementDataStorageInterface* DataStorage,
-	ITypedElementDataStorageUiInterface* DataStorageUi, RowHandle TargetRow, RowHandle WidgetRow,
+TSharedPtr<SWidget> FUObjectLabelWidgetConstructor::CreateWidget(IEditorDataStorageProvider* DataStorage,
+	IEditorDataStorageUiProvider* DataStorageUi, RowHandle TargetRow, RowHandle WidgetRow,
 	const UE::Editor::DataStorage::FMetaDataView& Arguments)
 {
 	using namespace UE::Editor::DataStorage;
diff --git a/Engine/Plugins/Experimental/EditorDataStorageFeatures/Source/TedsTableViewer/Public/TedsTableViewerColumn.h b/Engine/Plugins/Experimental/EditorDataStorageFeatures/Source/TedsTableViewer/Public/TedsTableViewerColumn.h
--- a/Engine/Plugins/Experimental/EditorDataStorageFeatures/Source/TedsTableViewer/Public/TedsTableViewerColumn.h
+++ b/Engine/Plugins/Experimental/EditorDataStorageFeatures/Source/TedsTableViewer/Public/TedsTableViewerColumn.h
@@ -7,9 +7,9 @@
 #include "UObject/NameTypes.h"
 #include "Widgets/Views/SHeaderRow.h"
 
-class ITypedElementDataStorageCompatibilityInterface;
-class ITypedElementDataStorageUiInterface;
-class ITypedElementDataStorageInterface;
+class IEditorDataStorageCompatibilityProvider;
+class IEditorDataStorageUiProvider;
+class IEditorDataStorageProvider;
 struct FTypedElementWidgetConstructor;
 class SWidget;
 
@@ -70,9 +70,9 @@ namespace UE::Editor::DataStorage
 		FMetaDataView WidgetMetaData;
 
 		// TEDS Constructs
-		ITypedElementDataStorageInterface* Storage;
-		ITypedElementDataStorageUiInterface* StorageUi;
-		ITypedElementDataStorageCompatibilityInterface* StorageCompatibility;
+		IEditorDataStorageProvider* Storage;
+		IEditorDataStorageUiProvider* StorageUi;
+		IEditorDataStorageCompatibilityProvider* StorageCompatibility;
 
 		// Queries used to virtualize widgets when a column is added to/remove from a row
 		TArray<QueryHandle> InternalObserverQueries;
diff --git a/Engine/Plugins/Experimental/EditorDataStorageFeatures/Source/TedsTableViewer/Public/TedsTableViewerModel.h b/Engine/Plugins/Experimental/EditorDataStorageFeatures/Source/TedsTableViewer/Public/TedsTableViewerModel.h
--- a/Engine/Plugins/Experimental/EditorDataStorageFeatures/Source/TedsTableViewer/Public/TedsTableViewerModel.h
+++ b/Engine/Plugins/Experimental/EditorDataStorageFeatures/Source/TedsTableViewer/Public/TedsTableViewerModel.h
@@ -9,9 +9,9 @@
 #include "TypedElementUITypes.h"
 #include "UObject/NameTypes.h"
 
-class ITypedElementDataStorageCompatibilityInterface;
-class ITypedElementDataStorageUiInterface;
-class ITypedElementDataStorageInterface;
+class IEditorDataStorageCompatibilityProvider;
+class IEditorDataStorageUiProvider;
+class IEditorDataStorageProvider;
 
 namespace UE::Editor::DataStorage
 {
@@ -59,7 +59,7 @@ namespace UE::Editor::DataStorage
 		// Add a custom column to display in the table viewer, that doesn't necessarily map to a Teds column
 		TEDSTABLEVIEWER_API void AddCustomColumn(const TSharedRef<FTedsTableViewerColumn>& InColumn);
 
-		TEDSTABLEVIEWER_API ITypedElementDataStorageInterface* GetDataStorageInterface() const;
+		TEDSTABLEVIEWER_API IEditorDataStorageProvider* GetDataStorageInterface() const;
 
 	protected:
 
@@ -105,8 +105,8 @@ namespace UE::Editor::DataStorage
 		FOnModelChanged OnModelChanged;
 		
 		// Teds Constructs
-		ITypedElementDataStorageInterface* Storage = nullptr;
-		ITypedElementDataStorageUiInterface* StorageUi = nullptr;
-		ITypedElementDataStorageCompatibilityInterface* StorageCompatibility = nullptr;
+		IEditorDataStorageProvider* Storage = nullptr;
+		IEditorDataStorageUiProvider* StorageUi = nullptr;
+		IEditorDataStorageCompatibilityProvider* StorageCompatibility = nullptr;
 	};
 } // namespace UE::Editor::DataStorage
diff --git a/Engine/Plugins/Experimental/EditorDataStorageFeatures/Source/TedsTableViewer/Public/TedsTableViewerUtils.h b/Engine/Plugins/Experimental/EditorDataStorageFeatures/Source/TedsTableViewer/Public/TedsTableViewerUtils.h
--- a/Engine/Plugins/Experimental/EditorDataStorageFeatures/Source/TedsTableViewer/Public/TedsTableViewerUtils.h
+++ b/Engine/Plugins/Experimental/EditorDataStorageFeatures/Source/TedsTableViewer/Public/TedsTableViewerUtils.h
@@ -14,8 +14,8 @@ namespace UE::Editor::DataStorage
 	class FMetaDataView;
 }
 
-class ITypedElementDataStorageInterface;
-class ITypedElementDataStorageUiInterface;
+class IEditorDataStorageProvider;
+class IEditorDataStorageUiProvider;
 struct FTypedElementWidgetConstructor;
 struct FSlateBrush;
 
@@ -28,7 +28,7 @@ namespace UE::Editor::DataStorage::TableViewerUtils
 	TEDSTABLEVIEWER_API FName FindLongestMatchingName(TConstArrayView<TWeakObjectPtr<const UScriptStruct>> ColumnTypes, int32 DefaultNameIndex);
 
 	// Create a header widget constructor for the given columns
-	TEDSTABLEVIEWER_API TSharedPtr<FTypedElementWidgetConstructor> CreateHeaderWidgetConstructor(ITypedElementDataStorageUiInterface& StorageUi, 
+	TEDSTABLEVIEWER_API TSharedPtr<FTypedElementWidgetConstructor> CreateHeaderWidgetConstructor(IEditorDataStorageUiProvider& StorageUi, 
 	const UE::Editor::DataStorage::FMetaDataView& InMetaData, TConstArrayView<TWeakObjectPtr<const UScriptStruct>> ColumnTypes,
 	const TConstArrayView<FName> CellWidgetPurposes);
 
@@ -36,17 +36,17 @@ namespace UE::Editor::DataStorage::TableViewerUtils
 	TEDSTABLEVIEWER_API TArray<TWeakObjectPtr<const UScriptStruct>> CreateVerifiedColumnTypeArray(TConstArrayView<TWeakObjectPtr<const UScriptStruct>> ColumnTypes);
 	
 	// Scans the row for columns and tries to find the best matching icon.
-	TEDSTABLEVIEWER_API const FSlateBrush* GetIconForRow(ITypedElementDataStorageInterface* DataStorage, RowHandle Row);
+	TEDSTABLEVIEWER_API const FSlateBrush* GetIconForRow(IEditorDataStorageProvider* DataStorage, RowHandle Row);
 } // namespace UE::Editor::DataStorage::TableViewerUtils
 
 UCLASS()
-class UTypedElementTableViewerFactory : public UTypedElementDataStorageFactory
+class UTypedElementTableViewerFactory : public UEditorDataStorageFactory
 {
 	GENERATED_BODY()
 
 public:
 	~UTypedElementTableViewerFactory() override = default;
 
-	void RegisterTables(ITypedElementDataStorageInterface& DataStorage) override;
+	void RegisterTables(IEditorDataStorageProvider& DataStorage) override;
 };
 
diff --git a/Engine/Plugins/Experimental/EditorDataStorageFeatures/Source/TedsTableViewer/Public/Widgets/RowHandleWidget.h b/Engine/Plugins/Experimental/EditorDataStorageFeatures/Source/TedsTableViewer/Public/Widgets/RowHandleWidget.h
--- a/Engine/Plugins/Experimental/EditorDataStorageFeatures/Source/TedsTableViewer/Public/Widgets/RowHandleWidget.h
+++ b/Engine/Plugins/Experimental/EditorDataStorageFeatures/Source/TedsTableViewer/Public/Widgets/RowHandleWidget.h
@@ -8,21 +8,21 @@
 
 #include "RowHandleWidget.generated.h"
 
-class ITypedElementDataStorageInterface;
+class IEditorDataStorageProvider;
 class UScriptStruct;
 
 UCLASS()
-class TEDSTABLEVIEWER_API URowHandleWidgetFactory : public UTypedElementDataStorageFactory
+class TEDSTABLEVIEWER_API URowHandleWidgetFactory : public UEditorDataStorageFactory
 {
 	GENERATED_BODY()
 
 public:
 	virtual ~URowHandleWidgetFactory() override = default;
 
-	virtual void RegisterWidgetConstructors(ITypedElementDataStorageInterface& DataStorage,
-		ITypedElementDataStorageUiInterface& DataStorageUi) const override;
+	virtual void RegisterWidgetConstructors(IEditorDataStorageProvider& DataStorage,
+		IEditorDataStorageUiProvider& DataStorageUi) const override;
 
-	void RegisterWidgetPurposes(ITypedElementDataStorageUiInterface& DataStorageUi) const override;
+	void RegisterWidgetPurposes(IEditorDataStorageUiProvider& DataStorageUi) const override;
 };
 
 // A custom widget to display the row handle of a row as text
@@ -38,8 +38,8 @@ public:
 protected:
 	virtual TSharedPtr<SWidget> CreateWidget(const UE::Editor::DataStorage::FMetaDataView& Arguments) override;
 	virtual bool FinalizeWidget(
-		ITypedElementDataStorageInterface* DataStorage,
-		ITypedElementDataStorageUiInterface* DataStorageUi,
+		IEditorDataStorageProvider* DataStorage,
+		IEditorDataStorageUiProvider* DataStorageUi,
 		UE::Editor::DataStorage::RowHandle Row,
 		const TSharedPtr<SWidget>& Widget) override;
 };
diff --git a/Engine/Plugins/Experimental/EditorDataStorageFeatures/Source/TedsTableViewer/Public/Widgets/SRowDetails.h b/Engine/Plugins/Experimental/EditorDataStorageFeatures/Source/TedsTableViewer/Public/Widgets/SRowDetails.h
--- a/Engine/Plugins/Experimental/EditorDataStorageFeatures/Source/TedsTableViewer/Public/Widgets/SRowDetails.h
+++ b/Engine/Plugins/Experimental/EditorDataStorageFeatures/Source/TedsTableViewer/Public/Widgets/SRowDetails.h
@@ -6,8 +6,8 @@
 #include "Widgets/SCompoundWidget.h"
 #include "Widgets/Views/SListView.h"
 
-class ITypedElementDataStorageInterface;
-class ITypedElementDataStorageUiInterface;
+class IEditorDataStorageProvider;
+class IEditorDataStorageUiProvider;
 struct FTypedElementWidgetConstructor;
 
 namespace UE::Editor::DataStorage
@@ -66,8 +66,8 @@ namespace UE::Editor::DataStorage
 
 		TArray<RowDetailsItemPtr> Items;
 
-		ITypedElementDataStorageInterface* DataStorage = nullptr; 
-		ITypedElementDataStorageUiInterface* DataStorageUi = nullptr;
+		IEditorDataStorageProvider* DataStorage = nullptr; 
+		IEditorDataStorageUiProvider* DataStorageUi = nullptr;
 
 		bool bShowAllDetails = true;
 
@@ -85,14 +85,14 @@ namespace UE::Editor::DataStorage
 
 		SLATE_END_ARGS()
 		
-		void Construct(const FArguments& Args, const TSharedRef<STableViewBase>& OwnerTableView, ITypedElementDataStorageInterface* InDataStorage,
-			ITypedElementDataStorageUiInterface* InDataStorageUi);
+		void Construct(const FArguments& Args, const TSharedRef<STableViewBase>& OwnerTableView, IEditorDataStorageProvider* InDataStorage,
+			IEditorDataStorageUiProvider* InDataStorageUi);
 		
 		TSharedRef<SWidget> GenerateWidgetForColumn(const FName& ColumnName) override;
 
 	private:
 		RowDetailsItemPtr Item;
-		ITypedElementDataStorageInterface* DataStorage = nullptr;
-		ITypedElementDataStorageUiInterface* DataStorageUi = nullptr;
+		IEditorDataStorageProvider* DataStorage = nullptr;
+		IEditorDataStorageUiProvider* DataStorageUi = nullptr;
 	};
 } // namespace UE::Editor::DataStorage
diff --git a/Engine/Plugins/Experimental/EditorDataStorageFeatures/Source/TedsTableViewer/Public/Widgets/UObjectLabelWidget.h b/Engine/Plugins/Experimental/EditorDataStorageFeatures/Source/TedsTableViewer/Public/Widgets/UObjectLabelWidget.h
--- a/Engine/Plugins/Experimental/EditorDataStorageFeatures/Source/TedsTableViewer/Public/Widgets/UObjectLabelWidget.h
+++ b/Engine/Plugins/Experimental/EditorDataStorageFeatures/Source/TedsTableViewer/Public/Widgets/UObjectLabelWidget.h
@@ -9,19 +9,19 @@
 
 #include "UObjectLabelWidget.generated.h"
 
-class ITypedElementDataStorageInterface;
+class IEditorDataStorageProvider;
 class UScriptStruct;
 
 UCLASS()
-class UUObjectLabelWidgetFactory : public UTypedElementDataStorageFactory
+class UUObjectLabelWidgetFactory : public UEditorDataStorageFactory
 {
 	GENERATED_BODY()
 
 public:
 	~UUObjectLabelWidgetFactory() override = default;
 
-	TEDSTABLEVIEWER_API void RegisterWidgetConstructors(ITypedElementDataStorageInterface& DataStorage,
-		ITypedElementDataStorageUiInterface& DataStorageUi) const override;
+	TEDSTABLEVIEWER_API void RegisterWidgetConstructors(IEditorDataStorageProvider& DataStorage,
+		IEditorDataStorageUiProvider& DataStorageUi) const override;
 };
 
 // Widget to show labels for UObjects in TEDS
@@ -36,8 +36,8 @@ public:
 	~FUObjectLabelWidgetConstructor() override = default;
 
 	TEDSTABLEVIEWER_API virtual TSharedPtr<SWidget> CreateWidget(
-		ITypedElementDataStorageInterface* DataStorage,
-		ITypedElementDataStorageUiInterface* DataStorageUi,
+		IEditorDataStorageProvider* DataStorage,
+		IEditorDataStorageUiProvider* DataStorageUi,
 		UE::Editor::DataStorage::RowHandle TargetRow,
 		UE::Editor::DataStorage::RowHandle WidgetRow, 
 		const UE::Editor::DataStorage::FMetaDataView& Arguments) override;
diff --git a/Engine/Plugins/Experimental/InstanceDataObjectFixupTool/Source/Private/InstanceDataObjectFixupPanel.cpp b/Engine/Plugins/Experimental/InstanceDataObjectFixupTool/Source/Private/InstanceDataObjectFixupPanel.cpp
--- a/Engine/Plugins/Experimental/InstanceDataObjectFixupTool/Source/Private/InstanceDataObjectFixupPanel.cpp
+++ b/Engine/Plugins/Experimental/InstanceDataObjectFixupTool/Source/Private/InstanceDataObjectFixupPanel.cpp
@@ -251,8 +251,8 @@ FInstanceDataObjectFixupPanel::~FInstanceDataObjectFixupPanel()
 	using namespace UE::Editor::DataStorage;
 
 	UTypedElementRegistry* Registry = UTypedElementRegistry::GetInstance();
-	ITypedElementDataStorageInterface* DataStorage = Registry ? Registry->GetMutableDataStorage() : nullptr;
-	ITypedElementDataStorageCompatibilityInterface* DataStorageCompatibility = Registry ? Registry->GetMutableDataStorageCompatibility() : nullptr;
+	IEditorDataStorageProvider* DataStorage = Registry ? Registry->GetMutableDataStorage() : nullptr;
+	IEditorDataStorageCompatibilityProvider* DataStorageCompatibility = Registry ? Registry->GetMutableDataStorageCompatibility() : nullptr;
 
 	if (DataStorageCompatibility != nullptr && DataStorage != nullptr)
 	{
diff --git a/Engine/Plugins/Experimental/InstanceDataObjectFixupTool/Source/Private/InstanceDataObjectFixupToolTedsQueries.cpp b/Engine/Plugins/Experimental/InstanceDataObjectFixupTool/Source/Private/InstanceDataObjectFixupToolTedsQueries.cpp
--- a/Engine/Plugins/Experimental/InstanceDataObjectFixupTool/Source/Private/InstanceDataObjectFixupToolTedsQueries.cpp
+++ b/Engine/Plugins/Experimental/InstanceDataObjectFixupTool/Source/Private/InstanceDataObjectFixupToolTedsQueries.cpp
@@ -15,7 +15,7 @@
 #include "Widgets/Notifications/SNotificationList.h"
 #define LOCTEXT_NAMESPACE "FixupToolTedsQueries"
 
-void UInstanceDataObjectFixupToolTedsQueryFactory::RegisterQueries(ITypedElementDataStorageInterface& DataStorage)
+void UInstanceDataObjectFixupToolTedsQueryFactory::RegisterQueries(IEditorDataStorageProvider& DataStorage)
 {
 	using namespace UE::Editor::DataStorage::Queries;
 	
@@ -61,7 +61,7 @@ void UInstanceDataObjectFixupToolTedsQueryFactory::ShowFixUpToolForLooseProperti
 
 void UInstanceDataObjectFixupToolTedsQueryFactory::ShowFixUpTool(UE::Editor::DataStorage::RowHandle Row, bool bRecurseIntoObject)
 {
-	ITypedElementDataStorageInterface* DataStorage = UTypedElementRegistry::GetInstance()->GetMutableDataStorage();
+	IEditorDataStorageProvider* DataStorage = UTypedElementRegistry::GetInstance()->GetMutableDataStorage();
 	if (FTypedElementUObjectColumn* ObjectColumn = DataStorage->GetColumn<FTypedElementUObjectColumn>(Row))
 	{
 		if (bRecurseIntoObject)
diff --git a/Engine/Plugins/Experimental/InstanceDataObjectFixupTool/Source/Private/InstanceDataObjectFixupToolTedsQueries.h b/Engine/Plugins/Experimental/InstanceDataObjectFixupTool/Source/Private/InstanceDataObjectFixupToolTedsQueries.h
--- a/Engine/Plugins/Experimental/InstanceDataObjectFixupTool/Source/Private/InstanceDataObjectFixupToolTedsQueries.h
+++ b/Engine/Plugins/Experimental/InstanceDataObjectFixupTool/Source/Private/InstanceDataObjectFixupToolTedsQueries.h
@@ -9,14 +9,14 @@
 #include "InstanceDataObjectFixupToolTedsQueries.generated.h"
 
 UCLASS()
-class UInstanceDataObjectFixupToolTedsQueryFactory : public UTypedElementDataStorageFactory
+class UInstanceDataObjectFixupToolTedsQueryFactory : public UEditorDataStorageFactory
 {
 	GENERATED_BODY()
 
 public:
 	~UInstanceDataObjectFixupToolTedsQueryFactory() override = default;
 
-	void RegisterQueries(ITypedElementDataStorageInterface& DataStorage) override;
+	void RegisterQueries(IEditorDataStorageProvider& DataStorage) override;
 
 private:
 	static void ShowFixUpToolForPlaceholders(UE::Editor::DataStorage::RowHandle Row);
diff --git a/Engine/Source/Editor/UnrealEd/Private/Elements/Actor/ActorElementEditorSelectionInterface.cpp b/Engine/Source/Editor/UnrealEd/Private/Elements/Actor/ActorElementEditorSelectionInterface.cpp
--- a/Engine/Source/Editor/UnrealEd/Private/Elements/Actor/ActorElementEditorSelectionInterface.cpp
+++ b/Engine/Source/Editor/UnrealEd/Private/Elements/Actor/ActorElementEditorSelectionInterface.cpp
@@ -53,12 +53,12 @@ bool UActorElementEditorSelectionInterface::SelectElement(const FTypedElementHan
 	const AActor* Actor = ActorElementDataUtil::GetActorFromHandle(InElementHandle);
 	UTypedElementRegistry* Registry = InSelectionSet->GetRegistry();
 	// Add a selection column in TEDS
-	if (const ITypedElementDataStorageCompatibilityInterface* Compatibility = Registry->GetDataStorageCompatibility())
+	if (const IEditorDataStorageCompatibilityProvider* Compatibility = Registry->GetDataStorageCompatibility())
 	{
 		RowHandle Row = Compatibility->FindRowWithCompatibleObject(Actor);
 		if (Row != InvalidRowHandle)
 		{
-			if (ITypedElementDataStorageInterface* DataStorage = UTypedElementRegistry::GetInstance()->GetMutableDataStorage())
+			if (IEditorDataStorageProvider* DataStorage = UTypedElementRegistry::GetInstance()->GetMutableDataStorage())
 			{
 				DataStorage->AddColumn<FTypedElementSelectionColumn>(Row);
 
@@ -81,12 +81,12 @@ bool UActorElementEditorSelectionInterface::DeselectElement(const FTypedElementH
 
 	const AActor* Actor = ActorElementDataUtil::GetActorFromHandle(InElementHandle);
 	UTypedElementRegistry* Registry = InSelectionSet->GetRegistry();
-	if (const ITypedElementDataStorageCompatibilityInterface* Compatibility = Registry->GetDataStorageCompatibility())
+	if (const IEditorDataStorageCompatibilityProvider* Compatibility = Registry->GetDataStorageCompatibility())
 	{
 		RowHandle Row = Compatibility->FindRowWithCompatibleObject(Actor);
 		if (Row != InvalidRowHandle)
 		{
-			if (ITypedElementDataStorageInterface* DataStorage = UTypedElementRegistry::GetInstance()->GetMutableDataStorage())
+			if (IEditorDataStorageProvider* DataStorage = UTypedElementRegistry::GetInstance()->GetMutableDataStorage())
             {
             	DataStorage->RemoveColumn(Row, FTypedElementSelectionColumn::StaticStruct());
             }
diff --git a/Engine/Source/Editor/UnrealEd/Private/LevelEditorViewport.cpp b/Engine/Source/Editor/UnrealEd/Private/LevelEditorViewport.cpp
--- a/Engine/Source/Editor/UnrealEd/Private/LevelEditorViewport.cpp
+++ b/Engine/Source/Editor/UnrealEd/Private/LevelEditorViewport.cpp
@@ -1671,7 +1671,7 @@ void FLevelEditorViewportClient::DestroyDropPreviewElements()
 	// and although it does still happen in the immediately triggered garbage cleanup, that feels
 	// potentially brittle.
 	UTypedElementRegistry* TypedElementRegistry = UTypedElementRegistry::GetInstance();
-	if (ITypedElementDataStorageCompatibilityInterface* TedsCompat = TypedElementRegistry->GetMutableDataStorageCompatibility())
+	if (IEditorDataStorageCompatibilityProvider* TedsCompat = TypedElementRegistry->GetMutableDataStorageCompatibility())
 	{
 		DropPreviewElements->ForEachElement<ITypedElementObjectInterface>(
 			[this, TedsCompat](const TTypedElement<ITypedElementObjectInterface>& InElement)
diff --git a/Engine/Source/Runtime/Engine/Private/Elements/Framework/EngineElementsLibrary.cpp b/Engine/Source/Runtime/Engine/Private/Elements/Framework/EngineElementsLibrary.cpp
--- a/Engine/Source/Runtime/Engine/Private/Elements/Framework/EngineElementsLibrary.cpp
+++ b/Engine/Source/Runtime/Engine/Private/Elements/Framework/EngineElementsLibrary.cpp
@@ -380,7 +380,7 @@ void UEngineElementsLibrary::DestroyActorElement(const AActor* InActor, TTypedEl
 void UEngineElementsLibrary::RegisterActorElement(const AActor* InActor)
 {
 	UTypedElementRegistry* Registry = UTypedElementRegistry::GetInstance();
-	ITypedElementDataStorageCompatibilityInterface* Storage = Registry->GetMutableDataStorageCompatibility();
+	IEditorDataStorageCompatibilityProvider* Storage = Registry->GetMutableDataStorageCompatibility();
 	if (Storage)
 	{
 		Storage->AddCompatibleObject(const_cast<AActor*>(InActor));
@@ -390,7 +390,7 @@ void UEngineElementsLibrary::RegisterActorElement(const AActor* InActor)
 void UEngineElementsLibrary::UnregisterActorElement(const AActor* InActor)
 {
 	UTypedElementRegistry* Registry = UTypedElementRegistry::GetInstance();
-	if (ITypedElementDataStorageCompatibilityInterface* Storage = Registry->GetMutableDataStorageCompatibility())
+	if (IEditorDataStorageCompatibilityProvider* Storage = Registry->GetMutableDataStorageCompatibility())
 	{
 		Storage->RemoveCompatibleObject(const_cast<AActor*>(InActor));
 	}
diff --git a/Engine/Source/Runtime/TypedElementFramework/Private/Elements/Framework/TypedElementAttributeBinding.cpp b/Engine/Source/Runtime/TypedElementFramework/Private/Elements/Framework/TypedElementAttributeBinding.cpp
--- a/Engine/Source/Runtime/TypedElementFramework/Private/Elements/Framework/TypedElementAttributeBinding.cpp
+++ b/Engine/Source/Runtime/TypedElementFramework/Private/Elements/Framework/TypedElementAttributeBinding.cpp
@@ -17,7 +17,7 @@ namespace UE::Editor::DataStorage
 	}
 
 	FAttributeBinder::FAttributeBinder(
-		RowHandle InTargetRow, ITypedElementDataStorageInterface* InDataStorage)
+		RowHandle InTargetRow, IEditorDataStorageProvider* InDataStorage)
 		: TargetRow(InTargetRow)
 		, DataStorage(InDataStorage)
 	{
diff --git a/Engine/Source/Runtime/TypedElementFramework/Private/Elements/Framework/TypedElementDataStorageWidget.cpp b/Engine/Source/Runtime/TypedElementFramework/Private/Elements/Framework/TypedElementDataStorageWidget.cpp
--- a/Engine/Source/Runtime/TypedElementFramework/Private/Elements/Framework/TypedElementDataStorageWidget.cpp
+++ b/Engine/Source/Runtime/TypedElementFramework/Private/Elements/Framework/TypedElementDataStorageWidget.cpp
@@ -29,7 +29,7 @@ void STedsWidget::Construct(const FArguments& InArgs)
 void STedsWidget::RegisterTedsWidget(const TSharedPtr<SWidget>& InContentWidget)
 {
 	using namespace UE::Editor::DataStorage;
-	ITypedElementDataStorageInterface* Storage = GetStorageIfAvailable();
+	IEditorDataStorageProvider* Storage = GetStorageIfAvailable();
 
 	// If TEDS is not enabled, STedsWidget will just behave like a regular widget
 	if(!Storage)
@@ -54,7 +54,7 @@ void STedsWidget::RegisterTedsWidget(const TSharedPtr<SWidget>& InContentWidget)
 
 void STedsWidget::SetContent(const TSharedRef< SWidget >& InContent)
 {
-	if(ITypedElementDataStorageInterface* Storage = GetStorageIfAvailable())
+	if(IEditorDataStorageProvider* Storage = GetStorageIfAvailable())
 	{
 		if(FTypedElementSlateWidgetReferenceColumn* WidgetReferenceColumn = Storage->GetColumn<FTypedElementSlateWidgetReferenceColumn>(UiRowHandle))
 		{
@@ -73,7 +73,7 @@ UE::Editor::DataStorage::RowHandle STedsWidget::GetRowHandle() const
 	return UiRowHandle;
 }
 
-ITypedElementDataStorageInterface* STedsWidget::GetStorageIfAvailable()
+IEditorDataStorageProvider* STedsWidget::GetStorageIfAvailable()
 {
 	UTypedElementRegistry* Registry = UTypedElementRegistry::GetInstance();
 
diff --git a/Engine/Source/Runtime/TypedElementFramework/Private/Elements/Framework/TypedElementRegistry.cpp b/Engine/Source/Runtime/TypedElementFramework/Private/Elements/Framework/TypedElementRegistry.cpp
--- a/Engine/Source/Runtime/TypedElementFramework/Private/Elements/Framework/TypedElementRegistry.cpp
+++ b/Engine/Source/Runtime/TypedElementFramework/Private/Elements/Framework/TypedElementRegistry.cpp
@@ -64,49 +64,49 @@ UTypedElementRegistry* UTypedElementRegistry::GetInstance()
 	return Instance.Get();
 }
 
-ITypedElementDataStorageInterface* UTypedElementRegistry::GetMutableDataStorage()
+IEditorDataStorageProvider* UTypedElementRegistry::GetMutableDataStorage()
 {
 	return DataStorage;
 }
 
-const ITypedElementDataStorageInterface* UTypedElementRegistry::GetDataStorage() const
+const IEditorDataStorageProvider* UTypedElementRegistry::GetDataStorage() const
 {
 	return DataStorage;
 }
 
-void UTypedElementRegistry::SetDataStorage(ITypedElementDataStorageInterface* Storage)
+void UTypedElementRegistry::SetDataStorage(IEditorDataStorageProvider* Storage)
 {
 	DataStorage = Storage;
 	CallDataStorageInterfacesSetDelegateIfNeeded();
 }
 
-ITypedElementDataStorageCompatibilityInterface* UTypedElementRegistry::GetMutableDataStorageCompatibility()
+IEditorDataStorageCompatibilityProvider* UTypedElementRegistry::GetMutableDataStorageCompatibility()
 {
 	return DataStorageCompatibility;
 }
 
-const ITypedElementDataStorageCompatibilityInterface* UTypedElementRegistry::GetDataStorageCompatibility() const
+const IEditorDataStorageCompatibilityProvider* UTypedElementRegistry::GetDataStorageCompatibility() const
 {
 	return DataStorageCompatibility;
 }
 
-void UTypedElementRegistry::SetDataStorageCompatibility(ITypedElementDataStorageCompatibilityInterface* Storage)
+void UTypedElementRegistry::SetDataStorageCompatibility(IEditorDataStorageCompatibilityProvider* Storage)
 {
 	DataStorageCompatibility = Storage;
 	CallDataStorageInterfacesSetDelegateIfNeeded();
 }
 
-ITypedElementDataStorageUiInterface* UTypedElementRegistry::GetMutableDataStorageUi()
+IEditorDataStorageUiProvider* UTypedElementRegistry::GetMutableDataStorageUi()
 {
 	return DataStorageUi;
 }
 
-const ITypedElementDataStorageUiInterface* UTypedElementRegistry::GetDataStorageUi() const
+const IEditorDataStorageUiProvider* UTypedElementRegistry::GetDataStorageUi() const
 {
 	return DataStorageUi;
 }
 
-void UTypedElementRegistry::SetDataStorageUi(ITypedElementDataStorageUiInterface* Storage)
+void UTypedElementRegistry::SetDataStorageUi(IEditorDataStorageUiProvider* Storage)
 {
 	DataStorageUi = Storage;
 	CallDataStorageInterfacesSetDelegateIfNeeded();
diff --git a/Engine/Source/Runtime/TypedElementFramework/Private/Elements/Interfaces/TypedElementDataStorageUiInterface.cpp b/Engine/Source/Runtime/TypedElementFramework/Private/Elements/Interfaces/TypedElementDataStorageUiInterface.cpp
--- a/Engine/Source/Runtime/TypedElementFramework/Private/Elements/Interfaces/TypedElementDataStorageUiInterface.cpp
+++ b/Engine/Source/Runtime/TypedElementFramework/Private/Elements/Interfaces/TypedElementDataStorageUiInterface.cpp
@@ -42,7 +42,7 @@ TConstArrayView<const UScriptStruct*> FTypedElementWidgetConstructor::GetAdditio
 }
 
 FString FTypedElementWidgetConstructor::CreateWidgetDisplayName(
-	ITypedElementDataStorageInterface* DataStorage, RowHandle Row) const
+	IEditorDataStorageProvider* DataStorage, RowHandle Row) const
 {
 	switch (MatchedColumnTypes.Num())
 	{
@@ -95,8 +95,8 @@ FString FTypedElementWidgetConstructor::CreateWidgetDisplayName(
 
 TSharedPtr<SWidget> FTypedElementWidgetConstructor::ConstructFinalWidget(
 	RowHandle Row,
-	ITypedElementDataStorageInterface* DataStorage,
-	ITypedElementDataStorageUiInterface* DataStorageUi,
+	IEditorDataStorageProvider* DataStorage,
+	IEditorDataStorageUiProvider* DataStorageUi,
 	const UE::Editor::DataStorage::FMetaDataView& Arguments)
 {
 	// Add the additional columns to the UI row
@@ -137,8 +137,8 @@ TSharedPtr<SWidget> FTypedElementWidgetConstructor::ConstructFinalWidget(
 
 TSharedPtr<SWidget> FTypedElementWidgetConstructor::Construct(
 	RowHandle Row,
-	ITypedElementDataStorageInterface* DataStorage,
-	ITypedElementDataStorageUiInterface* DataStorageUi,
+	IEditorDataStorageProvider* DataStorage,
+	IEditorDataStorageUiProvider* DataStorageUi,
 	const UE::Editor::DataStorage::FMetaDataView& Arguments)
 {
 	UE::Editor::DataStorage::RowHandle TargetRow = GetTargetRow(DataStorage, Row);
@@ -165,8 +165,8 @@ TSharedPtr<SWidget> FTypedElementWidgetConstructor::CreateWidget(const UE::Edito
 }
 
 TSharedPtr<SWidget> FTypedElementWidgetConstructor::CreateWidget(
-	ITypedElementDataStorageInterface* DataStorage,
-	ITypedElementDataStorageUiInterface* DataStorageUi,
+	IEditorDataStorageProvider* DataStorage,
+	IEditorDataStorageUiProvider* DataStorageUi,
 	UE::Editor::DataStorage::RowHandle TargetRow,
 	UE::Editor::DataStorage::RowHandle UiRow,
 	const UE::Editor::DataStorage::FMetaDataView& Arguments)
@@ -174,7 +174,7 @@ TSharedPtr<SWidget> FTypedElementWidgetConstructor::CreateWidget(
 	return CreateWidget(Arguments);
 }
 
-bool FTypedElementWidgetConstructor::SetColumns(ITypedElementDataStorageInterface* DataStorage, RowHandle Row)
+bool FTypedElementWidgetConstructor::SetColumns(IEditorDataStorageProvider* DataStorage, RowHandle Row)
 {
 	return true;
 }
@@ -197,15 +197,15 @@ FString FTypedElementWidgetConstructor::DescribeColumnType(const UScriptStruct*
 }
 
 bool FTypedElementWidgetConstructor::FinalizeWidget(
-	ITypedElementDataStorageInterface* DataStorage,
-	ITypedElementDataStorageUiInterface* DataStorageUi,
+	IEditorDataStorageProvider* DataStorage,
+	IEditorDataStorageUiProvider* DataStorageUi,
 	RowHandle Row,
 	const TSharedPtr<SWidget>& Widget)
 {
 	return true;
 }
 
-void FTypedElementWidgetConstructor::AddDefaultWidgetColumns(RowHandle Row, ITypedElementDataStorageInterface* DataStorage) const
+void FTypedElementWidgetConstructor::AddDefaultWidgetColumns(RowHandle Row, IEditorDataStorageProvider* DataStorage) const
 {
 	const FString WidgetLabel(CreateWidgetDisplayName(DataStorage, Row));
 	DataStorage->AddColumn(Row, FTypedElementLabelColumn{.Label = WidgetLabel} );
@@ -224,7 +224,7 @@ void FTypedElementWidgetConstructor::AddDefaultWidgetColumns(RowHandle Row, ITyp
 	}
 }
 
-FTypedElementWidgetConstructor::RowHandle FTypedElementWidgetConstructor::GetTargetRow(ITypedElementDataStorageInterface* DataStorage, RowHandle WidgetRow) const
+FTypedElementWidgetConstructor::RowHandle FTypedElementWidgetConstructor::GetTargetRow(IEditorDataStorageProvider* DataStorage, RowHandle WidgetRow) const
 {
 	using namespace UE::Editor::DataStorage;
 	RowHandle TargetRow = InvalidRowHandle;
@@ -245,14 +245,14 @@ FSimpleWidgetConstructor::FSimpleWidgetConstructor(const UScriptStruct* InTypeIn
 {
 }
 
-TSharedPtr<SWidget> FSimpleWidgetConstructor::CreateWidget(ITypedElementDataStorageInterface* DataStorage,
-	ITypedElementDataStorageUiInterface* DataStorageUi, UE::Editor::DataStorage::RowHandle TargetRow, UE::Editor::DataStorage::RowHandle WidgetRow,
+TSharedPtr<SWidget> FSimpleWidgetConstructor::CreateWidget(IEditorDataStorageProvider* DataStorage,
+	IEditorDataStorageUiProvider* DataStorageUi, UE::Editor::DataStorage::RowHandle TargetRow, UE::Editor::DataStorage::RowHandle WidgetRow,
 	const UE::Editor::DataStorage::FMetaDataView& Arguments)
 {
 	return nullptr;
 }
 
-bool FSimpleWidgetConstructor::SetColumns(ITypedElementDataStorageInterface* DataStorage, RowHandle Row)
+bool FSimpleWidgetConstructor::SetColumns(IEditorDataStorageProvider* DataStorage, RowHandle Row)
 {
 	return FTypedElementWidgetConstructor::SetColumns(DataStorage, Row);
 }
@@ -263,15 +263,15 @@ TSharedPtr<SWidget> FSimpleWidgetConstructor::CreateWidget(const UE::Editor::Dat
 	return nullptr;
 }
 
-bool FSimpleWidgetConstructor::FinalizeWidget(ITypedElementDataStorageInterface* DataStorage, ITypedElementDataStorageUiInterface* DataStorageUi,
+bool FSimpleWidgetConstructor::FinalizeWidget(IEditorDataStorageProvider* DataStorage, IEditorDataStorageUiProvider* DataStorageUi,
 	RowHandle Row, const TSharedPtr<SWidget>& Widget)
 {
 	// This function is not needed anymore and only exists so derived classes cannot derive from it anymore
 	return true;
 }
 
-TSharedPtr<SWidget> FSimpleWidgetConstructor::Construct(RowHandle WidgetRow, ITypedElementDataStorageInterface* DataStorage,
-	ITypedElementDataStorageUiInterface* DataStorageUi, const UE::Editor::DataStorage::FMetaDataView& Arguments)
+TSharedPtr<SWidget> FSimpleWidgetConstructor::Construct(RowHandle WidgetRow, IEditorDataStorageProvider* DataStorage,
+	IEditorDataStorageUiProvider* DataStorageUi, const UE::Editor::DataStorage::FMetaDataView& Arguments)
 {
 	const UE::Editor::DataStorage::RowHandle TargetRow = GetTargetRow(DataStorage, WidgetRow);
 
diff --git a/Engine/Source/Runtime/TypedElementFramework/Public/Elements/Framework/TypedElementAttributeBinding.h b/Engine/Source/Runtime/TypedElementFramework/Public/Elements/Framework/TypedElementAttributeBinding.h
--- a/Engine/Source/Runtime/TypedElementFramework/Public/Elements/Framework/TypedElementAttributeBinding.h
+++ b/Engine/Source/Runtime/TypedElementFramework/Public/Elements/Framework/TypedElementAttributeBinding.h
@@ -6,7 +6,7 @@
 #include "Elements/Interfaces/TypedElementDataStorageInterface.h"
 #include "Misc/Attribute.h"
 
-class ITypedElementDataStorageInterface;
+class IEditorDataStorageProvider;
 
 namespace UE::Editor::DataStorage
 {
@@ -31,7 +31,7 @@ namespace UE::Editor::DataStorage
 		TYPEDELEMENTFRAMEWORK_API FAttributeBinder(RowHandle InTargetRow);
 		
 		/* Create an attribute binder for a given row. */
-		TYPEDELEMENTFRAMEWORK_API FAttributeBinder(RowHandle InTargetRow, ITypedElementDataStorageInterface* InDataStorage);
+		TYPEDELEMENTFRAMEWORK_API FAttributeBinder(RowHandle InTargetRow, IEditorDataStorageProvider* InDataStorage);
 
 		/**
 		 * Bind a specific data member inside a TEDS column to an attribute of the same type as the data
@@ -110,7 +110,7 @@ namespace UE::Editor::DataStorage
 		RowHandle TargetRow;
 
 		// A ptr to the data storage for quick access
-		ITypedElementDataStorageInterface* DataStorage;
+		IEditorDataStorageProvider* DataStorage;
 	};
 } // namespace UE::Editor::DataStorage
 
diff --git a/Engine/Source/Runtime/TypedElementFramework/Public/Elements/Framework/TypedElementDataStorageWidget.h b/Engine/Source/Runtime/TypedElementFramework/Public/Elements/Framework/TypedElementDataStorageWidget.h
--- a/Engine/Source/Runtime/TypedElementFramework/Public/Elements/Framework/TypedElementDataStorageWidget.h
+++ b/Engine/Source/Runtime/TypedElementFramework/Public/Elements/Framework/TypedElementDataStorageWidget.h
@@ -45,7 +45,7 @@ private:
 
 	void RegisterTedsWidget(const TSharedPtr<SWidget>& InContentWidget);
 	
-	static ITypedElementDataStorageInterface* GetStorageIfAvailable();
+	static IEditorDataStorageProvider* GetStorageIfAvailable();
 
 private:
 	
diff --git a/Engine/Source/Runtime/TypedElementFramework/Public/Elements/Framework/TypedElementRegistry.h b/Engine/Source/Runtime/TypedElementFramework/Public/Elements/Framework/TypedElementRegistry.h
--- a/Engine/Source/Runtime/TypedElementFramework/Public/Elements/Framework/TypedElementRegistry.h
+++ b/Engine/Source/Runtime/TypedElementFramework/Public/Elements/Framework/TypedElementRegistry.h
@@ -77,17 +77,17 @@ public:
 	UFUNCTION(BlueprintPure, DisplayName="Get Default Typed Element Registry", Category = "TypedElementFramework|Registry", meta=(ScriptName="GetDefaultTypedElementRegistry"))
 	static TYPEDELEMENTFRAMEWORK_API UTypedElementRegistry* GetInstance();
 
-	TYPEDELEMENTFRAMEWORK_API ITypedElementDataStorageInterface* GetMutableDataStorage();
-	TYPEDELEMENTFRAMEWORK_API const ITypedElementDataStorageInterface* GetDataStorage() const;
-	TYPEDELEMENTFRAMEWORK_API void SetDataStorage(ITypedElementDataStorageInterface* Storage);
+	TYPEDELEMENTFRAMEWORK_API IEditorDataStorageProvider* GetMutableDataStorage();
+	TYPEDELEMENTFRAMEWORK_API const IEditorDataStorageProvider* GetDataStorage() const;
+	TYPEDELEMENTFRAMEWORK_API void SetDataStorage(IEditorDataStorageProvider* Storage);
 
-	TYPEDELEMENTFRAMEWORK_API ITypedElementDataStorageCompatibilityInterface* GetMutableDataStorageCompatibility();
-	TYPEDELEMENTFRAMEWORK_API const ITypedElementDataStorageCompatibilityInterface* GetDataStorageCompatibility() const;
-	TYPEDELEMENTFRAMEWORK_API void SetDataStorageCompatibility(ITypedElementDataStorageCompatibilityInterface* Storage);
+	TYPEDELEMENTFRAMEWORK_API IEditorDataStorageCompatibilityProvider* GetMutableDataStorageCompatibility();
+	TYPEDELEMENTFRAMEWORK_API const IEditorDataStorageCompatibilityProvider* GetDataStorageCompatibility() const;
+	TYPEDELEMENTFRAMEWORK_API void SetDataStorageCompatibility(IEditorDataStorageCompatibilityProvider* Storage);
 
-	TYPEDELEMENTFRAMEWORK_API ITypedElementDataStorageUiInterface* GetMutableDataStorageUi();
-	TYPEDELEMENTFRAMEWORK_API const ITypedElementDataStorageUiInterface* GetDataStorageUi() const;
-	TYPEDELEMENTFRAMEWORK_API void SetDataStorageUi(ITypedElementDataStorageUiInterface* Storage);
+	TYPEDELEMENTFRAMEWORK_API IEditorDataStorageUiProvider* GetMutableDataStorageUi();
+	TYPEDELEMENTFRAMEWORK_API const IEditorDataStorageUiProvider* GetDataStorageUi() const;
+	TYPEDELEMENTFRAMEWORK_API void SetDataStorageUi(IEditorDataStorageUiProvider* Storage);
 
 	TYPEDELEMENTFRAMEWORK_API bool AreDataStorageInterfacesSet() const;
 
@@ -701,7 +701,7 @@ private:
 	FSimpleMulticastDelegate OnProcessingDeferredElementsToDestroyDelegate;
 	
 	FOnDataStorageInterfacesSet OnDataStorageInterfacesSetDelegate;
-	ITypedElementDataStorageInterface* DataStorage = nullptr;
-	ITypedElementDataStorageCompatibilityInterface* DataStorageCompatibility = nullptr;
-	ITypedElementDataStorageUiInterface* DataStorageUi = nullptr;
+	IEditorDataStorageProvider* DataStorage = nullptr;
+	IEditorDataStorageCompatibilityProvider* DataStorageCompatibility = nullptr;
+	IEditorDataStorageUiProvider* DataStorageUi = nullptr;
 };
diff --git a/Engine/Source/Runtime/TypedElementFramework/Public/Elements/Interfaces/TypedElementDataStorageCompatibilityInterface.h b/Engine/Source/Runtime/TypedElementFramework/Public/Elements/Interfaces/TypedElementDataStorageCompatibilityInterface.h
--- a/Engine/Source/Runtime/TypedElementFramework/Public/Elements/Interfaces/TypedElementDataStorageCompatibilityInterface.h
+++ b/Engine/Source/Runtime/TypedElementFramework/Public/Elements/Interfaces/TypedElementDataStorageCompatibilityInterface.h
@@ -14,7 +14,7 @@
 #include "TypedElementDataStorageCompatibilityInterface.generated.h"
 
 UINTERFACE(MinimalAPI)
-class UTypedElementDataStorageCompatibilityInterface : public UInterface
+class UEditorDataStorageCompatibilityProvider : public UInterface
 {
 	GENERATED_BODY()
 };
@@ -23,13 +23,13 @@ class UTypedElementDataStorageCompatibilityInterface : public UInterface
  * Interface to provide compatibility with existing systems that don't directly
  * support the data storage.
  */
-class ITypedElementDataStorageCompatibilityInterface
+class IEditorDataStorageCompatibilityProvider
 {
 	GENERATED_BODY()
 
 public:
-	using ObjectRegistrationFilter = TFunction<bool(const ITypedElementDataStorageCompatibilityInterface&, const UObject*)>;
-	using ObjectToRowDealiaser = TFunction<UE::Editor::DataStorage::RowHandle(const ITypedElementDataStorageCompatibilityInterface&, const UObject*)>;
+	using ObjectRegistrationFilter = TFunction<bool(const IEditorDataStorageCompatibilityProvider&, const UObject*)>;
+	using ObjectToRowDealiaser = TFunction<UE::Editor::DataStorage::RowHandle(const IEditorDataStorageCompatibilityProvider&, const UObject*)>;
 	
 	/**
 	 * @section Type-agnostic functions
@@ -118,7 +118,7 @@ template<typename Type> Type* GetRawPointer(Type* Object)						{ return Object;
 template<typename Type> Type* GetRawPointer(Type& Object)						{ return &Object; }
 
 template<typename ObjectType>
-UE::Editor::DataStorage::RowHandle ITypedElementDataStorageCompatibilityInterface::AddCompatibleObject(ObjectType&& Object)
+UE::Editor::DataStorage::RowHandle IEditorDataStorageCompatibilityProvider::AddCompatibleObject(ObjectType&& Object)
 {
 	auto RawPointer = GetRawPointer(Forward<ObjectType>(Object));
 	using BaseType = std::remove_cv_t<std::remove_pointer_t<decltype(RawPointer)>>;
@@ -134,13 +134,13 @@ UE::Editor::DataStorage::RowHandle ITypedElementDataStorageCompatibilityInterfac
 }
 
 template<typename ObjectType>
-void ITypedElementDataStorageCompatibilityInterface::RemoveCompatibleObject(ObjectType&& Object)
+void IEditorDataStorageCompatibilityProvider::RemoveCompatibleObject(ObjectType&& Object)
 {
 	RemoveCompatibleObjectExplicit(GetRawPointer(Forward<ObjectType>(Object)));
 }
 
 template<typename ObjectType>
-UE::Editor::DataStorage::RowHandle ITypedElementDataStorageCompatibilityInterface::FindRowWithCompatibleObject(ObjectType&& Object) const
+UE::Editor::DataStorage::RowHandle IEditorDataStorageCompatibilityProvider::FindRowWithCompatibleObject(ObjectType&& Object) const
 {
 	return FindRowWithCompatibleObjectExplicit(GetRawPointer(Forward<ObjectType>(Object)));
 }
diff --git a/Engine/Source/Runtime/TypedElementFramework/Public/Elements/Interfaces/TypedElementDataStorageFactory.h b/Engine/Source/Runtime/TypedElementFramework/Public/Elements/Interfaces/TypedElementDataStorageFactory.h
--- a/Engine/Source/Runtime/TypedElementFramework/Public/Elements/Interfaces/TypedElementDataStorageFactory.h
+++ b/Engine/Source/Runtime/TypedElementFramework/Public/Elements/Interfaces/TypedElementDataStorageFactory.h
@@ -4,21 +4,21 @@
 
 #include "TypedElementDataStorageFactory.generated.h"
 
-class ITypedElementDataStorageCompatibilityInterface;
-class ITypedElementDataStorageInterface;
-class ITypedElementDataStorageUiInterface;
+class IEditorDataStorageCompatibilityProvider;
+class IEditorDataStorageProvider;
+class IEditorDataStorageUiProvider;
 
 /**
  * Base class that can be used to register various elements, such as queries and widgets, with
- * the Typed Elements Data Storage.
+ * the Editor Data Storage.
  */
 UCLASS(MinimalAPI)
-class UTypedElementDataStorageFactory : public UObject
+class UEditorDataStorageFactory : public UObject
 {
 	GENERATED_BODY()
 
 public:
-	~UTypedElementDataStorageFactory() override = default;
+	~UEditorDataStorageFactory() override = default;
 
 	/** 
 	 * Returns the order registration will be executed. Factories with a lower number will be executed
@@ -29,22 +29,22 @@ public:
 	/**
 	 * All factories will have this called before any Register functions on any factories are called
 	 */
-	virtual void PreRegister(ITypedElementDataStorageInterface& DataStorage) {}
+	virtual void PreRegister(IEditorDataStorageProvider& DataStorage) {}
 
-	virtual void RegisterTables(ITypedElementDataStorageInterface& DataStorage) {}
-	virtual void RegisterTables(ITypedElementDataStorageInterface& DataStorage, ITypedElementDataStorageCompatibilityInterface& DataStorageCompatibility) {}
-	virtual void RegisterTickGroups(ITypedElementDataStorageInterface& DataStorage) const {}
-	virtual void RegisterQueries(ITypedElementDataStorageInterface& DataStorage) {}
+	virtual void RegisterTables(IEditorDataStorageProvider& DataStorage) {}
+	virtual void RegisterTables(IEditorDataStorageProvider& DataStorage, IEditorDataStorageCompatibilityProvider& DataStorageCompatibility) {}
+	virtual void RegisterTickGroups(IEditorDataStorageProvider& DataStorage) const {}
+	virtual void RegisterQueries(IEditorDataStorageProvider& DataStorage) {}
 
-	virtual void RegisterRegistrationFilters(ITypedElementDataStorageCompatibilityInterface& DataStorageCompatibility) const {}
-	virtual void RegisterDealiaser(ITypedElementDataStorageCompatibilityInterface& DataStorageCompatibility) const {}
+	virtual void RegisterRegistrationFilters(IEditorDataStorageCompatibilityProvider& DataStorageCompatibility) const {}
+	virtual void RegisterDealiaser(IEditorDataStorageCompatibilityProvider& DataStorageCompatibility) const {}
 	
-	virtual void RegisterWidgetPurposes(ITypedElementDataStorageUiInterface& DataStorageUi) const {}
-	virtual void RegisterWidgetConstructors(ITypedElementDataStorageInterface& DataStorage,
-		ITypedElementDataStorageUiInterface& DataStorageUi) const {}
+	virtual void RegisterWidgetPurposes(IEditorDataStorageUiProvider& DataStorageUi) const {}
+	virtual void RegisterWidgetConstructors(IEditorDataStorageProvider& DataStorage,
+		IEditorDataStorageUiProvider& DataStorageUi) const {}
 
 	/**
 	 * Called in reverse order before the DataStorage object is shut down
 	 */
-	virtual void PreShutdown(ITypedElementDataStorageInterface& DataStorage) {}
+	virtual void PreShutdown(IEditorDataStorageProvider& DataStorage) {}
 };
diff --git a/Engine/Source/Runtime/TypedElementFramework/Public/Elements/Interfaces/TypedElementDataStorageInterface.h b/Engine/Source/Runtime/TypedElementFramework/Public/Elements/Interfaces/TypedElementDataStorageInterface.h
--- a/Engine/Source/Runtime/TypedElementFramework/Public/Elements/Interfaces/TypedElementDataStorageInterface.h
+++ b/Engine/Source/Runtime/TypedElementFramework/Public/Elements/Interfaces/TypedElementDataStorageInterface.h
@@ -25,14 +25,14 @@
 class UClass;
 class USubsystem;
 class UScriptStruct;
-class UTypedElementDataStorageFactory;
+class UEditorDataStorageFactory;
 
 using FTypedElementOnDataStorageCreation = FSimpleMulticastDelegate;
 using FTypedElementOnDataStorageDestruction = FSimpleMulticastDelegate;
 using FTypedElementOnDataStorageUpdate = FSimpleMulticastDelegate;
 
 UINTERFACE(MinimalAPI)
-class UTypedElementDataStorageInterface : public UInterface
+class UEditorDataStorageProvider : public UInterface
 {
 	GENERATED_BODY()
 };
@@ -51,7 +51,7 @@ struct TTypedElementColumnTypeList
 	operator TConstArrayView<const UScriptStruct*>() const { return ColumnTypes; }
 };
 
-class ITypedElementDataStorageInterface
+class IEditorDataStorageProvider
 {
 	GENERATED_BODY()
 
@@ -73,7 +73,7 @@ public:
 	 */
 
 	/** Finds a factory instance registered with TEDS */
-	virtual const UTypedElementDataStorageFactory* FindFactory(const UClass* FactoryType) const = 0;
+	virtual const UEditorDataStorageFactory* FindFactory(const UClass* FactoryType) const = 0;
 
 	/** Convenience function for FindFactory */
 	template<typename FactoryT>
@@ -446,57 +446,57 @@ public:
 // Implementations
 
 template <typename FactoryT>
-const FactoryT* ITypedElementDataStorageInterface::FindFactory() const
+const FactoryT* IEditorDataStorageProvider::FindFactory() const
 {
 	return static_cast<const FactoryT*>(FindFactory(FactoryT::StaticClass()));
 }
 
 template<UE::Editor::DataStorage::TColumnType... Columns>
-UE::Editor::DataStorage::TableHandle ITypedElementDataStorageInterface::RegisterTable(const FName Name)
+UE::Editor::DataStorage::TableHandle IEditorDataStorageProvider::RegisterTable(const FName Name)
 {
 	return RegisterTable({ Columns::StaticStruct()... }, Name);
 }
 
 template<UE::Editor::DataStorage::TColumnType... Columns>
-UE::Editor::DataStorage::TableHandle ITypedElementDataStorageInterface::RegisterTable(
+UE::Editor::DataStorage::TableHandle IEditorDataStorageProvider::RegisterTable(
 	TableHandle SourceTable, const FName Name)
 {
 	return RegisterTable(SourceTable, { Columns::StaticStruct()... }, Name);
 }
 
 template<UE::Editor::DataStorage::TColumnType Column>
-void ITypedElementDataStorageInterface::AddColumn(RowHandle Row)
+void IEditorDataStorageProvider::AddColumn(RowHandle Row)
 {
 	AddColumn(Row, Column::StaticStruct());
 }
 
 template<UE::Editor::DataStorage::TColumnType Column>
-void ITypedElementDataStorageInterface::RemoveColumn(RowHandle Row)
+void IEditorDataStorageProvider::RemoveColumn(RowHandle Row)
 {
 	RemoveColumn(Row, Column::StaticStruct());
 }
 
 template<UE::Editor::DataStorage::TColumnType... Columns>
-void ITypedElementDataStorageInterface::AddColumns(RowHandle Row)
+void IEditorDataStorageProvider::AddColumns(RowHandle Row)
 {
 	AddColumns(Row, { Columns::StaticStruct()...});
 }
 
 template <>
-inline void ITypedElementDataStorageInterface::AddColumn<UE::Editor::DataStorage::FValueTag>(RowHandle Row, const FName& Tag, const FName& Value)
+inline void IEditorDataStorageProvider::AddColumn<UE::Editor::DataStorage::FValueTag>(RowHandle Row, const FName& Tag, const FName& Value)
 {
 	AddColumn(Row, UE::Editor::DataStorage::FValueTag(Tag), Value);
 }
 
 template <>
-inline void ITypedElementDataStorageInterface::RemoveColumn<UE::Editor::DataStorage::FValueTag>(RowHandle Row, const FName& Tag)
+inline void IEditorDataStorageProvider::RemoveColumn<UE::Editor::DataStorage::FValueTag>(RowHandle Row, const FName& Tag)
 {
 	using namespace UE::Editor::DataStorage;
 	RemoveColumn(Row, FValueTag(Tag));
 }
 
 template<UE::Editor::DataStorage::TColumnType DynamicColumnTemplate>
-void ITypedElementDataStorageInterface::RemoveColumn(RowHandle Row, const FName& Identifier)
+void IEditorDataStorageProvider::RemoveColumn(RowHandle Row, const FName& Identifier)
 {
 	const UE::Editor::DataStorage::FDynamicColumnDescription Description
 	{
@@ -508,7 +508,7 @@ void ITypedElementDataStorageInterface::RemoveColumn(RowHandle Row, const FName&
 }
 
 template<UE::Editor::DataStorage::TEnumType EnumT>
-void ITypedElementDataStorageInterface::AddColumn(RowHandle Row, EnumT Value)
+void IEditorDataStorageProvider::AddColumn(RowHandle Row, EnumT Value)
 {
 	const UEnum* Enum = StaticEnum<EnumT>();
 	const FName ValueAsFName = *Enum->GetNameStringByValue(static_cast<int64>(Value));
@@ -519,13 +519,13 @@ void ITypedElementDataStorageInterface::AddColumn(RowHandle Row, EnumT Value)
 }
 
 template<auto Value, UE::Editor::DataStorage::TEnumType EnumT>
-void ITypedElementDataStorageInterface::AddColumn(RowHandle Row)
+void IEditorDataStorageProvider::AddColumn(RowHandle Row)
 {
 	AddColumn<EnumT>(Row, Value);
 }
 
 template <UE::Editor::DataStorage::TColumnType DynamicColumnTemplate>
-void ITypedElementDataStorageInterface::AddColumn(RowHandle Row, const FName& Identifier)
+void IEditorDataStorageProvider::AddColumn(RowHandle Row, const FName& Identifier)
 {
 	static_assert(UE::Editor::DataStorage::TDataColumnType<DynamicColumnTemplate> || UE::Editor::DataStorage::TTagColumnType<DynamicColumnTemplate>,
 		"DynamicColumnTemplate must be derived from either a Tag or Column");
@@ -539,7 +539,7 @@ void ITypedElementDataStorageInterface::AddColumn(RowHandle Row, const FName& Id
 }
 
 template <UE::Editor::DataStorage::TColumnType DynamicColumnTemplate>
-void ITypedElementDataStorageInterface::AddColumn(RowHandle Row, const FName& Identifier, DynamicColumnTemplate&& TemplateInstance)
+void IEditorDataStorageProvider::AddColumn(RowHandle Row, const FName& Identifier, DynamicColumnTemplate&& TemplateInstance)
 {
 	static_assert(UE::Editor::DataStorage::TDataColumnType<DynamicColumnTemplate> || UE::Editor::DataStorage::TTagColumnType<DynamicColumnTemplate>,
 		"DynamicColumnTemplate must be derived from either a Tag or Column");
@@ -576,20 +576,20 @@ void ITypedElementDataStorageInterface::AddColumn(RowHandle Row, const FName& Id
 }
 
 template<UE::Editor::DataStorage::TEnumType EnumT>
-void ITypedElementDataStorageInterface::RemoveColumn(RowHandle Row)
+void IEditorDataStorageProvider::RemoveColumn(RowHandle Row)
 {
 	const UEnum* Enum = StaticEnum<EnumT>();
 	RemoveColumn(Row, UE::Editor::DataStorage::FValueTag(Enum->GetFName()));
 }
 
 template<UE::Editor::DataStorage::TColumnType... Columns>
-void ITypedElementDataStorageInterface::RemoveColumns(RowHandle Row)
+void IEditorDataStorageProvider::RemoveColumns(RowHandle Row)
 {
 	RemoveColumns(Row, { Columns::StaticStruct()...});
 }
 
 template<UE::Editor::DataStorage::TDataColumnType ColumnType>
-void ITypedElementDataStorageInterface::AddColumn(RowHandle Row, ColumnType&& Column)
+void IEditorDataStorageProvider::AddColumn(RowHandle Row, ColumnType&& Column)
 {
 	AddColumnData(Row, ColumnType::StaticStruct(),
 		[&Column](void* ColumnData, const UScriptStruct&)
@@ -617,19 +617,19 @@ void ITypedElementDataStorageInterface::AddColumn(RowHandle Row, ColumnType&& Co
 }
 
 template<UE::Editor::DataStorage::TDataColumnType ColumnType>
-ColumnType* ITypedElementDataStorageInterface::GetColumn(RowHandle Row)
+ColumnType* IEditorDataStorageProvider::GetColumn(RowHandle Row)
 {
 	return reinterpret_cast<ColumnType*>(GetColumnData(Row, ColumnType::StaticStruct()));
 }
 
 template<UE::Editor::DataStorage::TDataColumnType ColumnType>
-const ColumnType* ITypedElementDataStorageInterface::GetColumn(RowHandle Row) const
+const ColumnType* IEditorDataStorageProvider::GetColumn(RowHandle Row) const
 {
 	return reinterpret_cast<const ColumnType*>(GetColumnData(Row, ColumnType::StaticStruct()));
 }
 
 template <UE::Editor::DataStorage::TDataColumnType DynamicColumnTemplate>
-DynamicColumnTemplate* ITypedElementDataStorageInterface::GetColumn(RowHandle Row, const FName& Identifier)
+DynamicColumnTemplate* IEditorDataStorageProvider::GetColumn(RowHandle Row, const FName& Identifier)
 {
 	const UE::Editor::DataStorage::FDynamicColumnDescription Description
 	{
@@ -645,7 +645,7 @@ DynamicColumnTemplate* ITypedElementDataStorageInterface::GetColumn(RowHandle Ro
 }
 
 template <UE::Editor::DataStorage::TDataColumnType DynamicColumnTemplate>
-const DynamicColumnTemplate* ITypedElementDataStorageInterface::GetColumn(RowHandle Row, const FName& Identifier) const
+const DynamicColumnTemplate* IEditorDataStorageProvider::GetColumn(RowHandle Row, const FName& Identifier) const
 {
 	const UE::Editor::DataStorage::FDynamicColumnDescription Description
 	{
@@ -661,13 +661,13 @@ const DynamicColumnTemplate* ITypedElementDataStorageInterface::GetColumn(RowHan
 }
 
 template<UE::Editor::DataStorage::TColumnType... ColumnType>
-bool ITypedElementDataStorageInterface::HasColumns(RowHandle Row) const
+bool IEditorDataStorageProvider::HasColumns(RowHandle Row) const
 {
 	return HasColumns(Row, TConstArrayView<const UScriptStruct*>({ ColumnType::StaticStruct()... }));
 }
 
 template<typename SystemType>
-SystemType* ITypedElementDataStorageInterface::GetExternalSystem()
+SystemType* IEditorDataStorageProvider::GetExternalSystem()
 {
 	return reinterpret_cast<SystemType*>(GetExternalSystemAddress(SystemType::StaticClass()));
 }
diff --git a/Engine/Source/Runtime/TypedElementFramework/Public/Elements/Interfaces/TypedElementDataStorageUiInterface.h b/Engine/Source/Runtime/TypedElementFramework/Public/Elements/Interfaces/TypedElementDataStorageUiInterface.h
--- a/Engine/Source/Runtime/TypedElementFramework/Public/Elements/Interfaces/TypedElementDataStorageUiInterface.h
+++ b/Engine/Source/Runtime/TypedElementFramework/Public/Elements/Interfaces/TypedElementDataStorageUiInterface.h
@@ -13,7 +13,7 @@
 
 #include "TypedElementDataStorageUiInterface.generated.h"
 
-class ITypedElementDataStorageUiInterface;
+class IEditorDataStorageUiProvider;
 class SWidget;
 
 /**
@@ -55,7 +55,7 @@ public:
 	 * Individual widget constructors can override this function with a name specific to them.
 	 */
 	TYPEDELEMENTFRAMEWORK_API virtual FString CreateWidgetDisplayName(
-		ITypedElementDataStorageInterface* DataStorage, RowHandle Row) const;
+		IEditorDataStorageProvider* DataStorage, RowHandle Row) const;
 	
 	/**
 	 *	Calls Construct() to create the internal widget, and then stores it in a container before returning.
@@ -67,8 +67,8 @@ public:
 	 */
 	TYPEDELEMENTFRAMEWORK_API TSharedPtr<SWidget> ConstructFinalWidget(
 		RowHandle Row, /** The row the widget will be stored in. */
-		ITypedElementDataStorageInterface* DataStorage,
-		ITypedElementDataStorageUiInterface* DataStorageUi,
+		IEditorDataStorageProvider* DataStorage,
+		IEditorDataStorageUiProvider* DataStorageUi,
 		const UE::Editor::DataStorage::FMetaDataView& Arguments);
 
 	/**
@@ -80,8 +80,8 @@ public:
 	 */
 	TYPEDELEMENTFRAMEWORK_API virtual TSharedPtr<SWidget> Construct(
 		RowHandle Row, /** The row the widget will be stored in. */
-		ITypedElementDataStorageInterface* DataStorage,
-		ITypedElementDataStorageUiInterface* DataStorageUi,
+		IEditorDataStorageProvider* DataStorage,
+		IEditorDataStorageUiProvider* DataStorageUi,
 		const UE::Editor::DataStorage::FMetaDataView& Arguments);
 
 protected:
@@ -89,13 +89,13 @@ protected:
 	TYPEDELEMENTFRAMEWORK_API virtual TSharedPtr<SWidget> CreateWidget(const UE::Editor::DataStorage::FMetaDataView& Arguments);
 	/** Create a new instance of the target widget. This is a required function. */
 	TYPEDELEMENTFRAMEWORK_API virtual TSharedPtr<SWidget> CreateWidget(
-		ITypedElementDataStorageInterface* DataStorage,
-		ITypedElementDataStorageUiInterface* DataStorageUi,
+		IEditorDataStorageProvider* DataStorage,
+		IEditorDataStorageUiProvider* DataStorageUi,
 		UE::Editor::DataStorage::RowHandle TargetRow,
 		UE::Editor::DataStorage::RowHandle WidgetRow, 
 		const UE::Editor::DataStorage::FMetaDataView& Arguments);
 	/** Set any values in columns if needed. The columns provided through GetAdditionalColumnsList() will have already been created. */
-	TYPEDELEMENTFRAMEWORK_API virtual bool SetColumns(ITypedElementDataStorageInterface* DataStorage, RowHandle Row);
+	TYPEDELEMENTFRAMEWORK_API virtual bool SetColumns(IEditorDataStorageProvider* DataStorage, RowHandle Row);
 	
 	/** Creates a (friendly) name for the provided column type. */
 	TYPEDELEMENTFRAMEWORK_API virtual FString DescribeColumnType(const UScriptStruct* ColumnType) const;
@@ -105,19 +105,19 @@ protected:
 	 * in columns.
 	 */
 	TYPEDELEMENTFRAMEWORK_API virtual bool FinalizeWidget(
-		ITypedElementDataStorageInterface* DataStorage,
-		ITypedElementDataStorageUiInterface* DataStorageUi,
+		IEditorDataStorageProvider* DataStorage,
+		IEditorDataStorageUiProvider* DataStorageUi,
 		RowHandle Row,
 		const TSharedPtr<SWidget>& Widget);
 
 	/** Add the default misc columns we want a widget row to have. */
-	TYPEDELEMENTFRAMEWORK_API void AddDefaultWidgetColumns(RowHandle Row, ITypedElementDataStorageInterface* DataStorage) const;
+	TYPEDELEMENTFRAMEWORK_API void AddDefaultWidgetColumns(RowHandle Row, IEditorDataStorageProvider* DataStorage) const;
 
 	/**
 	 * Helper function to get the actual target row with the data the widget is operating on (if applicable). Returns InvalidRowHandle if there is no
 	 * target row
 	 */
-	TYPEDELEMENTFRAMEWORK_API RowHandle GetTargetRow(ITypedElementDataStorageInterface* DataStorage, RowHandle WidgetRow) const;
+	TYPEDELEMENTFRAMEWORK_API RowHandle GetTargetRow(IEditorDataStorageProvider* DataStorage, RowHandle WidgetRow) const;
 
 protected:
 
@@ -153,8 +153,8 @@ struct FSimpleWidgetConstructor : public FTypedElementWidgetConstructor
 	 * @return The actual widget instance
 	 */
 	TYPEDELEMENTFRAMEWORK_API virtual TSharedPtr<SWidget> CreateWidget(
-		ITypedElementDataStorageInterface* DataStorage,
-		ITypedElementDataStorageUiInterface* DataStorageUi,
+		IEditorDataStorageProvider* DataStorage,
+		IEditorDataStorageUiProvider* DataStorageUi,
 		UE::Editor::DataStorage::RowHandle TargetRow,
 		UE::Editor::DataStorage::RowHandle WidgetRow,
 		const UE::Editor::DataStorage::FMetaDataView& Arguments) override;
@@ -166,7 +166,7 @@ struct FSimpleWidgetConstructor : public FTypedElementWidgetConstructor
 	 * @param WidgetRow The row that contains information about the widget itself
 	 * @return Whether any columns were added
 	 */
-	TYPEDELEMENTFRAMEWORK_API virtual bool SetColumns(ITypedElementDataStorageInterface* DataStorage, RowHandle WidgetRow) override;
+	TYPEDELEMENTFRAMEWORK_API virtual bool SetColumns(IEditorDataStorageProvider* DataStorage, RowHandle WidgetRow) override;
 
 protected:
 
@@ -177,8 +177,8 @@ protected:
 	 * Old function in the widget creation pipeline that isn't used anymore. All your logic should go in CreateWidget() itself
 	 */
 	TYPEDELEMENTFRAMEWORK_API virtual bool FinalizeWidget(
-		ITypedElementDataStorageInterface* DataStorage,
-		ITypedElementDataStorageUiInterface* DataStorageUi,
+		IEditorDataStorageProvider* DataStorage,
+		IEditorDataStorageUiProvider* DataStorageUi,
 		RowHandle Row,
 		const TSharedPtr<SWidget>& Widget) override final;
 	
@@ -187,8 +187,8 @@ protected:
 	 */
 	TYPEDELEMENTFRAMEWORK_API virtual TSharedPtr<SWidget> Construct(
 		RowHandle WidgetRow, 
-		ITypedElementDataStorageInterface* DataStorage,
-		ITypedElementDataStorageUiInterface* DataStorageUi,
+		IEditorDataStorageProvider* DataStorage,
+		IEditorDataStorageUiProvider* DataStorageUi,
 		const UE::Editor::DataStorage::FMetaDataView& Arguments) override final;
 };
 
@@ -203,12 +203,12 @@ struct TStructOpsTypeTraits<FTypedElementWidgetConstructor> : public TStructOpsT
 };
 
 UINTERFACE(MinimalAPI)
-class UTypedElementDataStorageUiInterface : public UInterface
+class UEditorDataStorageUiProvider : public UInterface
 {
 	GENERATED_BODY()
 };
 
-class ITypedElementDataStorageUiInterface
+class IEditorDataStorageUiProvider
 {
 	GENERATED_BODY()
 
@@ -361,13 +361,13 @@ public:
 //
 
 template<typename ConstructorType>
-bool ITypedElementDataStorageUiInterface::RegisterWidgetFactory(FName Purpose)
+bool IEditorDataStorageUiProvider::RegisterWidgetFactory(FName Purpose)
 {
 	return this->RegisterWidgetFactory(Purpose, ConstructorType::StaticStruct());
 }
 
 template<typename ConstructorType>
-bool ITypedElementDataStorageUiInterface::RegisterWidgetFactory(FName Purpose, UE::Editor::DataStorage::Queries::FConditions Columns)
+bool IEditorDataStorageUiProvider::RegisterWidgetFactory(FName Purpose, UE::Editor::DataStorage::Queries::FConditions Columns)
 {
 	return this->RegisterWidgetFactory(Purpose, ConstructorType::StaticStruct(), Columns);
 }
diff --git a/Engine/Source/Runtime/TypedElementFramework/Tests/Elements/Framework/TypedElementDataStoragePerformanceTestCommands.cpp b/Engine/Source/Runtime/TypedElementFramework/Tests/Elements/Framework/TypedElementDataStoragePerformanceTestCommands.cpp
--- a/Engine/Source/Runtime/TypedElementFramework/Tests/Elements/Framework/TypedElementDataStoragePerformanceTestCommands.cpp
+++ b/Engine/Source/Runtime/TypedElementFramework/Tests/Elements/Framework/TypedElementDataStoragePerformanceTestCommands.cpp
@@ -24,7 +24,7 @@ namespace UE::Editor::DataStorage
 				int32 EntitiesToAdd;
 				LexFromString(EntitiesToAdd, *Args[0]);
 
-				ITypedElementDataStorageInterface* DataStorageInterface = UTypedElementRegistry::GetInstance()->GetMutableDataStorage();
+				IEditorDataStorageProvider* DataStorageInterface = UTypedElementRegistry::GetInstance()->GetMutableDataStorage();
 			
 				DataStorageInterface->BatchAddRow(Private::PerformanceTestCommandTable, EntitiesToAdd, [DataStorageInterface](RowHandle Row)
 					{
@@ -42,14 +42,14 @@ namespace UE::Editor::DataStorage
 		{
 			using namespace UE::Editor::DataStorage::Queries;
 
-			ITypedElementDataStorageInterface* DataStorageInterface = UTypedElementRegistry::GetInstance()->GetMutableDataStorage();
+			IEditorDataStorageProvider* DataStorageInterface = UTypedElementRegistry::GetInstance()->GetMutableDataStorage();
 		
 			QueryHandle Query = DataStorageInterface->RegisterQuery(
 				Select().Where().All<FTest_PingPongPrePhys>().Compile());
 
 			TArray<RowHandle> RowsToDelete;
 
-			DataStorageInterface->RunQuery(Query, CreateDirectQueryCallbackBinding([&RowsToDelete](const ITypedElementDataStorageInterface::IDirectQueryContext& Context, const RowHandle* RowArray)
+			DataStorageInterface->RunQuery(Query, CreateDirectQueryCallbackBinding([&RowsToDelete](const IEditorDataStorageProvider::IDirectQueryContext& Context, const RowHandle* RowArray)
 			{
 				RowsToDelete.Insert(RowArray, Context.GetRowCount(), RowsToDelete.Num());
 			}));
@@ -64,7 +64,7 @@ namespace UE::Editor::DataStorage
 	));
 } // namespace UE::Editor::DataStorage
 
-void UTest_PingPongBetweenPhaseFactory::RegisterTables(ITypedElementDataStorageInterface& DataStorage)
+void UTest_PingPongBetweenPhaseFactory::RegisterTables(IEditorDataStorageProvider& DataStorage)
 {
 	UE::Editor::DataStorage::Private::PerformanceTestCommandTable = DataStorage.RegisterTable({FTest_PingPongPrePhys::StaticStruct()}, TEXT("Test_PingPongPrePhys"));
 }
@@ -74,7 +74,7 @@ void UTest_PingPongBetweenPhaseFactory::RegisterTables(ITypedElementDataStorageI
 // A is sensitive to tag PingPongPrePhys.  It adds a PingPongDurPhys tag and removes the PingPonPrePhys tag.  This causes the processed row
 // to then be processed by processor B.  Processor B does a similar thing to ensure it is processed by C.  C then ensures it is processed by A
 // the next time A is run.
-void UTest_PingPongBetweenPhaseFactory::RegisterQueries(ITypedElementDataStorageInterface& DataStorage)
+void UTest_PingPongBetweenPhaseFactory::RegisterQueries(IEditorDataStorageProvider& DataStorage)
 {
 	using namespace UE::Editor::DataStorage::Queries;
 
diff --git a/Engine/Source/Runtime/TypedElementFramework/Tests/Elements/Framework/TypedElementDataStoragePerformanceTestCommands.h b/Engine/Source/Runtime/TypedElementFramework/Tests/Elements/Framework/TypedElementDataStoragePerformanceTestCommands.h
--- a/Engine/Source/Runtime/TypedElementFramework/Tests/Elements/Framework/TypedElementDataStoragePerformanceTestCommands.h
+++ b/Engine/Source/Runtime/TypedElementFramework/Tests/Elements/Framework/TypedElementDataStoragePerformanceTestCommands.h
@@ -46,13 +46,13 @@ struct FTest_PingPongPostPhys final : public FEditorDataStorageColumn
  *		3. Unnormalized rotation
  */
 UCLASS()
-class UTest_PingPongBetweenPhaseFactory : public UTypedElementDataStorageFactory
+class UTest_PingPongBetweenPhaseFactory : public UEditorDataStorageFactory
 {
 	GENERATED_BODY()
 
 public:
 	~UTest_PingPongBetweenPhaseFactory() override = default;
 
-	void RegisterTables(ITypedElementDataStorageInterface& DataStorage) override;
-	void RegisterQueries(ITypedElementDataStorageInterface& DataStorage) override;
+	void RegisterTables(IEditorDataStorageProvider& DataStorage) override;
+	void RegisterQueries(IEditorDataStorageProvider& DataStorage) override;
 };
\ No newline at end of file
