TEDS UI: Add support to register widget constructors for dynamic columns
Add attribute binding support for dynamic columns

diff --git a/Engine/Plugins/Experimental/EditorDataStorage/Source/TedsCore/Private/TypedElementDatabase.cpp b/Engine/Plugins/Experimental/EditorDataStorage/Source/TedsCore/Private/TypedElementDatabase.cpp
--- a/Engine/Plugins/Experimental/EditorDataStorage/Source/TedsCore/Private/TypedElementDatabase.cpp
+++ b/Engine/Plugins/Experimental/EditorDataStorage/Source/TedsCore/Private/TypedElementDatabase.cpp
@@ -917,6 +917,8 @@ bool UEditorDataStorage::MatchesColumns(RowHandle Row, const UE::Editor::DataSto
 {
 	if (ActiveEditorEntityManager)
 	{
+		checkf(Conditions.IsCompiled(), TEXT("Query Conditions must be compiled before they can be used"));
+		
 		FMassEntityHandle Entity = FMassEntityHandle::FromNumber(Row);
 		if (ActiveEditorEntityManager->IsEntityActive(Entity))
 		{
diff --git a/Engine/Plugins/Experimental/EditorDataStorage/Source/TedsCore/Private/TypedElementDatabaseUI.cpp b/Engine/Plugins/Experimental/EditorDataStorage/Source/TedsCore/Private/TypedElementDatabaseUI.cpp
--- a/Engine/Plugins/Experimental/EditorDataStorage/Source/TedsCore/Private/TypedElementDatabaseUI.cpp
+++ b/Engine/Plugins/Experimental/EditorDataStorage/Source/TedsCore/Private/TypedElementDatabaseUI.cpp
@@ -270,7 +270,7 @@ void UEditorDataStorageUi::CreateWidgetConstructors(FName Purpose,
 	{
 		for (const FWidgetFactory& Factory : PurposeInfo->Factories)
 		{
-			if (!CreateSingleWidgetConstructor(Factory.Constructor, Arguments, {}, Factory.Columns, Callback))
+			if (!CreateSingleWidgetConstructor(Factory.Constructor, Arguments, {}, Factory.GetConditions(Storage), Callback))
 			{
 				return;
 			}
@@ -298,10 +298,10 @@ void UEditorDataStorageUi::CreateWidgetConstructors(FName Purpose, EMatchApproac
 			// This is the only call that requires the array of factories to be sorted from largest to smallest number
 			// of columns, so lazily sort only when needed.
 			PurposeInfo->Factories.StableSort(
-				[](const FWidgetFactory& Lhs, const FWidgetFactory& Rhs)
+				[this](const FWidgetFactory& Lhs, const FWidgetFactory& Rhs)
 				{
-					int32 LeftSize = Lhs.Columns.MinimumColumnMatchRequired();
-					int32 RightSize = Rhs.Columns.MinimumColumnMatchRequired();
+					int32 LeftSize = Lhs.GetConditions(Storage).MinimumColumnMatchRequired();
+					int32 RightSize = Rhs.GetConditions(Storage).MinimumColumnMatchRequired();
 					return LeftSize > RightSize;
 				});
 			PurposeInfo->bIsSorted = true;
@@ -473,7 +473,7 @@ void UEditorDataStorageUi::CreateWidgetConstructors_LongestMatch(const TArray<FW
 	TArray<TWeakObjectPtr<const UScriptStruct>> MatchedColumns;
 	for (auto FactoryIt = WidgetFactories.CreateConstIterator(); FactoryIt && !Columns.IsEmpty(); ++FactoryIt)
 	{
-		if (FactoryIt->Columns.MinimumColumnMatchRequired() > Columns.Num())
+		if (FactoryIt->GetConditions(Storage).MinimumColumnMatchRequired() > Columns.Num())
 		{
 			// There are more columns required for this factory than there are in the requested columns list so skip this
 			// factory.
@@ -481,8 +481,8 @@ void UEditorDataStorageUi::CreateWidgetConstructors_LongestMatch(const TArray<FW
 		}
 
 		MatchedColumns.Reset();
-
-		if (FactoryIt->Columns.Verify(MatchedColumns, Columns, true))
+		
+		if (FactoryIt->GetConditions(Storage).Verify(MatchedColumns, Columns, true))
 		{
 			// Remove the found columns from the requested list.
 			Algo::SortBy(MatchedColumns, [](const TWeakObjectPtr<const UScriptStruct>& Column) { return Column.Get(); });
@@ -508,7 +508,7 @@ void UEditorDataStorageUi::CreateWidgetConstructors_LongestMatch(const TArray<FW
 				--ColumnsEnd;
 			}
 			
-			if (!CreateSingleWidgetConstructor(FactoryIt->Constructor, Arguments, MoveTemp(MatchedColumns), FactoryIt->Columns, Callback))
+			if (!CreateSingleWidgetConstructor(FactoryIt->Constructor, Arguments, MoveTemp(MatchedColumns), FactoryIt->GetConditions(Storage), Callback))
 			{
 				return;
 			}
@@ -526,21 +526,21 @@ void UEditorDataStorageUi::CreateWidgetConstructors_ExactMatch(const TArray<FWid
 	{
 		// If there are more matches required that there are columns, then there will never be an exact match.
 		// Less than the column count can still result in a match that covers all columns.
-		if (Factory.Columns.MinimumColumnMatchRequired() > ColumnCount)
+		if (Factory.GetConditions(Storage).MinimumColumnMatchRequired() > ColumnCount)
 		{
 			continue;
 		}
 
 		MatchedColumns.Reset();
 
-		if (Factory.Columns.Verify(MatchedColumns, Columns, true))
+		if (Factory.GetConditions(Storage).Verify(MatchedColumns, Columns, true))
 		{
 			Algo::SortBy(MatchedColumns, [](const TWeakObjectPtr<const UScriptStruct>& Column) { return Column.Get(); });
 			MatchedColumns.SetNum(Algo::Unique(MatchedColumns), EAllowShrinking::No);
 			if (MatchedColumns.Num() == Columns.Num())
 			{
 				Columns.Reset();
-				CreateSingleWidgetConstructor(Factory.Constructor, Arguments, MoveTemp(MatchedColumns), Factory.Columns, Callback);
+				CreateSingleWidgetConstructor(Factory.Constructor, Arguments, MoveTemp(MatchedColumns), Factory.GetConditions(Storage), Callback);
 				return;
 			}
 		}
@@ -559,7 +559,7 @@ void UEditorDataStorageUi::CreateWidgetConstructors_SingleMatch(const TArray<FWi
 	{
 		for (; FactoryIt != FactoryEnd; ++FactoryIt)
 		{
-			TConstArrayView<TWeakObjectPtr<const UScriptStruct>> ColumnData = (*FactoryIt).Columns.GetColumns();
+			TConstArrayView<TWeakObjectPtr<const UScriptStruct>> ColumnData = (*FactoryIt).GetConditions(Storage).GetColumns();
 			if (ColumnData.Num() > 1)
 			{
 				// Moved passed the point where factories only have a single column.
@@ -575,7 +575,7 @@ void UEditorDataStorageUi::CreateWidgetConstructors_SingleMatch(const TArray<FWi
 			{
 				Columns.RemoveAt(ColumnIndex);
 				CreateSingleWidgetConstructor((*FactoryIt).Constructor, Arguments, 
-					TArray<TWeakObjectPtr<const UScriptStruct>>(ColumnData), (*FactoryIt).Columns, Callback);
+					TArray<TWeakObjectPtr<const UScriptStruct>>(ColumnData), (*FactoryIt).GetConditions(Storage), Callback);
 				// Match was found so move on to the next column in the column.
 				break;
 			}
@@ -616,3 +616,10 @@ UEditorDataStorageUi::FWidgetFactory::FWidgetFactory(TUniquePtr<FTypedElementWid
 	checkf(std::get<TUniquePtr<FTypedElementWidgetConstructor>>(Constructor)->GetTypeInfo(),
 		TEXT("Widget constructor registered that didn't contain valid type information."));
 }
+
+const UE::Editor::DataStorage::Queries::FConditions& UEditorDataStorageUi::FWidgetFactory::GetConditions(IEditorDataStorageProvider* DataStorage) const
+{
+	Columns.Compile(UE::Editor::DataStorage::Queries::FEditorStorageQueryConditionCompileContext(DataStorage));
+	
+	return Columns;
+}
diff --git a/Engine/Plugins/Experimental/EditorDataStorage/Source/TedsCore/Private/TypedElementDatabaseUI.h b/Engine/Plugins/Experimental/EditorDataStorage/Source/TedsCore/Private/TypedElementDatabaseUI.h
--- a/Engine/Plugins/Experimental/EditorDataStorage/Source/TedsCore/Private/TypedElementDatabaseUI.h
+++ b/Engine/Plugins/Experimental/EditorDataStorage/Source/TedsCore/Private/TypedElementDatabaseUI.h
@@ -58,9 +58,12 @@ public:
 private:
 	struct FWidgetFactory
 	{
-		using ConstructorType = std::variant<const UScriptStruct*, TUniquePtr<FTypedElementWidgetConstructor>>;
+	private:
+		// Private and mutable so we can control access to it via GetConditions() to compile the conditions on demand
+		mutable UE::Editor::DataStorage::Queries::FConditions Columns;
 
-		UE::Editor::DataStorage::Queries::FConditions Columns;
+	public:
+		using ConstructorType = std::variant<const UScriptStruct*, TUniquePtr<FTypedElementWidgetConstructor>>;
 		ConstructorType Constructor;
 
 		FWidgetFactory() = default;
@@ -68,6 +71,8 @@ private:
 		explicit FWidgetFactory(TUniquePtr<FTypedElementWidgetConstructor>&& InConstructor);
 		FWidgetFactory(const UScriptStruct* InConstructor, UE::Editor::DataStorage::Queries::FConditions&& InColumns);
 		FWidgetFactory(TUniquePtr<FTypedElementWidgetConstructor>&& InConstructor, UE::Editor::DataStorage::Queries::FConditions&& InColumns);
+		const UE::Editor::DataStorage::Queries::FConditions& GetConditions(IEditorDataStorageProvider* DataStorage) const;
+		
 	};
 
 	struct FPurposeInfo
diff --git a/Engine/Plugins/Experimental/EditorDataStorageFeatures/Source/TedsAssetData/Private/Widgets/StaticMeshTrianglesWidget.cpp b/Engine/Plugins/Experimental/EditorDataStorageFeatures/Source/TedsAssetData/Private/Widgets/StaticMeshTrianglesWidget.cpp
new file mode 100644
--- /dev/null
+++ b/Engine/Plugins/Experimental/EditorDataStorageFeatures/Source/TedsAssetData/Private/Widgets/StaticMeshTrianglesWidget.cpp
@@ -0,0 +1,32 @@
+// Copyright Epic Games, Inc. All Rights Reserved.
+
+#include "Widgets/StaticMeshTrianglesWidget.h"
+
+#include "Elements/Framework/TypedElementAttributeBinding.h"
+#include "Internationalization/Text.h"
+#include "TedsAssetDataColumns.h"
+#include "UObject/ObjectMacros.h"
+#include "Widgets/Text/STextBlock.h"
+
+void UStaticMeshTrianglesWidgetFactory::RegisterWidgetConstructors(IEditorDataStorageProvider& DataStorage,
+                                                        IEditorDataStorageUiProvider& DataStorageUi) const
+{
+	using namespace UE::Editor::DataStorage::Queries;
+	DataStorageUi.RegisterWidgetFactory<FStaticMeshTrianglesWidgetConstructor>(TEXT("General.Cell"), TColumn<FItemStringAttributeColumn_Experimental>("Triangles"));
+}
+
+FStaticMeshTrianglesWidgetConstructor::FStaticMeshTrianglesWidgetConstructor()
+	: FSimpleWidgetConstructor(StaticStruct())
+{
+}
+
+TSharedPtr<SWidget> FStaticMeshTrianglesWidgetConstructor::CreateWidget(IEditorDataStorageProvider* DataStorage,
+	IEditorDataStorageUiProvider* DataStorageUi, UE::Editor::DataStorage::RowHandle TargetRow, UE::Editor::DataStorage::RowHandle WidgetRow,
+	const UE::Editor::DataStorage::FMetaDataView& Arguments)
+{
+	UE::Editor::DataStorage::FAttributeBinder Binder(TargetRow, DataStorage);
+	
+	return SNew(STextBlock)
+			.Text(Binder.BindText(&FItemStringAttributeColumn_Experimental::Value, FName("Triangles")));
+			
+}
diff --git a/Engine/Plugins/Experimental/EditorDataStorageFeatures/Source/TedsAssetData/Private/Widgets/StaticMeshTrianglesWidget.h b/Engine/Plugins/Experimental/EditorDataStorageFeatures/Source/TedsAssetData/Private/Widgets/StaticMeshTrianglesWidget.h
new file mode 100644
--- /dev/null
+++ b/Engine/Plugins/Experimental/EditorDataStorageFeatures/Source/TedsAssetData/Private/Widgets/StaticMeshTrianglesWidget.h
@@ -0,0 +1,41 @@
+// Copyright Epic Games, Inc. All Rights Reserved.
+
+#pragma once
+
+#include "Elements/Interfaces/TypedElementDataStorageFactory.h"
+#include "Elements/Interfaces/TypedElementDataStorageUiInterface.h"
+
+#include "StaticMeshTrianglesWidget.generated.h"
+
+class IEditorDataStorageProvider;
+class UScriptStruct;
+
+UCLASS()
+class UStaticMeshTrianglesWidgetFactory : public UEditorDataStorageFactory
+{
+	GENERATED_BODY()
+
+public:
+	~UStaticMeshTrianglesWidgetFactory() override = default;
+
+	TEDSASSETDATA_API void RegisterWidgetConstructors(IEditorDataStorageProvider& DataStorage,
+		IEditorDataStorageUiProvider& DataStorageUi) const override;
+};
+
+// Specialized widget to display the "Triangles" metadata on static mesh assets
+USTRUCT()
+struct FStaticMeshTrianglesWidgetConstructor : public FSimpleWidgetConstructor
+{
+	GENERATED_BODY()
+
+public:
+	TEDSASSETDATA_API FStaticMeshTrianglesWidgetConstructor();
+	~FStaticMeshTrianglesWidgetConstructor() override = default;
+
+	TEDSASSETDATA_API virtual TSharedPtr<SWidget> CreateWidget(
+		IEditorDataStorageProvider* DataStorage,
+		IEditorDataStorageUiProvider* DataStorageUi,
+		UE::Editor::DataStorage::RowHandle TargetRow,
+		UE::Editor::DataStorage::RowHandle WidgetRow,
+		const UE::Editor::DataStorage::FMetaDataView& Arguments) override;
+};
\ No newline at end of file
diff --git a/Engine/Plugins/Experimental/EditorDataStorageFeatures/Source/TedsContentBrowser/Private/TedsContentBrowserModule.cpp b/Engine/Plugins/Experimental/EditorDataStorageFeatures/Source/TedsContentBrowser/Private/TedsContentBrowserModule.cpp
--- a/Engine/Plugins/Experimental/EditorDataStorageFeatures/Source/TedsContentBrowser/Private/TedsContentBrowserModule.cpp
+++ b/Engine/Plugins/Experimental/EditorDataStorageFeatures/Source/TedsContentBrowser/Private/TedsContentBrowserModule.cpp
@@ -130,13 +130,21 @@ namespace UE::Editor::ContentBrowser
 		
 		RowQueryStack = MakeShared<FQueryStackNode_RowView>(&Rows);
 
+		// Sample dynamic column to display the "Triangles" attribute on static meshes
+		// We probably want the dynamic columns in the table viewer to be data driven based on the rows in the future
+		const UScriptStruct* DynamicStaticMeshTrianglesColumn = DataStorage->GenerateDynamicColumn(FDynamicColumnDescription
+						{
+							.TemplateType = FItemStringAttributeColumn_Experimental::StaticStruct(),
+							.Identifier = "Triangles"
+						});
+
 		// Create the table viewer widget
 		TableViewer = SNew(STedsTableViewer)
 					.QueryStack(RowQueryStack)
 					.CellWidgetPurposes({TEXT("General.RowLabel"), TEXT("General.Cell")})
 					// Default list of columns to display
-					.Columns({ FNameColumn::StaticStruct(), FAssetTag::StaticStruct(), FAssetPathColumn_Experimental::StaticStruct(),
-						FDiskSizeColumn::StaticStruct(), FVirtualPathColumn_Experimental::StaticStruct() })
+					.Columns({ FNameColumn::StaticStruct(), FAssetClassColumn::StaticStruct(), FAssetTag::StaticStruct(), FAssetPathColumn_Experimental::StaticStruct(),
+						FDiskSizeColumn::StaticStruct(), FVirtualPathColumn_Experimental::StaticStruct(), DynamicStaticMeshTrianglesColumn })
 					.ListSelectionMode(ESelectionMode::Multi)
 					.OnSelectionChanged_Lambda([this](RowHandle Row)
 					{
diff --git a/Engine/Source/Runtime/TypedElementFramework/Private/Elements/Common/TypedElementQueryConditions.cpp b/Engine/Source/Runtime/TypedElementFramework/Private/Elements/Common/TypedElementQueryConditions.cpp
--- a/Engine/Source/Runtime/TypedElementFramework/Private/Elements/Common/TypedElementQueryConditions.cpp
+++ b/Engine/Source/Runtime/TypedElementFramework/Private/Elements/Common/TypedElementQueryConditions.cpp
@@ -2,18 +2,39 @@
 
 #include "Elements/Common/TypedElementQueryConditions.h"
 
+#include "Elements/Interfaces/TypedElementDataStorageInterface.h"
 #include "UObject/Class.h"
 
 namespace UE::Editor::DataStorage::Queries
 {
+	FEditorStorageQueryConditionCompileContext::FEditorStorageQueryConditionCompileContext(IEditorDataStorageProvider* InDataStorage)
+		: DataStorage(InDataStorage)
+	{
+	}
+
+	const UScriptStruct* FEditorStorageQueryConditionCompileContext::GenerateDynamicColumn(
+		const FDynamicColumnDescription& Description) const
+	{
+		return DataStorage->GenerateDynamicColumn(Description);
+	}
+
+	FConditions::FConditions()
+	{
+		Identifiers.Init(NAME_None, MaxColumnCount);
+	}
+
 	FConditions::FConditions(FColumnBase Column)
 		: ColumnCount(1)
 	{
+		Identifiers.Init(NAME_None, MaxColumnCount);
 		Columns[0] = Column.TypeInfo;
+		Identifiers[0] = Column.Identifier;
 	}
 
 	void FConditions::AppendToString(FString& Output) const
 	{
+		checkf(bIsCompiled, TEXT("Query Conditions must call Compile() before you can use them"));
+		
 		if (TokenCount > 0)
 		{
 			Output += "{ ";
@@ -62,6 +83,8 @@ namespace UE::Editor::DataStorage::Queries
 
 	bool FConditions::Verify(TConstArrayView<FColumnBase> AvailableColumns) const
 	{
+		checkf(bIsCompiled, TEXT("Query Conditions must call Compile() before you can use them"));
+		
 		return VerifyBootstrap(
 			[&AvailableColumns](uint8_t ColumnIndex, TWeakObjectPtr<const UScriptStruct> Column)
 			{
@@ -80,6 +103,9 @@ namespace UE::Editor::DataStorage::Queries
 		bool AvailableColumnsAreSorted) const
 	{
 		static_assert(MaxColumnCount < 64, "Query conditions use a bit mask to locate matches. As a result MaxColumnCount can be larger than 64.");
+
+		checkf(bIsCompiled, TEXT("Query Conditions must call Compile() before you can use them"));
+		
 		uint64 Matches = 0;
 		bool Result = AvailableColumnsAreSorted
 			? VerifyBootstrap(
@@ -116,8 +142,11 @@ namespace UE::Editor::DataStorage::Queries
 		return Result;
 	}
 
-	bool FConditions::Verify(TConstArrayView<TWeakObjectPtr<const UScriptStruct>> AvailableColumns, bool AvailableColumnsAreSorted) const
+	bool FConditions::Verify(TConstArrayView<TWeakObjectPtr<const UScriptStruct>> AvailableColumns,
+		bool AvailableColumnsAreSorted) const
 	{
+		checkf(bIsCompiled, TEXT("Query Conditions must call Compile() before you can use them"));
+		
 		return AvailableColumnsAreSorted
 			? VerifyBootstrap(
 				[&AvailableColumns](uint8_t ColumnIndex, TWeakObjectPtr<const UScriptStruct> Column)
@@ -136,9 +165,13 @@ namespace UE::Editor::DataStorage::Queries
 	}
 
 	bool FConditions::Verify(TArray<TWeakObjectPtr<const UScriptStruct>>& MatchedColumns,
-		TConstArrayView<TWeakObjectPtr<const UScriptStruct>> AvailableColumns, bool AvailableColumnsAreSorted) const
+		TConstArrayView<TWeakObjectPtr<const UScriptStruct>> AvailableColumns,
+		bool AvailableColumnsAreSorted) const
 	{
 		static_assert(MaxColumnCount < 64, "Query conditions use a bit mask to locate matches. As a result MaxColumnCount cannot be larger than 64.");
+		
+		checkf(bIsCompiled, TEXT("Query Conditions must call Compile() before you can use them"));
+		
 		uint64 Matches = 0;
 		bool Result = AvailableColumnsAreSorted
 			? VerifyBootstrap(
@@ -169,6 +202,8 @@ namespace UE::Editor::DataStorage::Queries
 
 	bool FConditions::Verify(TSet<TWeakObjectPtr<const UScriptStruct>> AvailableColumns) const
 	{
+		checkf(bIsCompiled, TEXT("Query Conditions must call Compile() before you can use them"));
+		
 		return VerifyBootstrap(
 			[&AvailableColumns](uint8_t ColumnIndex, TWeakObjectPtr<const UScriptStruct> Column)
 			{
@@ -178,6 +213,8 @@ namespace UE::Editor::DataStorage::Queries
 
 	bool FConditions::Verify(ContainsCallback Callback) const
 	{
+		checkf(bIsCompiled, TEXT("Query Conditions must call Compile() before you can use them"));
+		
 		return VerifyBootstrap(Callback);
 	}
 
@@ -200,6 +237,8 @@ namespace UE::Editor::DataStorage::Queries
 
 	TConstArrayView<TWeakObjectPtr<const UScriptStruct>> FConditions::GetColumns() const
 	{
+		checkf(bIsCompiled, TEXT("Query Conditions must call Compile() before you can use them"));
+		
 		return TConstArrayView<TWeakObjectPtr<const UScriptStruct>>(Columns, ColumnCount);
 	}
 
@@ -208,6 +247,39 @@ namespace UE::Editor::DataStorage::Queries
 		return ColumnCount == 0;
 	}
 
+	FConditions& FConditions::Compile(const IQueryConditionCompileContext& CompileContext)
+	{
+		if(bIsCompiled)
+		{
+			return *this;
+		}
+
+		// Resolve all dynamic columns by generating their UScriptStruct
+		for (uint8_t ColumnIndex = 0; ColumnIndex < ColumnCount; ++ColumnIndex)
+		{
+			if(Identifiers[ColumnIndex] != NAME_None)
+			{
+				const UScriptStruct* ResolvedDynamicColumn = CompileContext.GenerateDynamicColumn(
+					FDynamicColumnDescription
+						{
+							.TemplateType = Columns[ColumnIndex].Get(),
+							.Identifier = Identifiers[ColumnIndex]
+						});
+				
+				Columns[ColumnIndex] = ResolvedDynamicColumn;
+			}
+		}
+
+		bIsCompiled = true;
+
+		return *this;
+	}
+
+	bool FConditions::IsCompiled() const
+	{
+		return bIsCompiled;
+	}
+
 	void FConditions::AppendName(FString& Output, TWeakObjectPtr<const UScriptStruct> TypeInfo) const
 	{
 #if WITH_EDITORONLY_DATA
@@ -381,6 +453,7 @@ namespace UE::Editor::DataStorage::Queries
 		for (uint8_t Index = 0; Index < Source.ColumnCount; ++Index)
 		{
 			Target.Columns[Target.ColumnCount + Index] = Source.Columns[Index];
+			Target.Identifiers[Target.ColumnCount + Index] = Source.Identifiers[Index];
 		}
 
 		checkf(Target.TokenCount + Source.TokenCount < MaxTokenCount, TEXT("Too many operations in the query. Try simplifying your query."));
@@ -396,7 +469,8 @@ namespace UE::Editor::DataStorage::Queries
 	FConditions operator&&(const FConditions& Lhs, FColumnBase Rhs)
 	{
 		FConditions Result = Lhs;
-		Result.Columns[Result.ColumnCount++] = Rhs.TypeInfo;
+		Result.Columns[Result.ColumnCount] = Rhs.TypeInfo;
+		Result.Identifiers[Result.ColumnCount++] = Rhs.Identifier;
 		Result.Tokens[Result.TokenCount++] = FConditions::Token::And;
 		return Result;
 	}
@@ -420,7 +494,8 @@ namespace UE::Editor::DataStorage::Queries
 	FConditions operator&&(FColumnBase Lhs, FColumnBase Rhs)
 	{
 		FConditions Result(Lhs);
-		Result.Columns[Result.ColumnCount++] = Rhs.TypeInfo;
+		Result.Columns[Result.ColumnCount] = Rhs.TypeInfo;
+		Result.Identifiers[Result.ColumnCount++] = Rhs.Identifier;
 		Result.Tokens[Result.TokenCount++] = FConditions::Token::And;
 		return Result;
 	}
@@ -438,7 +513,8 @@ namespace UE::Editor::DataStorage::Queries
 	FConditions operator||(const FConditions& Lhs, FColumnBase Rhs)
 	{
 		FConditions Result = Lhs;
-		Result.Columns[Result.ColumnCount++] = Rhs.TypeInfo;
+		Result.Columns[Result.ColumnCount] = Rhs.TypeInfo;
+		Result.Identifiers[Result.ColumnCount++] = Rhs.Identifier;
 		Result.Tokens[Result.TokenCount++] = FConditions::Token::Or;
 		return Result;
 	}
@@ -462,7 +538,8 @@ namespace UE::Editor::DataStorage::Queries
 	FConditions operator||(FColumnBase Lhs, FColumnBase Rhs)
 	{
 		FConditions Result(Lhs);
-		Result.Columns[Result.ColumnCount++] = Rhs.TypeInfo;
+		Result.Columns[Result.ColumnCount] = Rhs.TypeInfo;
+		Result.Identifiers[Result.ColumnCount++] = Rhs.Identifier;
 		Result.Tokens[Result.TokenCount++] = FConditions::Token::Or;
 		return Result;
 	}
diff --git a/Engine/Source/Runtime/TypedElementFramework/Public/Elements/Common/TypedElementQueryConditions.h b/Engine/Source/Runtime/TypedElementFramework/Public/Elements/Common/TypedElementQueryConditions.h
--- a/Engine/Source/Runtime/TypedElementFramework/Public/Elements/Common/TypedElementQueryConditions.h
+++ b/Engine/Source/Runtime/TypedElementFramework/Public/Elements/Common/TypedElementQueryConditions.h
@@ -6,6 +6,7 @@
 #include "Templates/SubclassOf.h"
 #include "UObject/WeakObjectPtrTemplates.h"
 
+class IEditorDataStorageProvider;
 class UScriptStruct;
 
 namespace UE::Editor::DataStorage::Queries
@@ -13,24 +14,56 @@ namespace UE::Editor::DataStorage::Queries
 	struct FColumnBase
 	{
 		TWeakObjectPtr<const UScriptStruct> TypeInfo = nullptr;
+		FName Identifier = NAME_None;
 
 	protected:
 		FColumnBase() = default;
-		constexpr explicit FColumnBase(TWeakObjectPtr<const UScriptStruct> ColumnTypeInfo) : TypeInfo(ColumnTypeInfo) {}
+		constexpr explicit FColumnBase(TWeakObjectPtr<const UScriptStruct> ColumnTypeInfo, const FName& InIdentifier)
+			: TypeInfo(ColumnTypeInfo)
+			, Identifier(InIdentifier)
+		{}
 	};
 
 	template<typename T = void>
 	struct TColumn final : public FColumnBase
 	{
 		template <typename U = T> requires (!std::is_same_v<U, void>)
-			constexpr TColumn() : FColumnBase(T::StaticStruct()) {}
+			constexpr TColumn(const FName& Identifier = NAME_None) : FColumnBase(T::StaticStruct(), Identifier) {}
+		
 		template <typename U = T> requires (std::is_same_v<U, void>)
-			constexpr explicit TColumn(TWeakObjectPtr<const UScriptStruct> ColumnTypeInfo) : FColumnBase(ColumnTypeInfo) {};
+			constexpr explicit TColumn(TWeakObjectPtr<const UScriptStruct> ColumnTypeInfo, const FName& Identifier = NAME_None)
+				: FColumnBase(ColumnTypeInfo, Identifier)
+		{}
 	};
 
+	/*
+	 * Compile Context used to resolve dynamic columns when FConditions are compiled
+	 */
+	class IQueryConditionCompileContext 
+	{
+	public:
+		TYPEDELEMENTFRAMEWORK_API virtual ~IQueryConditionCompileContext() = default;
+		TYPEDELEMENTFRAMEWORK_API virtual const UScriptStruct* GenerateDynamicColumn(const FDynamicColumnDescription&) const = 0;
+	};
+
+	/*
+	 * Specialized compile context that accepts an IEditorDataStorageProvider
+	 */
+	class FEditorStorageQueryConditionCompileContext : public IQueryConditionCompileContext 
+	{
+	public:
+		TYPEDELEMENTFRAMEWORK_API explicit FEditorStorageQueryConditionCompileContext(IEditorDataStorageProvider* InDataStorage);
+		TYPEDELEMENTFRAMEWORK_API virtual const UScriptStruct* GenerateDynamicColumn(const FDynamicColumnDescription& Description) const override;
+		
+	private:
+		IEditorDataStorageProvider* DataStorage;
+	};
+
+
 	/**
 	 * Product of boolean combination of multiple columns. This can be used to verify if a collection of columns match
 	 * the stored columns.
+	 * NOTE: You must call Compile() before you call any members accessing the conditions
 	 */
 	class FConditions final
 	{
@@ -47,11 +80,17 @@ namespace UE::Editor::DataStorage::Queries
 		TYPEDELEMENTFRAMEWORK_API friend FConditions operator||(FColumnBase Lhs, FColumnBase Rhs);
 		TYPEDELEMENTFRAMEWORK_API friend FConditions operator||(FColumnBase Lhs, const FConditions& Rhs);
 
-		FConditions() = default;
+		TYPEDELEMENTFRAMEWORK_API FConditions();
 		// Not marked as "explicit" to allow conversion from a column. This means that conditions with a single
 		// argument can be written in the same way as ones that use combinations.
 		TYPEDELEMENTFRAMEWORK_API FConditions(FColumnBase Column);
 
+		// Compile must be called before using any functions that access the columns
+		TYPEDELEMENTFRAMEWORK_API FConditions& Compile(const IQueryConditionCompileContext& CompileContext);
+
+		// Check whether these query conditions have been compiled
+		TYPEDELEMENTFRAMEWORK_API bool IsCompiled() const;
+
 		/** Convert the conditions into a string and append them to the provided string. */
 		TYPEDELEMENTFRAMEWORK_API void AppendToString(FString& Output) const;
 
@@ -117,12 +156,14 @@ namespace UE::Editor::DataStorage::Queries
 		};
 		using ColumnArray = TWeakObjectPtr<const UScriptStruct>[MaxColumnCount];
 		using TokenArray = Token[MaxTokenCount];
+		using IdentifierArray = TArray<FName, TInlineAllocator<MaxColumnCount>>;
 
 		ColumnArray Columns{ nullptr };
 		TokenArray Tokens{ Token::None };
+		IdentifierArray Identifiers;
 		uint8_t ColumnCount = 0;
 		uint8_t TokenCount = 0;
-
+		bool bIsCompiled = false;
 	};
 
 	TYPEDELEMENTFRAMEWORK_API FConditions operator&&(const FConditions& Lhs, FColumnBase Rhs);
diff --git a/Engine/Source/Runtime/TypedElementFramework/Public/Elements/Framework/TypedElementAttributeBinding.h b/Engine/Source/Runtime/TypedElementFramework/Public/Elements/Framework/TypedElementAttributeBinding.h
--- a/Engine/Source/Runtime/TypedElementFramework/Public/Elements/Framework/TypedElementAttributeBinding.h
+++ b/Engine/Source/Runtime/TypedElementFramework/Public/Elements/Framework/TypedElementAttributeBinding.h
@@ -38,6 +38,7 @@ namespace UE::Editor::DataStorage
 		 * 
 		 * @param InVariable The data member inside a column to be bound
 		 * @param InDefaultValue The default value to be used when the column isn't present on a row
+		 * @param InIdentifier The identifier for this column if it is a dynamic column, NAME_None if it is not a dynamic column
 		 * @return A TAttribute bound to the row, column pair specified
 		 *
 		 * Example:
@@ -45,7 +46,8 @@ namespace UE::Editor::DataStorage
 		 * TAttribute<FString> TestAttribute(Binder.BindData(&FTypedElementLabelColumn::Label))
 		 */
 		template <typename AttributeType, TDataColumnType ColumnType>
-		TAttribute<AttributeType> BindData(AttributeType ColumnType::* InVariable, const AttributeType& InDefaultValue = AttributeType());
+		TAttribute<AttributeType> BindData(AttributeType ColumnType::* InVariable, const AttributeType& InDefaultValue = AttributeType(),
+			const FName& InIdentifier = NAME_None);
 
 		/**
 		 * Bind a specific data member inside a TEDS column to an attribute of a different type than the data by providing a conversion function
@@ -53,7 +55,8 @@ namespace UE::Editor::DataStorage
 		 *
 		 * @param InVariable The data member inside a column to be bound
 		 * @param InConverter Conversion function to convert from DataType -> AttributeType
-		 * @param InDefaultValue The default value to be used when the column isn't present on a row 
+		 * @param InDefaultValue The default value to be used when the column isn't present on a row
+		 * @param InIdentifier The identifier for this column if it is a dynamic column, NAME_None if it is not a dynamic column
 		 * @return A TAttribute bound to the row, column pair specified
 		 *
 		 * Example:
@@ -66,43 +69,50 @@ namespace UE::Editor::DataStorage
 		 *                                 ));
 		 */
 		template <typename AttributeType, typename DataType, TDataColumnType ColumnType>
-		TAttribute<AttributeType> BindData(DataType ColumnType::* InVariable, const TFunction<AttributeType(const DataType&)>& InConverter, const DataType& InDefaultValue = DataType());
+		TAttribute<AttributeType> BindData(DataType ColumnType::* InVariable, const TFunction<AttributeType(const DataType&)>& InConverter,
+			const DataType& InDefaultValue = DataType(), const FName& InIdentifier = NAME_None);
 
 		/**
 		 * Overload for the conversion binder to accept lambdas instead of TFunctions
 		 *
 		 * @param InVariable The data member inside a column to be bound
 		 * @param InConverter Conversion function to convert from DataType -> AttributeType
-		 * @param InDefaultValue The default value to be used when the column isn't present on a row
+		 * @param InDefaultValue The default value to be used when the column isn't present on a row\
+		 * @param InIdentifier The identifier for this column if it is a dynamic column, NAME_None if it is not a dynamic column
 		 * @return A TAttribute bound to the row, column pair specified
 		 */
 		template <typename DataType, TDataColumnType ColumnType, typename FunctionType>
 			requires AttributeBinderInvocable<FunctionType, DataType>
-		auto BindData(DataType ColumnType::* InVariable, FunctionType InConverter, const DataType& InDefaultValue = DataType());
+		auto BindData(DataType ColumnType::* InVariable, FunctionType InConverter, const DataType& InDefaultValue = DataType(),
+			const FName& InIdentifier = NAME_None);
 
 		/**
 		 * Bind a delegate inside a Teds column to a SLATE_EVENT macro on a widget
 		 * @param InVariable The delegate inside the Teds column
+		 * @param InIdentifier The identifier for this column if it is a dynamic column, NAME_None if it is not a dynamic column
 		 * @return A delegate that can be provided to an event on a slate widget
 		 */
-		template <typename InRetValType, typename... ParamTypes, typename ColumnType>
-		TDelegate<InRetValType(ParamTypes...)> BindEvent(TDelegate<InRetValType(ParamTypes...)> ColumnType::* InVariable);
+		template <typename InRetValType, typename... ParamTypes, TDataColumnType ColumnType>
+		TDelegate<InRetValType(ParamTypes...)> BindEvent(TDelegate<InRetValType(ParamTypes...)> ColumnType::* InVariable,
+			const FName& InIdentifier = NAME_None);
 
 		/**
 		 * Directly bind an FString member in a Teds column to an FText attribute as a shortcut
-		 * @param InFStringVariable The FString variable 
+		 * @param InFStringVariable The FString variable
+		 * @param InIdentifier The identifier for this column if it is a dynamic column, NAME_None if it is not a dynamic column
 		 * @return A delegate that can be provided to a text widget in Slate (e.g STextBlock)
 		 */
 		template <TDataColumnType ColumnType>
-		TAttribute<FText> BindText(FString ColumnType::* InFStringVariable);
+		TAttribute<FText> BindText(FString ColumnType::* InFStringVariable, const FName& InIdentifier = NAME_None);
 
 		/**
 		 * Directly bind an FName member in a Teds column to an FText attribute as a shortcut
-		 * @param InFNameVariable The FName variable 
+		 * @param InFNameVariable The FName variable
+		 * @param InIdentifier The identifier for this column if it is a dynamic column, NAME_None if it is not a dynamic column
 		 * @return A delegate that can be provided to a text widget in Slate (e.g STextBlock)
 		 */
 		template <TDataColumnType ColumnType>
-		TAttribute<FText> BindText(FName ColumnType::* InFNameVariable);
+		TAttribute<FText> BindText(FName ColumnType::* InFNameVariable, const FName& InIdentifier = NAME_None);
 		
 	private:
 
diff --git a/Engine/Source/Runtime/TypedElementFramework/Public/Elements/Framework/TypedElementAttributeBinding.inl b/Engine/Source/Runtime/TypedElementFramework/Public/Elements/Framework/TypedElementAttributeBinding.inl
--- a/Engine/Source/Runtime/TypedElementFramework/Public/Elements/Framework/TypedElementAttributeBinding.inl
+++ b/Engine/Source/Runtime/TypedElementFramework/Public/Elements/Framework/TypedElementAttributeBinding.inl
@@ -6,6 +6,25 @@
 
 namespace UE::Editor::DataStorage
 {
+	namespace AttributeBinding::Private
+	{
+		template <TDataColumnType ColumnType>
+		ColumnType* GetColumn(IEditorDataStorageProvider* DataStorage, RowHandle Row, const FName& InIdentifier)
+		{
+			if(InIdentifier != NAME_None)
+			{
+				if(ColumnType* DynamicColumn = DataStorage->GetColumn<ColumnType>(Row, InIdentifier))
+				{
+					return DynamicColumn;
+				}
+			}
+
+			// If we weren't provided a dynamic column identifier, or the dynamic column could not be found simply fallback to the non-dynamic column
+			return DataStorage->GetColumn<ColumnType>(Row);
+		}
+		
+	}
+	
 	// A property that can be directly be accessed from an object
 	template <typename PropertyType>
 	struct DirectProperty final
@@ -120,7 +139,8 @@ namespace UE::Editor::DataStorage
 	};
 	
 	template <typename AttributeType, TDataColumnType ColumnType>
-	TAttribute<AttributeType> FAttributeBinder::BindData(AttributeType ColumnType::* InVariable, const AttributeType& InDefaultValue)
+	TAttribute<AttributeType> FAttributeBinder::BindData(AttributeType ColumnType::* InVariable, const AttributeType& InDefaultValue,
+		const FName& InIdentifier)
 	{
 		if(!DataStorage)
 		{
@@ -132,10 +152,11 @@ namespace UE::Editor::DataStorage
 		Prop.Bind(InVariable);
 
 		// We don't want any references to this in the lambda because binders are designed to be used and destructed on the stack
-		return TAttribute<AttributeType>::CreateLambda([Property = MoveTemp(Prop), Storage = DataStorage, Row = TargetRow, DefaultValue = InDefaultValue]()
+		return TAttribute<AttributeType>::CreateLambda([Property = MoveTemp(Prop), Storage = DataStorage, Row = TargetRow,
+			DefaultValue = InDefaultValue, Identifier = InIdentifier]()
 		{
 			// Get the column from the given row and use that to return the stored property
-			if(ColumnType* Column = Storage->GetColumn<ColumnType>(Row))
+			if(ColumnType* Column = AttributeBinding::Private::GetColumn<ColumnType>(Storage, Row, Identifier))
 			{
 				return Property.Get(Column, ColumnType::StaticStruct());
 			}
@@ -144,7 +165,8 @@ namespace UE::Editor::DataStorage
 	}
 
 	template <typename AttributeType, typename DataType, TDataColumnType ColumnType>
-	TAttribute<AttributeType> FAttributeBinder::BindData(DataType ColumnType::* InVariable, const TFunction<AttributeType(const DataType&)>& InConverter, const DataType& InDefaultValue)
+	TAttribute<AttributeType> FAttributeBinder::BindData(DataType ColumnType::* InVariable,
+		const TFunction<AttributeType(const DataType&)>& InConverter, const DataType& InDefaultValue, const FName& InIdentifier)
 	{
 		if(!DataStorage)
 		{
@@ -155,9 +177,10 @@ namespace UE::Editor::DataStorage
 		Property<AttributeType> Prop;
 		Prop.Bind(InVariable, InConverter);
 	
-		return TAttribute<AttributeType>::CreateLambda([Property = MoveTemp(Prop), Storage = DataStorage, Row = TargetRow, DefaultValue = InDefaultValue, Converter = InConverter]()
+		return TAttribute<AttributeType>::CreateLambda([Property = MoveTemp(Prop), Storage = DataStorage, Row = TargetRow,
+			DefaultValue = InDefaultValue, Converter = InConverter, Identifier = InIdentifier]()
 		{
-			if(ColumnType* Column = Storage->GetColumn<ColumnType>(Row))
+			if(ColumnType* Column = AttributeBinding::Private::GetColumn<ColumnType>(Storage, Row, Identifier))
 			{
 				return Property.Get(Column, ColumnType::StaticStruct());
 			}
@@ -168,24 +191,26 @@ namespace UE::Editor::DataStorage
 	
 	template <typename DataType, TDataColumnType ColumnType, typename FunctionType>
 		requires AttributeBinderInvocable<FunctionType, DataType>
-	auto FAttributeBinder::BindData(DataType ColumnType::* InVariable, FunctionType InConverter, const DataType& InDefaultValue)
+	auto FAttributeBinder::BindData(DataType ColumnType::* InVariable, FunctionType InConverter, const DataType& InDefaultValue, const FName& InIdentifier)
 	{
 		// Deduce the attribute type from the return value of the converter function
 		using AttributeType = decltype(InConverter(std::declval<DataType>()));
 		
-		return BindData<AttributeType, DataType>(InVariable, TFunction<AttributeType(const DataType&)>(InConverter), InDefaultValue);
+		return BindData<AttributeType, DataType>(InVariable, TFunction<AttributeType(const DataType&)>(InConverter), InDefaultValue, InIdentifier);
 	}
 
-	template <typename InRetValType, typename... ParamTypes, typename ColumnType>
-		TDelegate<InRetValType(ParamTypes...)> FAttributeBinder::BindEvent(TDelegate<InRetValType(ParamTypes...)> ColumnType::* InVariable)
+	template <typename InRetValType, typename... ParamTypes, TDataColumnType ColumnType>
+		TDelegate<InRetValType(ParamTypes...)> FAttributeBinder::BindEvent(TDelegate<InRetValType(ParamTypes...)> ColumnType::* InVariable,
+			const FName& InIdentifier)
 	{
 		// Create a property for the delegate
 		Property<TDelegate<InRetValType(ParamTypes...)>> Prop;
 		Prop.Bind(InVariable);
 	
-		return TDelegate<InRetValType(ParamTypes...)>::CreateLambda([Property = MoveTemp(Prop), Storage = DataStorage, Row = TargetRow](ParamTypes&&... Params)
+		return TDelegate<InRetValType(ParamTypes...)>::CreateLambda([Property = MoveTemp(Prop), Storage = DataStorage,
+			Row = TargetRow, Identifier = InIdentifier](ParamTypes&&... Params)
 		{
-			if(ColumnType* Column = Storage->GetColumn<ColumnType>(Row))
+			if(ColumnType* Column = AttributeBinding::Private::GetColumn<ColumnType>(Storage, Row, Identifier))
 			{
 				// Get the delegate in the bound column for the specified row
 				TDelegate<InRetValType(ParamTypes...)> Delegate = Property.Get(Column, ColumnType::StaticStruct());
@@ -201,21 +226,21 @@ namespace UE::Editor::DataStorage
 	}
 	
 	template <TDataColumnType ColumnType>
-	TAttribute<FText> FAttributeBinder::BindText(FString ColumnType::* InFStringVariable)
+	TAttribute<FText> FAttributeBinder::BindText(FString ColumnType::* InFStringVariable, const FName& InIdentifier)
 	{
 		return BindData(InFStringVariable, [](const FString& InString)
 		{
 			return FText::FromString(InString);
-		});
+		}, FString(), InIdentifier);
 	}
 
 	template <TDataColumnType ColumnType>
-	TAttribute<FText> FAttributeBinder::BindText(FName ColumnType::* InFNameVariable)
+	TAttribute<FText> FAttributeBinder::BindText(FName ColumnType::* InFNameVariable, const FName& InIdentifier)
 	{
 		return BindData(InFNameVariable, [](const FName& InName)
 		{
 			return FText::FromName(InName);
-		});
+		}, FName(), InIdentifier);
 	}
 	
 } // namespace UE::Editor::DataStorage
diff --git a/Engine/Source/Runtime/TypedElementFramework/Tests/Elements/Framework/TypedElementQueryBuilderTests.cpp b/Engine/Source/Runtime/TypedElementFramework/Tests/Elements/Framework/TypedElementQueryBuilderTests.cpp
--- a/Engine/Source/Runtime/TypedElementFramework/Tests/Elements/Framework/TypedElementQueryBuilderTests.cpp
+++ b/Engine/Source/Runtime/TypedElementFramework/Tests/Elements/Framework/TypedElementQueryBuilderTests.cpp
@@ -1,5 +1,6 @@
 // Copyright Epic Games, Inc. All Rights Reserved.
 
+#include "Elements/Framework/TypedElementRegistry.h"
 #if WITH_TESTS
 #include "TypedElementTestColumns.h"
 
@@ -28,8 +29,12 @@ static void AppendColumnName(FString& Output, TWeakObjectPtr<const UScriptStruct
 #endif
 }
 
-static bool TestMatching(const FConditions& TestQuery, const TArray<FColumnBase>& RequestedColumns, bool Expected, bool Sort = false)
+static bool TestMatching(FConditions& TestQuery, const TArray<FColumnBase>& RequestedColumns, bool Expected, bool Sort = false)
 {
+	IEditorDataStorageProvider* Storage = UTypedElementRegistry::GetInstance()->GetMutableDataStorage();
+
+	TestQuery.Compile(UE::Editor::DataStorage::Queries::FEditorStorageQueryConditionCompileContext(Storage));
+	
 	if (Sort)
 	{
 		Algo::SortBy(RequestedColumns,
