[Teds] Baby step move compatibility interface to IModularFeature

#rnx

diff --git a/Engine/Plugins/Experimental/EditorDataStorage/Source/TedsCore/Private/TypedElementDatabase.cpp b/Engine/Plugins/Experimental/EditorDataStorage/Source/TedsCore/Private/TypedElementDatabase.cpp
--- a/Engine/Plugins/Experimental/EditorDataStorage/Source/TedsCore/Private/TypedElementDatabase.cpp
+++ b/Engine/Plugins/Experimental/EditorDataStorage/Source/TedsCore/Private/TypedElementDatabase.cpp
@@ -34,7 +34,7 @@ FAutoConsoleCommandWithOutputDevice PrintQueryCallbacksConsoleCommand(
 		{
 			if (UTypedElementRegistry* Registry = UTypedElementRegistry::GetInstance())
 			{
-				if (UEditorDataStorage* DataStorage = Cast<UEditorDataStorage>(Registry->GetMutableDataStorage()))
+				if (IEditorDataStorageProvider* DataStorage = Registry->GetMutableDataStorage())
 				{
 					DataStorage->DebugPrintQueryCallbacks(Output);
 				}
@@ -46,7 +46,7 @@ FAutoConsoleCommandWithOutputDevice PrintSupportedColumnsConsoleCommand(
 	TEXT("Prints out a list of available Data Storage columns."),
 	FConsoleCommandWithOutputDeviceDelegate::CreateLambda([](FOutputDevice& Output)
 		{
-			Output.Log(TEXT("The Typed Elements Data Storage supports the following columns:"));
+			Output.Log(TEXT("The Editor Data Storage supports the following columns:"));
 			
 			UScriptStruct* FragmentTypeInfo = FMassFragment::StaticStruct();
 			UScriptStruct* TagTypeInfo = FMassTag::StaticStruct();
diff --git a/Engine/Plugins/Experimental/EditorDataStorage/Source/TedsCore/Private/TypedElementDatabase.h b/Engine/Plugins/Experimental/EditorDataStorage/Source/TedsCore/Private/TypedElementDatabase.h
--- a/Engine/Plugins/Experimental/EditorDataStorage/Source/TedsCore/Private/TypedElementDatabase.h
+++ b/Engine/Plugins/Experimental/EditorDataStorage/Source/TedsCore/Private/TypedElementDatabase.h
@@ -175,7 +175,7 @@ public:
 
 	FMassArchetypeHandle LookupArchetype(TableHandle InTableHandle) const;
 
-	void DebugPrintQueryCallbacks(FOutputDevice& Output);
+	void DebugPrintQueryCallbacks(FOutputDevice& Output) override;
 
 	static void AddReferencedObjects(UObject* InThis, FReferenceCollector& Collector);
 
diff --git a/Engine/Plugins/Experimental/EditorDataStorage/Source/TedsCore/Private/TypedElementsDataStorage.cpp b/Engine/Plugins/Experimental/EditorDataStorage/Source/TedsCore/Private/TypedElementsDataStorage.cpp
--- a/Engine/Plugins/Experimental/EditorDataStorage/Source/TedsCore/Private/TypedElementsDataStorage.cpp
+++ b/Engine/Plugins/Experimental/EditorDataStorage/Source/TedsCore/Private/TypedElementsDataStorage.cpp
@@ -4,8 +4,10 @@
 
 #include "EditorDataStorageSettings.h"
 #include "Elements/Common/TypedElementDataStorageLog.h"
+#include "Elements/Common/EditorDataStorageFeatures.h"
 #include "Elements/Framework/TypedElementRegistry.h"
 #include "Elements/Interfaces/TypedElementDataStorageFactory.h"
+#include "Features/IModularFeatures.h"
 #include "ISettingsModule.h"
 #include "MassEntityTypes.h"
 #include "Misc/CoreDelegates.h"
@@ -76,6 +78,8 @@ void FEditorDataStorageModule::StartupModule()
 		{
 			if (!bInitialized)
 			{
+				using namespace UE::Editor::DataStorage;
+
 				UE_LOG(LogEditorDataStorage, Log, TEXT("Initializing"));
 				
 				DataStorage = NewObject<UEditorDataStorage>();
@@ -91,12 +95,10 @@ void FEditorDataStorageModule::StartupModule()
 				ObjectReinstancingManager->Initialize(*DataStorage, *DataStorageCompatibility);
 
 				// Register the various DataStorage instances.
-				UTypedElementRegistry* Registry = UTypedElementRegistry::GetInstance();
-				checkf(Registry, TEXT(
-					"FEditorDataStorageModule tried to register itself, but there was no Typed Element Registry to register to."));
-				Registry->SetDataStorage(DataStorage.Get());
-				Registry->SetDataStorageCompatibility(DataStorageCompatibility.Get());
-				Registry->SetDataStorageUi(DataStorageUi.Get());
+				IModularFeatures::Get().RegisterModularFeature(StorageFeatureName, DataStorage.Get());
+				IModularFeatures::Get().RegisterModularFeature(CompatibilityFeatureName, DataStorageCompatibility.Get());
+				IModularFeatures::Get().RegisterModularFeature(UiFeatureName, DataStorageUi.Get());
+				OnEditorDataStorageFeaturesEnabled().Broadcast();
 
 				// Allow any factories to register their content.
 				TArray<UClass*> FactoryClasses;
@@ -142,15 +144,12 @@ void FEditorDataStorageModule::ShutdownModule()
 	{
 		UE_LOG(LogEditorDataStorage, Log, TEXT("Deinitializing"));
 
+		using namespace UE::Editor::DataStorage;
 		DataStorage->ResetFactories();
 
-		UTypedElementRegistry* Registry = UTypedElementRegistry::GetInstance();
-		if (Registry) // If the registry has already been destroyed there's no point in clearing the reference.
-		{
-			Registry->SetDataStorage(nullptr);
-			Registry->SetDataStorageCompatibility(nullptr);
-			Registry->SetDataStorageUi(nullptr);
-		}
+		IModularFeatures::Get().UnregisterModularFeature(UiFeatureName, DataStorageUi.Get());
+		IModularFeatures::Get().UnregisterModularFeature(CompatibilityFeatureName, DataStorageCompatibility.Get());
+		IModularFeatures::Get().UnregisterModularFeature(StorageFeatureName, DataStorage.Get());
 
 		if (UObjectInitialized())
 		{
diff --git a/Engine/Source/Runtime/TypedElementFramework/Private/Elements/Common/EditorDataStorageFeatures.cpp b/Engine/Source/Runtime/TypedElementFramework/Private/Elements/Common/EditorDataStorageFeatures.cpp
new file mode 100644
--- /dev/null
+++ b/Engine/Source/Runtime/TypedElementFramework/Private/Elements/Common/EditorDataStorageFeatures.cpp
@@ -0,0 +1,23 @@
+// Copyright Epic Games, Inc. All Rights Reserved.
+
+#include "Elements/Common/EditorDataStorageFeatures.h"
+
+namespace UE::Editor::DataStorage
+{
+	const FName StorageFeatureName = TEXT("EditorDataStorage");
+	const FName CompatibilityFeatureName = TEXT("EditorDataStorageCompatibility");
+	const FName UiFeatureName = TEXT("EditorDataStorageUi");
+	
+	FSimpleMulticastDelegate& OnEditorDataStorageFeaturesEnabled()
+	{
+		static FSimpleMulticastDelegate OnTedsFeaturesEnabled;
+		return OnTedsFeaturesEnabled;
+	}
+
+	bool AreEditorDataStorageFeaturesEnabled()
+	{
+		return IModularFeatures::Get().IsModularFeatureAvailable(StorageFeatureName)
+			&& IModularFeatures::Get().IsModularFeatureAvailable(CompatibilityFeatureName)
+			&& IModularFeatures::Get().IsModularFeatureAvailable(UiFeatureName);
+	}
+} // namespace UE::Editor::DataStorage
diff --git a/Engine/Source/Runtime/TypedElementFramework/Private/Elements/Framework/TypedElementRegistry.cpp b/Engine/Source/Runtime/TypedElementFramework/Private/Elements/Framework/TypedElementRegistry.cpp
--- a/Engine/Source/Runtime/TypedElementFramework/Private/Elements/Framework/TypedElementRegistry.cpp
+++ b/Engine/Source/Runtime/TypedElementFramework/Private/Elements/Framework/TypedElementRegistry.cpp
@@ -2,6 +2,7 @@
 
 #include "Elements/Framework/TypedElementRegistry.h"
 
+#include "Elements/Common/EditorDataStorageFeatures.h"
 #include "HAL/IConsoleManager.h"
 #include "Misc/CoreDelegates.h"
 #include "Misc/ScopeLock.h"
@@ -66,55 +67,48 @@ UTypedElementRegistry* UTypedElementRegistry::GetInstance()
 
 IEditorDataStorageProvider* UTypedElementRegistry::GetMutableDataStorage()
 {
-	return DataStorage;
+	using namespace UE::Editor::DataStorage;
+	return GetMutableDataStorageFeature<IEditorDataStorageProvider>(StorageFeatureName);
 }
 
 const IEditorDataStorageProvider* UTypedElementRegistry::GetDataStorage() const
 {
-	return DataStorage;
-}
-
-void UTypedElementRegistry::SetDataStorage(IEditorDataStorageProvider* Storage)
-{
-	DataStorage = Storage;
-	CallDataStorageInterfacesSetDelegateIfNeeded();
+	using namespace UE::Editor::DataStorage;
+	return GetDataStorageFeature<IEditorDataStorageProvider>(StorageFeatureName);
 }
 
 IEditorDataStorageCompatibilityProvider* UTypedElementRegistry::GetMutableDataStorageCompatibility()
 {
-	return DataStorageCompatibility;
+	using namespace UE::Editor::DataStorage;
+	return GetMutableDataStorageFeature<IEditorDataStorageCompatibilityProvider>(CompatibilityFeatureName);
 }
 
 const IEditorDataStorageCompatibilityProvider* UTypedElementRegistry::GetDataStorageCompatibility() const
 {
-	return DataStorageCompatibility;
-}
-
-void UTypedElementRegistry::SetDataStorageCompatibility(IEditorDataStorageCompatibilityProvider* Storage)
-{
-	DataStorageCompatibility = Storage;
-	CallDataStorageInterfacesSetDelegateIfNeeded();
+	using namespace UE::Editor::DataStorage;
+	return GetDataStorageFeature<IEditorDataStorageCompatibilityProvider>(CompatibilityFeatureName);
 }
 
 IEditorDataStorageUiProvider* UTypedElementRegistry::GetMutableDataStorageUi()
 {
-	return DataStorageUi;
+	using namespace UE::Editor::DataStorage;
+	return GetMutableDataStorageFeature<IEditorDataStorageUiProvider>(UiFeatureName);
 }
 
 const IEditorDataStorageUiProvider* UTypedElementRegistry::GetDataStorageUi() const
 {
-	return DataStorageUi;
+	using namespace UE::Editor::DataStorage;
+	return GetDataStorageFeature<IEditorDataStorageUiProvider>(UiFeatureName);
 }
 
-void UTypedElementRegistry::SetDataStorageUi(IEditorDataStorageUiProvider* Storage)
+bool UTypedElementRegistry::AreDataStorageInterfacesSet() const
 {
-	DataStorageUi = Storage;
-	CallDataStorageInterfacesSetDelegateIfNeeded();
+	return UE::Editor::DataStorage::AreEditorDataStorageFeaturesEnabled();
 }
 
-bool UTypedElementRegistry::AreDataStorageInterfacesSet() const
+FSimpleMulticastDelegate& UTypedElementRegistry::OnDataStorageInterfacesSet()
 {
-	return DataStorage && DataStorageCompatibility && DataStorageUi;
+	return UE::Editor::DataStorage::OnEditorDataStorageFeaturesEnabled();
 }
 
 void UTypedElementRegistry::FinishDestroy()
@@ -406,11 +400,3 @@ void UTypedElementRegistry::OnPostGarbageCollect()
 		ProcessDeferredElementsToDestroy();
 	}
 }
-
-void UTypedElementRegistry::CallDataStorageInterfacesSetDelegateIfNeeded()
-{
-	if (OnDataStorageInterfacesSetDelegate.IsBound() && AreDataStorageInterfacesSet())
-	{
-		OnDataStorageInterfacesSetDelegate.Broadcast();
-	}
-}
diff --git a/Engine/Source/Runtime/TypedElementFramework/Public/Elements/Common/EditorDataStorageFeatures.h b/Engine/Source/Runtime/TypedElementFramework/Public/Elements/Common/EditorDataStorageFeatures.h
new file mode 100644
--- /dev/null
+++ b/Engine/Source/Runtime/TypedElementFramework/Public/Elements/Common/EditorDataStorageFeatures.h
@@ -0,0 +1,38 @@
+// Copyright Epic Games, Inc. All Rights Reserved.
+
+#pragma once
+
+#include "Delegates/Delegate.h"
+#include "Features/IModularFeatures.h"
+#include "UObject/NameTypes.h"
+#include "GenericPlatform/GenericPlatform.h"
+
+namespace UE::Editor::DataStorage
+{
+	TYPEDELEMENTFRAMEWORK_API const extern FName StorageFeatureName;
+	TYPEDELEMENTFRAMEWORK_API const extern FName CompatibilityFeatureName;
+	TYPEDELEMENTFRAMEWORK_API const extern FName UiFeatureName;
+
+	TYPEDELEMENTFRAMEWORK_API FSimpleMulticastDelegate& OnEditorDataStorageFeaturesEnabled();
+	bool AreEditorDataStorageFeaturesEnabled();
+
+	template<typename T>
+	const T* GetDataStorageFeature(const FName InName)
+	{
+		if (IModularFeatures::Get().IsModularFeatureAvailable(InName))
+		{
+			return &IModularFeatures::Get().GetModularFeature<T>(InName);
+		}
+		return nullptr;
+	}
+
+	template<typename T>
+	T* GetMutableDataStorageFeature(const FName InName)
+	{
+		if (IModularFeatures::Get().IsModularFeatureAvailable(InName))
+		{
+			return &IModularFeatures::Get().GetModularFeature<T>(InName);
+		}
+		return nullptr;
+	}
+} // namespace UE::Editor::DataStorage
diff --git a/Engine/Source/Runtime/TypedElementFramework/Public/Elements/Framework/TypedElementRegistry.h b/Engine/Source/Runtime/TypedElementFramework/Public/Elements/Framework/TypedElementRegistry.h
--- a/Engine/Source/Runtime/TypedElementFramework/Public/Elements/Framework/TypedElementRegistry.h
+++ b/Engine/Source/Runtime/TypedElementFramework/Public/Elements/Framework/TypedElementRegistry.h
@@ -77,28 +77,26 @@ public:
 	UFUNCTION(BlueprintPure, DisplayName="Get Default Typed Element Registry", Category = "TypedElementFramework|Registry", meta=(ScriptName="GetDefaultTypedElementRegistry"))
 	static TYPEDELEMENTFRAMEWORK_API UTypedElementRegistry* GetInstance();
 
+	//UE_DEPRECATED(5.5, "Function has moved to modular features")
 	TYPEDELEMENTFRAMEWORK_API IEditorDataStorageProvider* GetMutableDataStorage();
+	//UE_DEPRECATED(5.5, "Function has moved to modular features")
 	TYPEDELEMENTFRAMEWORK_API const IEditorDataStorageProvider* GetDataStorage() const;
-	TYPEDELEMENTFRAMEWORK_API void SetDataStorage(IEditorDataStorageProvider* Storage);
 
+	//UE_DEPRECATED(5.5, "Function has moved to modular features")
 	TYPEDELEMENTFRAMEWORK_API IEditorDataStorageCompatibilityProvider* GetMutableDataStorageCompatibility();
+	//UE_DEPRECATED(5.5, "Function has moved to modular features")
 	TYPEDELEMENTFRAMEWORK_API const IEditorDataStorageCompatibilityProvider* GetDataStorageCompatibility() const;
-	TYPEDELEMENTFRAMEWORK_API void SetDataStorageCompatibility(IEditorDataStorageCompatibilityProvider* Storage);
 
+	//UE_DEPRECATED(5.5, "Function has moved to modular features")
 	TYPEDELEMENTFRAMEWORK_API IEditorDataStorageUiProvider* GetMutableDataStorageUi();
+	//UE_DEPRECATED(5.5, "Function has moved to modular features")
 	TYPEDELEMENTFRAMEWORK_API const IEditorDataStorageUiProvider* GetDataStorageUi() const;
-	TYPEDELEMENTFRAMEWORK_API void SetDataStorageUi(IEditorDataStorageUiProvider* Storage);
 
+	//UE_DEPRECATED(5.5, "Use UE::Editor::DataStorage::AreDataStorageInterfacesSet")
 	TYPEDELEMENTFRAMEWORK_API bool AreDataStorageInterfacesSet() const;
 
-	/**
-	 * Event fired when all Data Storage Interfaces have been set. 
-	 */
-	DECLARE_MULTICAST_DELEGATE(FOnDataStorageInterfacesSet);
-	FOnDataStorageInterfacesSet& OnDataStorageInterfacesSet()
-	{
-		return OnDataStorageInterfacesSetDelegate;
-	}
+	//UE_DEPRECATED(5.5, "Use UE::Editor::DataStorage::OnAllStorageFeaturesEnabled")
+	TYPEDELEMENTFRAMEWORK_API FSimpleMulticastDelegate& OnDataStorageInterfacesSet();
 
 	/**
 	 * Event fired when references to one element should be replaced with a reference to a different element.
@@ -682,8 +680,6 @@ private:
 		--DisableElementDestructionOnGCCount;
 	}
 
-	TYPEDELEMENTFRAMEWORK_API void CallDataStorageInterfacesSetDelegateIfNeeded();
-
 	mutable FRWLock RegisteredElementTypesRW;
 	TUniquePtr<FRegisteredElementType> RegisteredElementTypes[TypedHandleMaxTypeId - 1];
 	TSortedMap<FName, FTypedHandleTypeId, FDefaultAllocator, FNameFastLess> RegisteredElementTypesNameToId;
@@ -699,9 +695,4 @@ private:
 	FOnElementReplaced OnElementReplacedDelegate;
 	FOnElementUpdated OnElementUpdatedDelegate;
 	FSimpleMulticastDelegate OnProcessingDeferredElementsToDestroyDelegate;
-	
-	FOnDataStorageInterfacesSet OnDataStorageInterfacesSetDelegate;
-	IEditorDataStorageProvider* DataStorage = nullptr;
-	IEditorDataStorageCompatibilityProvider* DataStorageCompatibility = nullptr;
-	IEditorDataStorageUiProvider* DataStorageUi = nullptr;
 };
diff --git a/Engine/Source/Runtime/TypedElementFramework/Public/Elements/Interfaces/TypedElementDataStorageCompatibilityInterface.h b/Engine/Source/Runtime/TypedElementFramework/Public/Elements/Interfaces/TypedElementDataStorageCompatibilityInterface.h
--- a/Engine/Source/Runtime/TypedElementFramework/Public/Elements/Interfaces/TypedElementDataStorageCompatibilityInterface.h
+++ b/Engine/Source/Runtime/TypedElementFramework/Public/Elements/Interfaces/TypedElementDataStorageCompatibilityInterface.h
@@ -3,6 +3,7 @@
 #pragma once
 
 #include "Elements/Interfaces/TypedElementDataStorageInterface.h"
+#include "Features/IModularFeature.h"
 #include "Templates/Function.h"
 #include "UObject/Interface.h"
 #include "UObject/ObjectKey.h"
@@ -11,22 +12,12 @@
 #include "UObject/StrongObjectPtr.h"
 #include "UObject/WeakObjectPtrTemplates.h"
 
-#include "TypedElementDataStorageCompatibilityInterface.generated.h"
-
-UINTERFACE(MinimalAPI)
-class UEditorDataStorageCompatibilityProvider : public UInterface
-{
-	GENERATED_BODY()
-};
-
 /**
  * Interface to provide compatibility with existing systems that don't directly
  * support the data storage.
  */
-class IEditorDataStorageCompatibilityProvider
+class IEditorDataStorageCompatibilityProvider : public IModularFeature
 {
-	GENERATED_BODY()
-
 public:
 	using ObjectRegistrationFilter = TFunction<bool(const IEditorDataStorageCompatibilityProvider&, const UObject*)>;
 	using ObjectToRowDealiaser = TFunction<UE::Editor::DataStorage::RowHandle(const IEditorDataStorageCompatibilityProvider&, const UObject*)>;
diff --git a/Engine/Source/Runtime/TypedElementFramework/Public/Elements/Interfaces/TypedElementDataStorageInterface.h b/Engine/Source/Runtime/TypedElementFramework/Public/Elements/Interfaces/TypedElementDataStorageInterface.h
--- a/Engine/Source/Runtime/TypedElementFramework/Public/Elements/Interfaces/TypedElementDataStorageInterface.h
+++ b/Engine/Source/Runtime/TypedElementFramework/Public/Elements/Interfaces/TypedElementDataStorageInterface.h
@@ -13,6 +13,7 @@
 #include "Elements/Common/TypedElementQueryTypes.h"
 #include "Elements/Framework/TypedElementColumnUtils.h"
 #include "Elements/Interfaces/TypedElementQueryStorageInterfaces.h"
+#include "Features/IModularFeature.h"
 #include "Math/NumericLimits.h"
 #include "Misc/EnumClassFlags.h"
 #include "Templates/Function.h"
@@ -20,8 +21,6 @@
 #include "UObject/NameTypes.h"
 #include "UObject/ObjectMacros.h"
 
-#include "TypedElementDataStorageInterface.generated.h"
-
 class UClass;
 class USubsystem;
 class UScriptStruct;
@@ -31,12 +30,6 @@ using FTypedElementOnDataStorageCreation = FSimpleMulticastDelegate;
 using FTypedElementOnDataStorageDestruction = FSimpleMulticastDelegate;
 using FTypedElementOnDataStorageUpdate = FSimpleMulticastDelegate;
 
-UINTERFACE(MinimalAPI)
-class UEditorDataStorageProvider : public UInterface
-{
-	GENERATED_BODY()
-};
-
 /**
  * Convenience structure that can be used to pass a list of columns to functions that don't
  * have an dedicate templated version that takes a column list directly, for instance when
@@ -51,10 +44,8 @@ struct TTypedElementColumnTypeList
 	operator TConstArrayView<const UScriptStruct*>() const { return ColumnTypes; }
 };
 
-class IEditorDataStorageProvider
+class IEditorDataStorageProvider : public IModularFeature
 {
-	GENERATED_BODY()
-
 	using RowCreationCallbackRef = UE::Editor::DataStorage::RowCreationCallbackRef;
 	using ColumnCreationCallbackRef = UE::Editor::DataStorage::ColumnCreationCallbackRef;
 	using ColumnListCallbackRef = UE::Editor::DataStorage::ColumnListCallbackRef;
@@ -318,6 +309,11 @@ public:
 	 * The TemplateType may be a typed derived from either UE::Editor::DataStorage::FColumn or UE::Editor::DataStorage::FTag
 	 */
 	virtual const UScriptStruct* GenerateDynamicColumn(const UE::Editor::DataStorage::FDynamicColumnDescription& Description) = 0;
+
+	/**
+	 * Outputs the registered query callbacks to the given output device for debugging purposes.
+	 */
+	virtual void DebugPrintQueryCallbacks(FOutputDevice& Output) = 0;
 	
 	/**
 	 * @section Query
diff --git a/Engine/Source/Runtime/TypedElementFramework/Public/Elements/Interfaces/TypedElementDataStorageUiInterface.h b/Engine/Source/Runtime/TypedElementFramework/Public/Elements/Interfaces/TypedElementDataStorageUiInterface.h
--- a/Engine/Source/Runtime/TypedElementFramework/Public/Elements/Interfaces/TypedElementDataStorageUiInterface.h
+++ b/Engine/Source/Runtime/TypedElementFramework/Public/Elements/Interfaces/TypedElementDataStorageUiInterface.h
@@ -7,6 +7,7 @@
 #include "Elements/Interfaces/TypedElementDataStorageInterface.h"
 #include "Elements/Framework/TypedElementMetaData.h"
 #include "Elements/Framework/TypedElementQueryBuilder.h"
+#include "Features/IModularFeature.h"
 #include "Templates/UnrealTypeTraits.h"
 #include "Templates/SharedPointer.h"
 #include "UObject/Interface.h"
@@ -202,16 +203,8 @@ struct TStructOpsTypeTraits<FTypedElementWidgetConstructor> : public TStructOpsT
 	};
 };
 
-UINTERFACE(MinimalAPI)
-class UEditorDataStorageUiProvider : public UInterface
+class IEditorDataStorageUiProvider : public IModularFeature
 {
-	GENERATED_BODY()
-};
-
-class IEditorDataStorageUiProvider
-{
-	GENERATED_BODY()
-
 	using RowHandle = UE::Editor::DataStorage::RowHandle;
 
 public:
