Added support for binding to formatted text in TEDS UI.

This change adds the ability to bind a Slate widget to a text attribute that uses formatted text and and pulls arguments from TEDS columns.

diff --git a/Engine/Plugins/Experimental/EditorDataStorageFeatures/Source/TedsAssetData/Private/Widgets/AssetDataLabelWidget.cpp b/Engine/Plugins/Experimental/EditorDataStorageFeatures/Source/TedsAssetData/Private/Widgets/AssetDataLabelWidget.cpp
--- a/Engine/Plugins/Experimental/EditorDataStorageFeatures/Source/TedsAssetData/Private/Widgets/AssetDataLabelWidget.cpp
+++ b/Engine/Plugins/Experimental/EditorDataStorageFeatures/Source/TedsAssetData/Private/Widgets/AssetDataLabelWidget.cpp
@@ -84,10 +84,18 @@ TSharedPtr<SWidget> FAssetDataLabelWidgetConstructor::CreateWidget(IEditorDataSt
 			[
 				SNew(STextBlock)
 				.Text(Binder.BindText(&FNameColumn::Name))
-				.ToolTipText_Lambda([DataStorage, DataRow = TargetRow]()
-					{
-						return ConstructToolTip(DataStorage, DataRow);
-					})
+				.ToolTipText(Binder.BindTextFormat(
+								LOCTEXT("AssetLabelTooltip", 
+									"{Name}\n\nVirtual path: {VirtualPath}\n  Asset path: {AssetPath}\n  Verse path: {VersePath}"))
+								.Arg(TEXT("Name"), &FNameColumn::Name)
+								.Arg(TEXT("VirtualPath"), &FVirtualPathColumn_Experimental::VirtualPath, LOCTEXT("PathNotSet", "<not set>"))
+								.Arg(TEXT("AssetPath"), &FAssetPathColumn_Experimental::Path, LOCTEXT("PathNotSet", "<not set>"))
+								.Arg(TEXT("VersePath"), &FVersePathColumn::VersePath, 
+									[](const UE::Core::FVersePath& Path) 
+									{
+										return FText::FromStringView(Path.AsStringView());
+									}, 
+									LOCTEXT("PathNotSet", "<not set>")))
 			];
 	}
 	else
@@ -95,35 +103,4 @@ TSharedPtr<SWidget> FAssetDataLabelWidgetConstructor::CreateWidget(IEditorDataSt
 		return SNullWidget::NullWidget;
 	}
 }
-
-FText FAssetDataLabelWidgetConstructor::ConstructToolTip(
-	IEditorDataStorageProvider* DataStorage, UE::Editor::DataStorage::RowHandle DataRow)
-{
-	TStringBuilder<1024> ToolTip;
-	if (FNameColumn* ItemName = DataStorage->GetColumn<FNameColumn>(DataRow))
-	{
-		ToolTip.Append(ItemName->Name.ToString());
-		ToolTip.AppendChar(TEXT('\n'));
-	}
-	if (FAssetPathColumn_Experimental* AssetPath = DataStorage->GetColumn<FAssetPathColumn_Experimental>(DataRow))
-	{
-		ToolTip.AppendChar(TEXT('\n'));
-		ToolTip.Append(LOCTEXT("AssetPath", "Asset path: ").ToString());
-		ToolTip.Append(AssetPath->Path.ToString());
-	}
-	if (FVersePathColumn* VersePath = DataStorage->GetColumn<FVersePathColumn>(DataRow))
-	{
-		ToolTip.AppendChar(TEXT('\n'));
-		ToolTip.Append(LOCTEXT("VersePath", "Verse path: ").ToString());
-		ToolTip.Append(VersePath->VersePath.AsStringView());
-	}
-	if (FVirtualPathColumn_Experimental* VirtualPath = DataStorage->GetColumn<FVirtualPathColumn_Experimental>(DataRow))
-	{
-		ToolTip.AppendChar(TEXT('\n'));
-		ToolTip.Append(LOCTEXT("VirtualPath", "Virtual path: ").ToString());
-		ToolTip.Append(VirtualPath->VirtualPath.ToString());
-	}
-	return FText::FromString(ToolTip.ToString());
-}
-
 #undef LOCTEXT_NAMESPACE
\ No newline at end of file
diff --git a/Engine/Plugins/Experimental/EditorDataStorageFeatures/Source/TedsAssetData/Public/Widgets/AssetDataLabelWidget.h b/Engine/Plugins/Experimental/EditorDataStorageFeatures/Source/TedsAssetData/Public/Widgets/AssetDataLabelWidget.h
--- a/Engine/Plugins/Experimental/EditorDataStorageFeatures/Source/TedsAssetData/Public/Widgets/AssetDataLabelWidget.h
+++ b/Engine/Plugins/Experimental/EditorDataStorageFeatures/Source/TedsAssetData/Public/Widgets/AssetDataLabelWidget.h
@@ -41,8 +41,4 @@ public:
 		UE::Editor::DataStorage::RowHandle TargetRow,
 		UE::Editor::DataStorage::RowHandle WidgetRow, 
 		const UE::Editor::DataStorage::FMetaDataView& Arguments) override;
-
-protected:
-	
-	static FText ConstructToolTip(IEditorDataStorageProvider* DataStorage, UE::Editor::DataStorage::RowHandle DataRow);
 };
\ No newline at end of file
diff --git a/Engine/Source/Runtime/TypedElementFramework/Private/Elements/Framework/TypedElementAttributeBinding.cpp b/Engine/Source/Runtime/TypedElementFramework/Private/Elements/Framework/TypedElementAttributeBinding.cpp
--- a/Engine/Source/Runtime/TypedElementFramework/Private/Elements/Framework/TypedElementAttributeBinding.cpp
+++ b/Engine/Source/Runtime/TypedElementFramework/Private/Elements/Framework/TypedElementAttributeBinding.cpp
@@ -7,13 +7,8 @@
 namespace UE::Editor::DataStorage
 {
 	FAttributeBinder::FAttributeBinder(RowHandle InTargetRow)
-		: TargetRow(InTargetRow)
+		: FAttributeBinder(InTargetRow, GetMutableDataStorageFeature<IEditorDataStorageProvider>(StorageFeatureName))
 	{
-		if (ensureMsgf(AreEditorDataStorageFeaturesEnabled(),
-				TEXT("The Editor Data Storage plugin needs to be enabled to use attribute bindings.")))
-		{
-			DataStorage = GetMutableDataStorageFeature<IEditorDataStorageProvider>(StorageFeatureName);
-		}
 	}
 
 	FAttributeBinder::FAttributeBinder(
@@ -23,4 +18,9 @@ namespace UE::Editor::DataStorage
 	{
 		ensureMsgf(DataStorage, TEXT("The Editor Data Storage plugin needs to be enabled to use attribute bindings."));
 	}
+
+	FTextAttributeFormatted FAttributeBinder::BindTextFormat(FTextFormat Format) const
+	{
+		return FTextAttributeFormatted(MoveTemp(Format), TargetRow, DataStorage);
+	}
 } // namespace UE::Editor::DataStorage
diff --git a/Engine/Source/Runtime/TypedElementFramework/Private/Elements/Framework/TypedElementtAttributeBindingText.cpp b/Engine/Source/Runtime/TypedElementFramework/Private/Elements/Framework/TypedElementtAttributeBindingText.cpp
new file mode 100644
--- /dev/null
+++ b/Engine/Source/Runtime/TypedElementFramework/Private/Elements/Framework/TypedElementtAttributeBindingText.cpp
@@ -0,0 +1,67 @@
+// Copyright Epic Games, Inc. All Rights Reserved.
+
+#pragma once
+
+#include "Elements/Framework/TypedElementAttributeBindingText.h"
+
+namespace UE::Editor::DataStorage
+{
+	FTextAttributeFormatted::FTextAttributeFormatted(
+		FTextFormat InFormat, RowHandle InTargetRow, IEditorDataStorageProvider* InDataStorage)
+		: Format(MoveTemp(InFormat))
+		, TargetRow(InTargetRow)
+		, DataStorage(InDataStorage)
+	{
+	}
+	
+	FTextAttributeFormatted& FTextAttributeFormatted::Arg(FString Name, FFormatArgumentValue Value)
+	{
+		NamedProperties.Add(MoveTemp(Name), FPropertyInfo{ .Default = MoveTemp(Value) });
+		return *this;
+	}
+
+	FTextAttributeFormatted::operator TAttribute<FText>()
+	{
+		using namespace Private; 
+		
+		return TAttribute<FText>::CreateLambda(
+			[
+				Format = MoveTemp(Format),
+				NamedProperties = MoveTemp(NamedProperties),
+				Storage = DataStorage,
+				Row = TargetRow
+			]()
+			{
+				FFormatNamedArguments NamedArguments;
+				NamedArguments.Reserve(NamedProperties.Num());
+				for (const TPair<FString, FPropertyInfo>& Property : NamedProperties)
+				{
+					if (Property.Value.Property.IsBound())
+					{
+						const UScriptStruct* ColumnType = Property.Value.Property.GetObjectTypeInfo();
+
+						if(Property.Value.DynamicColumnIdentifier != NAME_None)
+						{
+							ColumnType = Storage->GenerateDynamicColumn(FDynamicColumnDescription{
+								.TemplateType = ColumnType,
+								.Identifier = Property.Value.DynamicColumnIdentifier
+							});
+						}
+
+						if(ColumnType)
+						{
+							if (void* Column = Storage->GetColumnData(Row, ColumnType))
+							{
+								// We use the original column type (Property.GetObjectTypeInfo()) instead of the dynamic ColumnType
+								// so that the type info in the property matches
+								NamedArguments.Add(Property.Key, Property.Value.Property.Get(Column, Property.Value.Property.GetObjectTypeInfo()));
+								continue;
+							}
+						}
+					}
+					NamedArguments.Add(Property.Key, Property.Value.Default);
+				}
+				return FText::Format(Format, NamedArguments);
+			});
+	}
+}
diff --git a/Engine/Source/Runtime/TypedElementFramework/Public/Elements/Framework/TypedElementAttributeBinding.h b/Engine/Source/Runtime/TypedElementFramework/Public/Elements/Framework/TypedElementAttributeBinding.h
--- a/Engine/Source/Runtime/TypedElementFramework/Public/Elements/Framework/TypedElementAttributeBinding.h
+++ b/Engine/Source/Runtime/TypedElementFramework/Public/Elements/Framework/TypedElementAttributeBinding.h
@@ -4,16 +4,14 @@
 
 #include "Elements/Common/TypedElementHandles.h"
 #include "Elements/Interfaces/TypedElementDataStorageInterface.h"
+#include "Internationalization/Text.h"
 #include "Misc/Attribute.h"
+#include "TypedElementAttributeBindingText.h"
 
 class IEditorDataStorageProvider;
 
 namespace UE::Editor::DataStorage
 {
-	// Concept for a conversion function used by the attribute binder to bind a column data member to an attribute of a different type
-	template <typename FunctionType, typename ArgumentType>
-	concept AttributeBinderInvocable = std::is_invocable_v<std::decay_t<FunctionType>, const ArgumentType&>;
-
 	/**
 	 * Builder class that can be used as a shorthand to bind data inside a TEDS row, column pair to a TAttribute so the attribute updates if the data
 	 * in the column is changed.
@@ -69,8 +67,7 @@ namespace UE::Editor::DataStorage
 		 *                                 ));
 		 */
 		template <typename AttributeType, typename DataType, TDataColumnType ColumnType>
-		TAttribute<AttributeType> BindData(DataType ColumnType::* InVariable, const TFunction<AttributeType(const DataType&)>& InConverter,
-			const DataType& InDefaultValue = DataType(), const FName& InIdentifier = NAME_None);
+		TAttribute<AttributeType> BindData(DataType ColumnType::* InVariable, TFunction<AttributeType(const DataType&)> InConverter, const DataType& InDefaultValue = DataType(), const FName& InIdentifier = NAME_None);
 
 		/**
 		 * Overload for the conversion binder to accept lambdas instead of TFunctions
@@ -82,13 +79,13 @@ namespace UE::Editor::DataStorage
 		 * @return A TAttribute bound to the row, column pair specified
 		 */
 		template <typename DataType, TDataColumnType ColumnType, typename FunctionType>
-			requires AttributeBinderInvocable<FunctionType, DataType>
+			requires Private::AttributeBinderInvocable<FunctionType, DataType>
 		auto BindData(DataType ColumnType::* InVariable, FunctionType InConverter, const DataType& InDefaultValue = DataType(),
 			const FName& InIdentifier = NAME_None);
 
 		/**
-		 * Bind a delegate inside a Teds column to a SLATE_EVENT macro on a widget
-		 * @param InVariable The delegate inside the Teds column
+		 * Bind a delegate inside a TEDS column to a SLATE_EVENT macro on a widget
+		 * @param InVariable The delegate inside the TEDS column
 		 * @param InIdentifier The identifier for this column if it is a dynamic column, NAME_None if it is not a dynamic column
 		 * @return A delegate that can be provided to an event on a slate widget
 		 */
@@ -97,7 +94,7 @@ namespace UE::Editor::DataStorage
 			const FName& InIdentifier = NAME_None);
 
 		/**
-		 * Directly bind an FString member in a Teds column to an FText attribute as a shortcut
+		 * Directly bind an FString member in a TEDS column to an FText attribute as a shortcut
 		 * @param InFStringVariable The FString variable
 		 * @param InIdentifier The identifier for this column if it is a dynamic column, NAME_None if it is not a dynamic column
 		 * @return A delegate that can be provided to a text widget in Slate (e.g STextBlock)
@@ -106,13 +103,34 @@ namespace UE::Editor::DataStorage
 		TAttribute<FText> BindText(FString ColumnType::* InFStringVariable, const FName& InIdentifier = NAME_None);
 
 		/**
-		 * Directly bind an FName member in a Teds column to an FText attribute as a shortcut
+		 * Directly bind an FName member in a TEDS column to an FText attribute as a shortcut
 		 * @param InFNameVariable The FName variable
 		 * @param InIdentifier The identifier for this column if it is a dynamic column, NAME_None if it is not a dynamic column
 		 * @return A delegate that can be provided to a text widget in Slate (e.g STextBlock)
 		 */
 		template <TDataColumnType ColumnType>
 		TAttribute<FText> BindText(FName ColumnType::* InFNameVariable, const FName& InIdentifier = NAME_None);
+
+		/**
+		 * Composite a FText attribute using a format string and named arguments bound to TEDS columns.
+		 * Arguments are passed in using the .Arg(...) function. Each argument starts with the name of the
+		 * argument in the format string followed by one of the following options:
+		 *		- A column variable that's a string (FText, FString or FName).
+		 *		- A column variable with a converter to a string.
+		 *		- A direct value that's supported by FFormatArgumentValue (FText and numbers).
+		 * The final parameter for an argument is an optional default that optionally takes a FFormatArgumentValue value.
+		 * @param Format The format to use composite a string together.
+		 * @return A delegate that can be provided to a text widget in Slate (e.g STextBlock)
+		 * 
+		 * Example:
+		 * TAttribute<FText> Attribute(Binder.BindTextFormat(
+		 *		LOCTEXT("Format", "{Label}: {Value1}, {Value2}, {Value3}, {Value4}")
+		 *		.Arg(TEXT("Label"), &FTypedElementLabelColumn::Label)
+		 *		.Arg(TEXT("Value1"), &FValueColumn::Integer, [](const int32& Value) { return FText::AsNumber(Value); }, 42)
+		 *		.Arg(TEXT("Value2"), &FTextColumn::Text, LOCTEXT("Default", "<no value>"))
+		 *		.Arg(TEXT("Value3"), 42)
+		 */
+		TYPEDELEMENTFRAMEWORK_API FTextAttributeFormatted BindTextFormat(FTextFormat Format) const;
 		
 	private:
 
diff --git a/Engine/Source/Runtime/TypedElementFramework/Public/Elements/Framework/TypedElementAttributeBinding.inl b/Engine/Source/Runtime/TypedElementFramework/Public/Elements/Framework/TypedElementAttributeBinding.inl
--- a/Engine/Source/Runtime/TypedElementFramework/Public/Elements/Framework/TypedElementAttributeBinding.inl
+++ b/Engine/Source/Runtime/TypedElementFramework/Public/Elements/Framework/TypedElementAttributeBinding.inl
@@ -3,155 +3,27 @@
 #pragma once
 
 #include "TypedElementRegistry.h"
+#include "TypedElementAttributeBindingProperty.h"
 
 namespace UE::Editor::DataStorage
 {
-	namespace AttributeBinding::Private
-	{
-		template <TDataColumnType ColumnType>
-		ColumnType* GetColumn(IEditorDataStorageProvider* DataStorage, RowHandle Row, const FName& InIdentifier)
-		{
-			if(InIdentifier != NAME_None)
-			{
-				if(ColumnType* DynamicColumn = DataStorage->GetColumn<ColumnType>(Row, InIdentifier))
-				{
-					return DynamicColumn;
-				}
-
-				// The dynamic column wasn't found on this row
-				return nullptr;
-			}
-
-			// If we weren't provided a dynamic column identifier, we are looking for the regular non-dynamic column
-			return DataStorage->GetColumn<ColumnType>(Row);
-		}
-		
-	}
-	
-	// A property that can be directly be accessed from an object
-	template <typename PropertyType>
-	struct DirectProperty final
-	{
-		// Offset to a particular data member inside the object
-		size_t Offset;
-
-		PropertyType* Get(void* Object) const
-		{ 
-			return reinterpret_cast<PropertyType*>((reinterpret_cast<char*>(Object) + Offset));
-		}
-		const PropertyType* Get(const void* Object) const
-		{ 
-			return reinterpret_cast<const PropertyType*>((reinterpret_cast<const char*>(Object) + Offset)); 
-		}
-	};
+	//
+	// FAttributeBinder
+	//
 
-	// A property that goes through a conversion function before being accessed from the object
-	template <typename PropertyType>
-	struct ConvertibleProperty final
-	{
-		// Conversion function
-		TFunction<PropertyType(const void*)> Converter;
-
-		// Cache to avoid a copy when returning if not necessary
-		mutable PropertyType Cache;
-
-		PropertyType* Get(void* Object) const
-		{ 
-			Cache = Converter(Object);
-			return &Cache;
-		}
-		const PropertyType* Get(const void* Object) const
-		{ 
-			Cache = Converter(Object);
-			return &Cache;
-		}
-	};
-
-	template <typename PropertyType>
-	class Property
-	{
-	public:
-		
-		// Bind this property directly
-		template <TDataColumnType ObjectType>
-		void Bind(PropertyType ObjectType::* Variable)
-		{
-			DirectProperty<PropertyType> Result;
-			Result.Offset = reinterpret_cast<size_t>(&(reinterpret_cast<const ObjectType*>(0)->*Variable));
-			
-			ObjectTypeInfo = ObjectType::StaticStruct();
-			InternalProperty = InternalPropertyType(TInPlaceType<DirectProperty<PropertyType>>(), Result);;
-		}
-
-		// Bind this property using a conversion function
-		template <typename InputType, TDataColumnType ObjectType>
-		void Bind(InputType ObjectType::* Variable, TFunction<PropertyType(const InputType&)> Converter)
-		{
-			ConvertibleProperty<PropertyType> Result;
-			Result.Converter = [Converter, Variable](const void* Object)
-			{
-				return Converter(reinterpret_cast<const ObjectType*>(Object)->*Variable);
-			};
-			
-			ObjectTypeInfo = ObjectType::StaticStruct();
-			InternalProperty = InternalPropertyType(TInPlaceType<ConvertibleProperty<PropertyType>>(), Result);;
-		}
-
-		// Get the bound property for the specified object
-		template <TDataColumnType ObjectType>
-		PropertyType& Get(ObjectType& Object) const
-		{
-			return Get(&Object, ObjectType::StaticStruct());
-		}
-
-		// Get the bound property for the specified object
-		template <TDataColumnType ObjectType>
-		const PropertyType& Get(const ObjectType& Object) const
-		{
-			return Get(&Object, ObjectType::StaticStruct());
-		}
-		
-		// Get the bound property for this specified object ptr by providing type information about the object
-		PropertyType& Get(void* Object, const UScriptStruct* Type) const
-		{
-			ensureMsgf(Type == ObjectTypeInfo, TEXT("Provided object type (%s) did not match bound object type (%s)."),
-				*Type->GetFName().ToString(), *ObjectTypeInfo->GetFName().ToString());
-
-			return *Visit([Object](auto&& Prop) { return Prop.Get(Object); }, InternalProperty);
-		}
-		
-		// Get the bound property for this specified object ptr by providing type information about the object
-		const PropertyType& Get(const void* Object, const UScriptStruct* Type) const
-		{
-			ensureMsgf(Type == ObjectTypeInfo, TEXT("Provided object type (%s) did not match bound object type (%s)."),
-				*Type->GetFName().ToString(), *ObjectTypeInfo->GetFName().ToString());
-
-			return *Visit([Object](auto&& Prop) { return Prop.Get(Object); }, InternalProperty);
-		}
-
-	private:
-
-		// Internally we could be storing a direct property or a convertible property
-		using InternalPropertyType = TVariant<DirectProperty<PropertyType>, ConvertibleProperty<PropertyType>>;
-
-		// The actual property
-		InternalPropertyType InternalProperty;
-
-		// The type of the object we are bound to
-		UScriptStruct* ObjectTypeInfo = nullptr;
-	};
-	
 	template <typename AttributeType, TDataColumnType ColumnType>
 	TAttribute<AttributeType> FAttributeBinder::BindData(AttributeType ColumnType::* InVariable, const AttributeType& InDefaultValue,
 		const FName& InIdentifier)
 	{
+		using namespace Private;
+
 		if(!DataStorage)
 		{
 			return TAttribute<AttributeType>();
 		}
 		
 		// Create a direct property and bind it to the given variable
-		Property<AttributeType> Prop;
+		TProperty<AttributeType> Prop;
 		Prop.Bind(InVariable);
 
 		// We don't want any references to this in the lambda because binders are designed to be used and destructed on the stack
@@ -159,7 +31,7 @@ namespace UE::Editor::DataStorage
 			DefaultValue = InDefaultValue, Identifier = InIdentifier]()
 		{
 			// Get the column from the given row and use that to return the stored property
-			if(ColumnType* Column = AttributeBinding::Private::GetColumn<ColumnType>(Storage, Row, Identifier))
+			if(ColumnType* Column = Private::GetColumn<ColumnType>(Storage, Row, Identifier))
 			{
 				return Property.Get(Column, ColumnType::StaticStruct());
 			}
@@ -168,52 +40,64 @@ namespace UE::Editor::DataStorage
 	}
 
 	template <typename AttributeType, typename DataType, TDataColumnType ColumnType>
-	TAttribute<AttributeType> FAttributeBinder::BindData(DataType ColumnType::* InVariable,
-		const TFunction<AttributeType(const DataType&)>& InConverter, const DataType& InDefaultValue, const FName& InIdentifier)
+	TAttribute<AttributeType> FAttributeBinder::BindData(DataType ColumnType::* InVariable, TFunction<AttributeType(const DataType&)> InConverter, const DataType& InDefaultValue, const FName& InIdentifier)
 	{
+		using namespace Private;
+
 		if(!DataStorage)
 		{
 			return TAttribute<AttributeType>();
 		}
 		
+		AttributeType ConvertedDefault = InConverter(InDefaultValue);
+
 		// Create a convertible property and bind it to the given variable
-		Property<AttributeType> Prop;
-		Prop.Bind(InVariable, InConverter);
+		TProperty<AttributeType> Prop;
+		Prop.Bind(InVariable, MoveTemp(InConverter));
 	
-		return TAttribute<AttributeType>::CreateLambda([Property = MoveTemp(Prop), Storage = DataStorage, Row = TargetRow,
-			DefaultValue = InDefaultValue, Converter = InConverter, Identifier = InIdentifier]()
-		{
-			if(ColumnType* Column = AttributeBinding::Private::GetColumn<ColumnType>(Storage, Row, Identifier))
+		return TAttribute<AttributeType>::CreateLambda(
+			[
+				Property = MoveTemp(Prop), 
+				Storage = DataStorage, 
+				Row = TargetRow, 
+				DefaultValue = InDefaultValue, 
+				ConvertedDefault = MoveTemp(ConvertedDefault),
+				Identifier = InIdentifier
+			]()
 			{
-				return Property.Get(Column, ColumnType::StaticStruct());
-			}
+				if(ColumnType* Column = Private::GetColumn<ColumnType>(Storage, Row, Identifier))
+				{
+					return Property.Get(Column, ColumnType::StaticStruct());
+				}
 
-			return Converter(DefaultValue);
-		});
+				return ConvertedDefault;
+			});
 	}
 	
 	template <typename DataType, TDataColumnType ColumnType, typename FunctionType>
-		requires AttributeBinderInvocable<FunctionType, DataType>
+		requires Private::AttributeBinderInvocable<FunctionType, DataType>
 	auto FAttributeBinder::BindData(DataType ColumnType::* InVariable, FunctionType InConverter, const DataType& InDefaultValue, const FName& InIdentifier)
 	{
 		// Deduce the attribute type from the return value of the converter function
 		using AttributeType = decltype(InConverter(std::declval<DataType>()));
 		
-		return BindData<AttributeType, DataType>(InVariable, TFunction<AttributeType(const DataType&)>(InConverter), InDefaultValue, InIdentifier);
+		return BindData<AttributeType, DataType>(InVariable, TFunction<AttributeType(const DataType&)>(MoveTemp(InConverter)), InDefaultValue, InIdentifier);
 	}
 
 	template <typename InRetValType, typename... ParamTypes, TDataColumnType ColumnType>
 		TDelegate<InRetValType(ParamTypes...)> FAttributeBinder::BindEvent(TDelegate<InRetValType(ParamTypes...)> ColumnType::* InVariable,
 			const FName& InIdentifier)
 	{
+		using namespace Private;
+		
 		// Create a property for the delegate
-		Property<TDelegate<InRetValType(ParamTypes...)>> Prop;
+		TProperty<TDelegate<InRetValType(ParamTypes...)>> Prop;
 		Prop.Bind(InVariable);
 	
 		return TDelegate<InRetValType(ParamTypes...)>::CreateLambda([Property = MoveTemp(Prop), Storage = DataStorage,
 			Row = TargetRow, Identifier = InIdentifier](ParamTypes&&... Params)
 		{
-			if(ColumnType* Column = AttributeBinding::Private::GetColumn<ColumnType>(Storage, Row, Identifier))
+			if(ColumnType* Column = Private::GetColumn<ColumnType>(Storage, Row, Identifier))
 			{
 				// Get the delegate in the bound column for the specified row
 				TDelegate<InRetValType(ParamTypes...)> Delegate = Property.Get(Column, ColumnType::StaticStruct());
@@ -245,7 +129,5 @@ namespace UE::Editor::DataStorage
 			return FText::FromName(InName);
 		}, FName(), InIdentifier);
 	}
-	
-} // namespace UE::Editor::DataStorage
 
-	
\ No newline at end of file
+} // namespace UE::Editor::DataStorage
diff --git a/Engine/Source/Runtime/TypedElementFramework/Public/Elements/Framework/TypedElementAttributeBindingProperty.h b/Engine/Source/Runtime/TypedElementFramework/Public/Elements/Framework/TypedElementAttributeBindingProperty.h
new file mode 100644
--- /dev/null
+++ b/Engine/Source/Runtime/TypedElementFramework/Public/Elements/Framework/TypedElementAttributeBindingProperty.h
@@ -0,0 +1,249 @@
+// Copyright Epic Games, Inc. All Rights Reserved.
+
+#pragma once
+
+#include "Containers/Map.h"
+#include "Elements/Common/TypedElementHandles.h"
+#include "Elements/Interfaces/TypedElementDataStorageInterface.h"
+#include "Internationalization/Text.h"
+#include "Misc/Attribute.h"
+
+namespace UE::Editor::DataStorage
+{
+	namespace Private
+	{
+		// Concept for a conversion function used by the attribute binder to bind a column data member to an attribute of a different type
+		template <typename FunctionType, typename ArgumentType>
+		concept AttributeBinderInvocable = std::is_invocable_v<std::decay_t<FunctionType>, const ArgumentType&>;
+		
+		template <TDataColumnType ColumnType>
+		ColumnType* GetColumn(IEditorDataStorageProvider* DataStorage, RowHandle Row, const FName& InIdentifier)
+		{
+			if(InIdentifier != NAME_None)
+			{
+				if(ColumnType* DynamicColumn = DataStorage->GetColumn<ColumnType>(Row, InIdentifier))
+				{
+					return DynamicColumn;
+				}
+
+				// The dynamic column wasn't found on this row
+				return nullptr;
+			}
+
+			// If we weren't provided a dynamic column identifier, we are looking for the regular non-dynamic column
+			return DataStorage->GetColumn<ColumnType>(Row);
+		}
+
+		template <typename PropertyType>
+		struct TEmptyProperty final
+		{
+			PropertyType* Get(void* Object) const;
+			const PropertyType* Get(const void* Object) const;
+		};
+
+		template <typename PropertyType>
+		struct TDirectProperty final
+		{
+			// Offset to a particular data member inside the object
+			size_t Offset;
+
+			PropertyType* Get(void* Object) const;
+			const PropertyType* Get(const void* Object) const;
+		};
+
+		// A property that goes through a conversion function before being accessed from the object
+		template <typename PropertyType>
+		struct TConvertibleProperty final
+		{
+			// Conversion function
+			TFunction<PropertyType(const void*)> Converter;
+
+			// Cache to avoid a copy when returning if not necessary
+			mutable PropertyType Cache;
+
+			PropertyType* Get(void* Object) const;
+			const PropertyType* Get(const void* Object) const;
+		};
+
+		template <typename PropertyType>
+		class TProperty
+		{
+		public:
+			// Bind this property directly
+			template <TDataColumnType ObjectType>
+			void Bind(PropertyType ObjectType::* Variable);
+
+			// Bind this property using a conversion function
+			template <typename InputType, TDataColumnType ObjectType>
+			void Bind(InputType ObjectType::* Variable, TFunction<PropertyType(const InputType&)> Converter);
+
+			// Get the bound property for the specified object
+			template <TDataColumnType ObjectType>
+			PropertyType& Get(ObjectType& Object) const;
+
+			// Get the bound property for the specified object
+			template <TDataColumnType ObjectType>
+			const PropertyType& Get(const ObjectType& Object) const;
+
+			// Get the bound property for this specified object ptr by providing type information about the object
+			PropertyType& Get(void* Object, const UScriptStruct* Type) const;
+
+			// Get the bound property for this specified object ptr by providing type information about the object
+			const PropertyType& Get(const void* Object, const UScriptStruct* Type) const;
+
+			// Returns the stored type information or null if no type was set.
+			const UScriptStruct* GetObjectTypeInfo() const;
+
+			// Whether or not this property has been bound.
+			bool IsBound() const;
+
+		private:
+			// Internally we could be storing a direct property or a convertible property
+			using InternalPropertyType = TVariant<
+				TEmptyProperty<PropertyType>,
+				TDirectProperty<PropertyType>,
+				TConvertibleProperty<PropertyType>>;
+
+			// The actual property
+			InternalPropertyType InternalProperty;
+
+			// The type of the object we are bound to
+			UScriptStruct* ObjectTypeInfo = nullptr;
+		};
+	
+		//
+		// TEmptyProperty
+		//
+		template <typename PropertyType>
+		PropertyType* TEmptyProperty<PropertyType>::Get(void* Object) const
+		{
+			checkf(false, TEXT("Empty properties shouldn't be accessed."));
+			return nullptr;
+		}
+		
+		template <typename PropertyType>
+		const PropertyType* TEmptyProperty<PropertyType>::Get(const void* Object) const
+		{
+			checkf(false, TEXT("Empty properties shouldn't be accessed."));
+			return nullptr;
+		}
+		
+		
+		//
+		// TDirectProperty
+		//
+
+		template <typename PropertyType>
+		PropertyType* TDirectProperty<PropertyType>::Get(void* Object) const
+		{
+			return reinterpret_cast<PropertyType*>((reinterpret_cast<char*>(Object) + Offset));
+		}
+
+		template <typename PropertyType>
+		const PropertyType* TDirectProperty<PropertyType>::Get(const void* Object) const
+		{
+			return reinterpret_cast<const PropertyType*>((reinterpret_cast<const char*>(Object) + Offset));
+		}
+
+
+		//
+		// TConvertibleProperty
+		//
+
+		template <typename PropertyType>
+		PropertyType* TConvertibleProperty<PropertyType>::Get(void* Object) const
+		{
+			Cache = Converter(Object);
+			return &Cache;
+		}
+
+		template <typename PropertyType>
+		const PropertyType* TConvertibleProperty<PropertyType>::Get(const void* Object) const
+		{
+			Cache = Converter(Object);
+			return &Cache;
+		}
+
+
+		//
+		// TProperty
+		//
+
+		template <typename PropertyType>
+		template <TDataColumnType ObjectType>
+		void TProperty<PropertyType>::Bind(PropertyType ObjectType::* Variable)
+		{
+			TDirectProperty<PropertyType> Result;
+			Result.Offset = reinterpret_cast<size_t>(&(reinterpret_cast<const ObjectType*>(0)->*Variable));
+
+			ObjectTypeInfo = ObjectType::StaticStruct();
+			InternalProperty = InternalPropertyType(TInPlaceType<TDirectProperty<PropertyType>>(), MoveTemp(Result));
+		}
+
+		template <typename PropertyType>
+		template <typename InputType, TDataColumnType ObjectType>
+		void TProperty<PropertyType>::Bind(InputType ObjectType::* Variable, TFunction<PropertyType(const InputType&)> Converter)
+		{
+			TConvertibleProperty<PropertyType> Result;
+			Result.Converter = [Converter = MoveTemp(Converter), Variable](const void* Object)
+				{
+					return Converter(reinterpret_cast<const ObjectType*>(Object)->*Variable);
+				};
+
+			ObjectTypeInfo = ObjectType::StaticStruct();
+			InternalProperty = InternalPropertyType(TInPlaceType<TConvertibleProperty<PropertyType>>(), MoveTemp(Result));
+		}
+
+		template <typename PropertyType>
+		template <TDataColumnType ObjectType>
+		PropertyType& TProperty<PropertyType>::Get(ObjectType& Object) const
+		{
+			return Get(&Object, ObjectType::StaticStruct());
+		}
+
+		template <typename PropertyType>
+		template <TDataColumnType ObjectType>
+		const PropertyType& TProperty<PropertyType>::Get(const ObjectType& Object) const
+		{
+			return Get(&Object, ObjectType::StaticStruct());
+		}
+
+		template <typename PropertyType>
+		PropertyType& TProperty<PropertyType>::Get(void* Object, const UScriptStruct* Type) const
+		{
+			checkf(Object, TEXT("Nullptr pointer provided to Object while trying to retrieve a property value."));
+			checkf(Type, TEXT("Nullptr pointer provided to Type while trying to retrieve a property value."));
+			checkf(IsBound(), TEXT("Attempting to retrieve the value of object type (%s) from a property that wasn't bound."),
+				*Type->GetFName().ToString());
+			ensureMsgf(Type == ObjectTypeInfo, TEXT("Provided object type (%s) did not match bound object type (%s)."),
+				*Type->GetFName().ToString(), ObjectTypeInfo ? *ObjectTypeInfo->GetFName().ToString() : TEXT("<no type>"));
+
+			return *Visit([Object](auto&& Prop) { return Prop.Get(Object); }, InternalProperty);
+		}
+
+		template <typename PropertyType>
+		const PropertyType& TProperty<PropertyType>::Get(const void* Object, const UScriptStruct* Type) const
+		{
+			checkf(Object, TEXT("Nullptr pointer provided to Object while trying to retrieve a property value."));
+			checkf(Type, TEXT("Nullptr pointer provided to Type while trying to retrieve a property value."));
+			checkf(IsBound(), TEXT("Attempting to retrieve the value of object type (%s) from a property that wasn't bound."),
+				*Type->GetFName().ToString());
+			ensureMsgf(Type == ObjectTypeInfo, TEXT("Provided object type (%s) did not match bound object type (%s)."),
+				*Type->GetFName().ToString(), ObjectTypeInfo ? *ObjectTypeInfo->GetFName().ToString() : TEXT("<no type>"));
+			
+			return *Visit([Object](auto&& Prop) { return Prop.Get(Object); }, InternalProperty);
+		}
+
+		template <typename PropertyType>
+		const UScriptStruct* TProperty<PropertyType>::GetObjectTypeInfo() const
+		{
+			return ObjectTypeInfo;
+		}
+
+		template <typename PropertyType>
+		bool TProperty<PropertyType>::IsBound() const
+		{
+			return !InternalProperty.template IsType<TEmptyProperty<PropertyType>>();
+		}
+	} // namespace Private
+}
\ No newline at end of file
diff --git a/Engine/Source/Runtime/TypedElementFramework/Public/Elements/Framework/TypedElementAttributeBindingText.h b/Engine/Source/Runtime/TypedElementFramework/Public/Elements/Framework/TypedElementAttributeBindingText.h
new file mode 100644
--- /dev/null
+++ b/Engine/Source/Runtime/TypedElementFramework/Public/Elements/Framework/TypedElementAttributeBindingText.h
@@ -0,0 +1,125 @@
+// Copyright Epic Games, Inc. All Rights Reserved.
+
+#pragma once
+
+#include "Containers/Map.h"
+#include "Elements/Common/TypedElementHandles.h"
+#include "Elements/Interfaces/TypedElementDataStorageInterface.h"
+#include "Internationalization/Text.h"
+#include "Misc/Attribute.h"
+#include "TypedElementAttributeBindingProperty.h"
+
+namespace UE::Editor::DataStorage
+{
+	
+	//
+	// FTextAttributeFormatted
+	//
+	
+	class FTextAttributeFormatted
+	{
+		friend class FAttributeBinder;
+		
+	public:
+		TYPEDELEMENTFRAMEWORK_API FTextAttributeFormatted& Arg(FString Name, FFormatArgumentValue Value);
+			
+		template <TDataColumnType ColumnType>
+		FTextAttributeFormatted& Arg(FString Name, FText ColumnType::* Variable, FFormatArgumentValue Default = {},
+			FName ColumnIdentifier = NAME_None);
+
+		template <TDataColumnType ColumnType>
+		FTextAttributeFormatted& Arg(FString Name, FString ColumnType::* Variable, FFormatArgumentValue Default = {},
+			FName ColumnIdentifier = NAME_None);
+
+		template <TDataColumnType ColumnType>
+		FTextAttributeFormatted& Arg(FString Name, FName ColumnType::* Variable, FFormatArgumentValue Default = {},
+			FName ColumnIdentifier = NAME_None);
+
+		template <typename DataType, TDataColumnType ColumnType>
+		FTextAttributeFormatted& Arg(
+			FString Name, 
+			DataType ColumnType::* Variable, 
+			const TFunction<FText(const DataType&)>& Converter, 
+			FFormatArgumentValue Default = {},
+			FName ColumnIdentifier = NAME_None);
+			
+		template <typename DataType, TDataColumnType ColumnType, typename FunctionType> 
+			requires Private::AttributeBinderInvocable<FunctionType, DataType>
+		FTextAttributeFormatted& Arg(FString Name, DataType ColumnType::* Variable, FunctionType Converter, FFormatArgumentValue Default = {},
+			FName ColumnIdentifier = NAME_None);
+	
+		TYPEDELEMENTFRAMEWORK_API operator TAttribute<FText>();
+
+	private:
+		FTextAttributeFormatted(FTextFormat InFormat, RowHandle InTargetRow, IEditorDataStorageProvider* InDataStorage);
+
+		struct FPropertyInfo
+		{
+			Private::TProperty<FText> Property;
+			FFormatArgumentValue Default;
+			FName DynamicColumnIdentifier;
+		};
+		TMap<FString, FPropertyInfo> NamedProperties;
+		FTextFormat Format;
+			
+		// The target row for this binder
+		RowHandle TargetRow;
+		// A ptr to the data storage for quick access
+		IEditorDataStorageProvider* DataStorage;
+	};
+
+	template <TDataColumnType ColumnType>
+	FTextAttributeFormatted& FTextAttributeFormatted::Arg(
+		FString Name, FText ColumnType::* Variable, FFormatArgumentValue Default, FName ColumnIdentifier)
+	{
+		Private::TProperty<FText> Prop;
+		Prop.Bind(Variable);
+		NamedProperties.Add(MoveTemp(Name), FPropertyInfo{ .Property = MoveTemp(Prop), .Default = MoveTemp(Default),
+			.DynamicColumnIdentifier = ColumnIdentifier});
+
+		return *this;
+	}
+
+	template <TDataColumnType ColumnType>
+	FTextAttributeFormatted& FTextAttributeFormatted::Arg(
+		FString Name, FString ColumnType::* Variable, FFormatArgumentValue Default, FName ColumnIdentifier)
+	{
+		return Arg(MoveTemp(Name), Variable, [](const FString& Value) { return FText::FromString(Value); }, MoveTemp(Default), ColumnIdentifier);
+	}
+
+	template <TDataColumnType ColumnType>
+	FTextAttributeFormatted& FTextAttributeFormatted::Arg(
+		FString Name, FName ColumnType::* Variable, FFormatArgumentValue Default, FName ColumnIdentifier)
+	{
+		return Arg(MoveTemp(Name), Variable, [](const FName& Value) { return FText::FromName(Value); }, MoveTemp(Default), ColumnIdentifier);
+	}
+
+	template <typename DataType, TDataColumnType ColumnType>
+	FTextAttributeFormatted& FTextAttributeFormatted::Arg(
+		FString Name,
+		DataType ColumnType::* Variable,
+		const TFunction<FText(const DataType&)>& Converter, 
+		FFormatArgumentValue Default,
+		FName ColumnIdentifier)
+	{
+		Private::TProperty<FText> Prop;
+		Prop.Bind(Variable, Converter);
+		NamedProperties.Add(MoveTemp(Name), FPropertyInfo{ .Property = MoveTemp(Prop), .Default = MoveTemp(Default),
+			.DynamicColumnIdentifier = ColumnIdentifier });
+		
+		return *this;
+	}
+
+	template <typename DataType, TDataColumnType ColumnType, typename FunctionType>
+		requires Private::AttributeBinderInvocable<FunctionType, DataType>
+	FTextAttributeFormatted& FTextAttributeFormatted::Arg(
+		FString Name,
+		DataType ColumnType::* Variable,
+		FunctionType Converter,
+		FFormatArgumentValue Default,
+		FName ColumnIdentifier)
+	{
+		return Arg(MoveTemp(Name), Variable, TFunction<FText(const DataType&)>(MoveTemp(Converter)), MoveTemp(Default), ColumnIdentifier);
+	}
+
+}
\ No newline at end of file
