Teds Attribute Binding: Remove accidentally pragma once in a cpp file

diff --git a/Engine/Source/Runtime/TypedElementFramework/Private/Elements/Framework/TypedElementtAttributeBindingText.cpp b/Engine/Source/Runtime/TypedElementFramework/Private/Elements/Framework/TypedElementtAttributeBindingText.cpp
--- a/Engine/Source/Runtime/TypedElementFramework/Private/Elements/Framework/TypedElementtAttributeBindingText.cpp
+++ b/Engine/Source/Runtime/TypedElementFramework/Private/Elements/Framework/TypedElementtAttributeBindingText.cpp
@@ -1,7 +1,5 @@
 // Copyright Epic Games, Inc. All Rights Reserved.
 
-#pragma once
-
 #include "Elements/Framework/TypedElementAttributeBindingText.h"
 
 namespace UE::Editor::DataStorage
