Fixed an issue where ListColumns wouldn't report pending columns.

To remain consistent with `HasColumns`, `ListColumns` should also return the list of columns that are known up until that point. Since this applies to reserved rows, it's still impossible to return the list of columns for the table as those are not yet known. The documentation in the interface has been updated to reflect this.

diff --git a/Engine/Plugins/Experimental/EditorDataStorage/Source/TedsCore/Private/TypedElementDatabase.cpp b/Engine/Plugins/Experimental/EditorDataStorage/Source/TedsCore/Private/TypedElementDatabase.cpp
--- a/Engine/Plugins/Experimental/EditorDataStorage/Source/TedsCore/Private/TypedElementDatabase.cpp
+++ b/Engine/Plugins/Experimental/EditorDataStorage/Source/TedsCore/Private/TypedElementDatabase.cpp
@@ -874,7 +874,10 @@ void UEditorDataStorage::ListColumns(RowHandle Row, ColumnListCallbackRef Callba
 				};
 			Composition.Fragments.ExportTypes(CallbackWrapper);
 			Composition.Tags.ExportTypes(CallbackWrapper);
-
+		}
+		else
+		{
+			Environment->GetDirectDeferredCommands().ListColumns(Row, Callback);
 		}
 	}
 }
@@ -909,6 +912,10 @@ void UEditorDataStorage::ListColumns(RowHandle Row, ColumnListWithDataCallbackRe
 				});
 
 		}
+		else
+		{
+			Environment->GetDirectDeferredCommands().ListColumns(Row, Callback);
+		}
 	}
 }
 
diff --git a/Engine/Plugins/Experimental/EditorDataStorage/Source/TedsCore/Private/TypedElementDatabaseCommandBuffer.cpp b/Engine/Plugins/Experimental/EditorDataStorage/Source/TedsCore/Private/TypedElementDatabaseCommandBuffer.cpp
--- a/Engine/Plugins/Experimental/EditorDataStorage/Source/TedsCore/Private/TypedElementDatabaseCommandBuffer.cpp
+++ b/Engine/Plugins/Experimental/EditorDataStorage/Source/TedsCore/Private/TypedElementDatabaseCommandBuffer.cpp
@@ -55,6 +55,34 @@ namespace UE::Editor::DataStorage::Legacy
 		return PendingColumns.Contains(PendingColumnMappingKey(Row, ColumnType));
 	}
 
+	void FCommandBuffer::ListColumns(RowHandle Row, ColumnListCallbackRef Callback) const
+	{
+		for (const TPair<PendingColumnMappingKey, void*>& Column : PendingColumns)
+		{
+			if (Column.Key.Key == Row)
+			{
+				if (const UScriptStruct* ColumnType = Column.Key.Value.Get())
+				{
+					Callback(*ColumnType);
+				}
+			}
+		}
+	}
+
+	void FCommandBuffer::ListColumns(RowHandle Row, ColumnListWithDataCallbackRef Callback)
+	{
+		for (const TPair<PendingColumnMappingKey, void*>& Column : PendingColumns)
+		{
+			if (Column.Key.Key == Row)
+			{
+				if (const UScriptStruct* ColumnType = Column.Key.Value.Get())
+				{
+					Callback(Column.Value, *ColumnType);
+				}
+			}
+		}
+	}
+
 	void FCommandBuffer::Clear(RowHandle Row)
 	{
 		for (TMap<PendingColumnMappingKey, void*>::TIterator It = PendingColumns.CreateIterator(); It; ++It)
diff --git a/Engine/Plugins/Experimental/EditorDataStorage/Source/TedsCore/Private/TypedElementDatabaseCommandBuffer.h b/Engine/Plugins/Experimental/EditorDataStorage/Source/TedsCore/Private/TypedElementDatabaseCommandBuffer.h
--- a/Engine/Plugins/Experimental/EditorDataStorage/Source/TedsCore/Private/TypedElementDatabaseCommandBuffer.h
+++ b/Engine/Plugins/Experimental/EditorDataStorage/Source/TedsCore/Private/TypedElementDatabaseCommandBuffer.h
@@ -28,6 +28,10 @@ namespace UE::Editor::DataStorage
 			void* GetQueuedDataColumn(RowHandle Row, const UScriptStruct* ColumnType);
 			/** Returns if the column on the provided row is pending processing. */
 			bool HasColumn(RowHandle Row, const UScriptStruct* ColumnType) const;
+			/** Returns a list of columns pending processing for a row.. */
+			void ListColumns(RowHandle Row, ColumnListCallbackRef Callback) const;
+			/** Returns a list of columns and their data addresses (can be null for tags) pending processing for a row.. */
+			void ListColumns(RowHandle Row, ColumnListWithDataCallbackRef Callback);
 			/** Clears all entries in the internal cache for the provided row. */
 			void Clear(RowHandle Row);
 	
diff --git a/Engine/Source/Runtime/TypedElementFramework/Public/Elements/Interfaces/TypedElementDataStorageInterface.h b/Engine/Source/Runtime/TypedElementFramework/Public/Elements/Interfaces/TypedElementDataStorageInterface.h
--- a/Engine/Source/Runtime/TypedElementFramework/Public/Elements/Interfaces/TypedElementDataStorageInterface.h
+++ b/Engine/Source/Runtime/TypedElementFramework/Public/Elements/Interfaces/TypedElementDataStorageInterface.h
@@ -278,18 +278,28 @@ public:
 	template<UE::Editor::DataStorage::TDataColumnType ColumnTypeTemplate>
 	const ColumnTypeTemplate* GetColumn(RowHandle Row, const FName& Identifer) const;
 	
-	/** Determines if the provided row contains the collection of columns and tags. */
+	/** 
+	 * Determines if the provided row contains the collection of columns and tags.
+	 * Note that for rows that haven't been assigned yet it's not possible to check if a column exists as the table to check for hasn't
+	 * been assigned yet. In these cases a list of known additions will be checked as these will be added once the table is assigned.
+	 */
 	virtual bool HasColumns(RowHandle Row, TConstArrayView<const UScriptStruct*> ColumnTypes) const = 0;
 	virtual bool HasColumns(RowHandle Row, TConstArrayView<TWeakObjectPtr<const UScriptStruct>> ColumnTypes) const = 0;
 	template<UE::Editor::DataStorage::TColumnType... ColumnTypes>
 	bool HasColumns(RowHandle Row) const;
 
-	/** Lists the columns on a row. This includes data and tag columns. */
+	/** 
+	 * Lists the columns on a row. This includes data and tag columns.
+	 * Note that for rows that haven't been assigned yet it's not possible to return the full list as no table has been assigned yet.
+	 * In these cases a list of known additions will be returned as these will be added once the table is assigned.
+	 */
 	virtual void ListColumns(RowHandle Row, ColumnListCallbackRef Callback) const = 0;
 
 	/** 
 	 * Lists the column type and data on a row. This includes data and tag columns. Not all columns may have data so the data pointer in 
 	 * the callback can be null.
+	 * Note that for rows that haven't been assigned yet it's not possible to return the full list as no table has been assigned yet.
+	 * In these cases a list of known additions will be returned.
 	 */
 	virtual void ListColumns(RowHandle Row, ColumnListWithDataCallbackRef Callback) = 0;
 
