Fixed incorrect columns being added to widget rows through TEDS UI.

diff --git a/Engine/Plugins/Experimental/EditorDataStorageFeatures/Source/TedsOutliner/Private/Widgets/AlertWidget.cpp b/Engine/Plugins/Experimental/EditorDataStorageFeatures/Source/TedsOutliner/Private/Widgets/AlertWidget.cpp
--- a/Engine/Plugins/Experimental/EditorDataStorageFeatures/Source/TedsOutliner/Private/Widgets/AlertWidget.cpp
+++ b/Engine/Plugins/Experimental/EditorDataStorageFeatures/Source/TedsOutliner/Private/Widgets/AlertWidget.cpp
@@ -193,7 +193,7 @@ void UAlertWidgetFactory::RegisterAlertQueries(IEditorDataStorageProvider& DataS
 		.Compile());
 	
 	DataStorage.RegisterQuery(
-		Select(TEXT("Sync Transform column to heads up display"),
+		Select(TEXT("Update alert widget state"),
 			FProcessor(EQueryTickPhase::FrameEnd, DataStorage.GetQueryTickGroupName(EQueryTickGroups::SyncWidgets))
 				.SetExecutionMode(EExecutionMode::GameThread),
 			[](IQueryContext& Context,
diff --git a/Engine/Plugins/Experimental/EditorDataStorageFeatures/Source/TedsTableViewer/Private/TedsTableViewerColumn.cpp b/Engine/Plugins/Experimental/EditorDataStorageFeatures/Source/TedsTableViewer/Private/TedsTableViewerColumn.cpp
--- a/Engine/Plugins/Experimental/EditorDataStorageFeatures/Source/TedsTableViewer/Private/TedsTableViewerColumn.cpp
+++ b/Engine/Plugins/Experimental/EditorDataStorageFeatures/Source/TedsTableViewer/Private/TedsTableViewerColumn.cpp
@@ -211,18 +211,10 @@ namespace UE::Editor::DataStorage
 		// We are looking for widgets that have a row reference
 		TArray<const UScriptStruct*> SelectionColumns({FTypedElementSlateWidgetReferenceColumn::StaticStruct(), FTypedElementRowReferenceColumn::StaticStruct()});
 
-		// We need to remove duplicates because TEDS/Mass does not handle having the same column in Select() and Where()
-		TArray<const UScriptStruct*> AdditionalWidgetColumns = CellWidgetConstructor->GetAdditionalColumnsList().FilterByPredicate([&SelectionColumns](const UScriptStruct* Column)
-		{
-			return !SelectionColumns.Contains(Column);
-		});
-		
 		// Query to get all widgets that were created by this column
 		WidgetQuery = Storage->RegisterQuery(
 			Select()
 				.ReadOnly(SelectionColumns)
-			.Where()
-				.All(AdditionalWidgetColumns)
 			.Compile());
 	}
 
diff --git a/Engine/Source/Runtime/TypedElementFramework/Private/Elements/Interfaces/TypedElementDataStorageUiInterface.cpp b/Engine/Source/Runtime/TypedElementFramework/Private/Elements/Interfaces/TypedElementDataStorageUiInterface.cpp
--- a/Engine/Source/Runtime/TypedElementFramework/Private/Elements/Interfaces/TypedElementDataStorageUiInterface.cpp
+++ b/Engine/Source/Runtime/TypedElementFramework/Private/Elements/Interfaces/TypedElementDataStorageUiInterface.cpp
@@ -101,26 +101,28 @@ TSharedPtr<SWidget> FTypedElementWidgetConstructor::ConstructFinalWidget(
 {
 	// Add the additional columns to the UI row
 	TSharedPtr<SWidget> Widget = SNullWidget::NullWidget;
-	DataStorage->AddColumns(Row, GetAdditionalColumnsList());
 	
 	if (const FTypedElementRowReferenceColumn* RowReference = DataStorage->GetColumn<FTypedElementRowReferenceColumn>(Row))
 	{
 		bool bConstructWidget = DataStorage->IsRowAssigned(RowReference->Row);
 
 		// If the original row matches this widgets query conditions currently, create the actual internal widget
-		if (const UE::Editor::DataStorage::Queries::FConditions* MatchedQueryConditions = GetQueryConditions())
+		if (const UE::Editor::DataStorage::Queries::FConditions* MatchedQueryConditions = GetQueryConditions();
+			MatchedQueryConditions != nullptr && bConstructWidget)
 		{
 			bConstructWidget &= DataStorage->MatchesColumns(RowReference->Row, *MatchedQueryConditions);
 		}
 		
 		if (bConstructWidget)
 		{
+			DataStorage->AddColumns(Row, GetAdditionalColumnsList());
 			Widget = Construct(Row, DataStorage, DataStorageUi, Arguments);
 		}
 	}
 	// If we don't have an original row, simply construct the widget
 	else
 	{
+		DataStorage->AddColumns(Row, GetAdditionalColumnsList());
 		Widget = Construct(Row, DataStorage, DataStorageUi, Arguments);
 	}
 
