The name for tables is now taken by const reference instead of by const value.

diff --git a/Engine/Plugins/Experimental/EditorDataStorage/Source/TedsCore/Private/TypedElementDatabase.cpp b/Engine/Plugins/Experimental/EditorDataStorage/Source/TedsCore/Private/TypedElementDatabase.cpp
--- a/Engine/Plugins/Experimental/EditorDataStorage/Source/TedsCore/Private/TypedElementDatabase.cpp
+++ b/Engine/Plugins/Experimental/EditorDataStorage/Source/TedsCore/Private/TypedElementDatabase.cpp
@@ -265,7 +265,7 @@ TSharedPtr<const FMassEntityManager> UEditorDataStorage::GetActiveEditorEntityMa
 	return ActiveEditorEntityManager;
 }
 
-UE::Editor::DataStorage::TableHandle UEditorDataStorage::RegisterTable(TConstArrayView<const UScriptStruct*> ColumnList, const FName Name)
+UE::Editor::DataStorage::TableHandle UEditorDataStorage::RegisterTable(TConstArrayView<const UScriptStruct*> ColumnList, const FName& Name)
 {
 	using namespace UE::Editor::DataStorage;
 
@@ -286,7 +286,7 @@ UE::Editor::DataStorage::TableHandle UEditorDataStorage::RegisterTable(TConstArr
 }
 
 UE::Editor::DataStorage::TableHandle UEditorDataStorage::RegisterTable(UE::Editor::DataStorage::TableHandle SourceTable,
-	TConstArrayView<const UScriptStruct*> ColumnList, const FName Name)
+	TConstArrayView<const UScriptStruct*> ColumnList, const FName& Name)
 {
 	using namespace UE::Editor::DataStorage;
 
@@ -306,7 +306,7 @@ UE::Editor::DataStorage::TableHandle UEditorDataStorage::RegisterTable(UE::Edito
 	return InvalidTableHandle;
 }
 
-UE::Editor::DataStorage::TableHandle UEditorDataStorage::FindTable(const FName Name)
+UE::Editor::DataStorage::TableHandle UEditorDataStorage::FindTable(const FName& Name)
 {
 	using namespace UE::Editor::DataStorage;
 
diff --git a/Engine/Plugins/Experimental/EditorDataStorage/Source/TedsCore/Private/TypedElementDatabase.h b/Engine/Plugins/Experimental/EditorDataStorage/Source/TedsCore/Private/TypedElementDatabase.h
--- a/Engine/Plugins/Experimental/EditorDataStorage/Source/TedsCore/Private/TypedElementDatabase.h
+++ b/Engine/Plugins/Experimental/EditorDataStorage/Source/TedsCore/Private/TypedElementDatabase.h
@@ -93,10 +93,10 @@ public:
 	TSharedPtr<FMassEntityManager> GetActiveMutableEditorEntityManager();
 	TSharedPtr<const FMassEntityManager> GetActiveEditorEntityManager() const;
 
-	virtual TableHandle RegisterTable(TConstArrayView<const UScriptStruct*> ColumnList, const FName Name) override;
+	virtual TableHandle RegisterTable(TConstArrayView<const UScriptStruct*> ColumnList, const FName& Name) override;
 	virtual TableHandle RegisterTable(
-		TableHandle SourceTable, TConstArrayView<const UScriptStruct*> ColumnList, const FName Name) override;
-	virtual TableHandle FindTable(const FName Name) override;
+		TableHandle SourceTable, TConstArrayView<const UScriptStruct*> ColumnList, const FName& Name) override;
+	virtual TableHandle FindTable(const FName& Name) override;
 
 	virtual RowHandle ReserveRow() override;
 	virtual void BatchReserveRows(int32 Count, TFunctionRef<void(RowHandle)> ReservationCallback) override;
diff --git a/Engine/Source/Runtime/TypedElementFramework/Public/Elements/Interfaces/TypedElementDataStorageInterface.h b/Engine/Source/Runtime/TypedElementFramework/Public/Elements/Interfaces/TypedElementDataStorageInterface.h
--- a/Engine/Source/Runtime/TypedElementFramework/Public/Elements/Interfaces/TypedElementDataStorageInterface.h
+++ b/Engine/Source/Runtime/TypedElementFramework/Public/Elements/Interfaces/TypedElementDataStorageInterface.h
@@ -81,20 +81,20 @@ public:
 	 */
 
 	/** Creates a new table for with the provided columns. Optionally a name can be given which is useful for retrieval later. */
-	virtual TableHandle RegisterTable(TConstArrayView<const UScriptStruct*> ColumnList, const FName Name) = 0;
+	virtual TableHandle RegisterTable(TConstArrayView<const UScriptStruct*> ColumnList, const FName& Name) = 0;
 	template<UE::Editor::DataStorage::TColumnType... Columns>
-	TableHandle RegisterTable(const FName Name);
+	TableHandle RegisterTable(const FName& Name);
 	/** 
 	 * Copies the column information from the provided table and creates a new table for with the provided columns. Optionally a 
 	 * name can be given which is useful for retrieval later.
 	 */
 	virtual TableHandle RegisterTable(TableHandle SourceTable,
-		TConstArrayView<const UScriptStruct*> ColumnList, const FName Name) = 0;
+		TConstArrayView<const UScriptStruct*> ColumnList, const FName& Name) = 0;
 	template<UE::Editor::DataStorage::TColumnType... Columns>
-	TableHandle RegisterTable(TableHandle SourceTable, const FName Name);
+	TableHandle RegisterTable(TableHandle SourceTable, const FName& Name);
 
 	/** Returns a previously created table with the provided name or TypedElementInvalidTableHandle if not found. */
-	virtual TableHandle FindTable(const FName Name) = 0;
+	virtual TableHandle FindTable(const FName& Name) = 0;
 	
 	/**
 	 * @section Row management
@@ -458,14 +458,14 @@ const FactoryT* IEditorDataStorageProvider::FindFactory() const
 }
 
 template<UE::Editor::DataStorage::TColumnType... Columns>
-UE::Editor::DataStorage::TableHandle IEditorDataStorageProvider::RegisterTable(const FName Name)
+UE::Editor::DataStorage::TableHandle IEditorDataStorageProvider::RegisterTable(const FName& Name)
 {
 	return RegisterTable({ Columns::StaticStruct()... }, Name);
 }
 
 template<UE::Editor::DataStorage::TColumnType... Columns>
 UE::Editor::DataStorage::TableHandle IEditorDataStorageProvider::RegisterTable(
-	TableHandle SourceTable, const FName Name)
+	TableHandle SourceTable, const FName& Name)
 {
 	return RegisterTable(SourceTable, { Columns::StaticStruct()... }, Name);
 }
