Added a new array and array view optimized for row handles.

diff --git a/Engine/Source/Runtime/TypedElementFramework/Private/Elements/Framework/TypedElementRowHandleArray.cpp b/Engine/Source/Runtime/TypedElementFramework/Private/Elements/Framework/TypedElementRowHandleArray.cpp
new file mode 100644
--- /dev/null
+++ b/Engine/Source/Runtime/TypedElementFramework/Private/Elements/Framework/TypedElementRowHandleArray.cpp
@@ -0,0 +1,348 @@
+// Copyright Epic Games, Inc. All Rights Reserved.
+
+#include "Elements/Framework/TypedElementRowHandleArray.h"
+
+#include <algorithm>
+#include <iterator>
+#include <limits>
+#include "HAL/UnrealMemory.h"
+#include "Templates/Sorting.h"
+
+namespace UE::Editor::DataStorage
+{
+	FRowHandleArray::FRowHandleArray(FRowHandleArrayView InRows)
+		: bIsSorted(InRows.IsSorted())
+		, bIsUnique(InRows.IsUnique())
+	{
+		Rows.AddUninitialized(InRows.Num());
+		FMemory::Memcpy(Rows.GetData(), InRows.GetData(), InRows.NumBytes());
+	}
+
+	int32 FRowHandleArray::Num() const
+	{
+		return Rows.Num();
+	}
+
+	void FRowHandleArray::Reserve(int32 Count)
+	{
+		Rows.Reserve(Count);
+	}
+
+	void FRowHandleArray::Shrink()
+	{
+		Rows.Shrink();
+	}
+
+	void FRowHandleArray::Add(RowHandle Row)
+	{
+		bIsSorted = bIsSorted && (Rows.IsEmpty() || Rows.Last() <= Row);
+		bIsUnique = bIsUnique && bIsSorted && (Rows.IsEmpty() || Rows.Last() < Row);
+		Rows.Add(Row);
+	}
+
+	void FRowHandleArray::Append(TConstArrayView<RowHandle> AdditionalRows)
+	{
+		if (!AdditionalRows.IsEmpty())
+		{
+			bIsSorted = false;
+			bIsUnique = false;
+			int32 StartIndex = Rows.AddUninitialized(AdditionalRows.Num());
+			FMemory::Memcpy(Rows.GetData() + StartIndex, AdditionalRows.GetData(), AdditionalRows.NumBytes());
+		}
+	}
+
+	void FRowHandleArray::Append(FRowHandleArrayView AdditionalRows)
+	{
+		if (!AdditionalRows.IsEmpty())
+		{
+			if (!Rows.IsEmpty())
+			{
+				RowHandle CurrentLast = Rows.Last();
+				RowHandle FirstNew = *AdditionalRows.begin();
+				bIsSorted = bIsSorted && AdditionalRows.IsSorted() && (CurrentLast <= FirstNew);
+				bIsUnique = bIsSorted && bIsUnique && AdditionalRows.IsUnique() && (CurrentLast < FirstNew);
+			}
+			else
+			{
+				bIsSorted = AdditionalRows.IsSorted();
+				bIsUnique = AdditionalRows.IsUnique();
+			}
+			
+			int32 StartIndex = Rows.AddUninitialized(AdditionalRows.Num());
+			FMemory::Memcpy(Rows.GetData() + StartIndex, AdditionalRows.GetData(), AdditionalRows.NumBytes());
+		}
+	}
+
+	void FRowHandleArray::Sort()
+	{
+		if (!bIsSorted)
+		{
+			RadixSort64(Rows.GetData(), Rows.Num());
+			bIsSorted = true;
+		}
+	}
+
+	void FRowHandleArray::Sort(TArrayView<RowHandle> ScratchBuffer)
+	{
+		
+		if (!bIsSorted)
+		{
+			checkf(Rows.Num() <= ScratchBuffer.Num(), 
+				TEXT("The scratch buffer used for sorting a row handle array needs to contain at least %i entries."), Rows.Num());
+			// The row handle array holds 64-bit integers so the values will always be copied, so it doesn't matter whether or not
+			// the buffer is initialized, so use the more straightforward code path.
+			RadixSort64<ERadixSortBufferState::IsInitialized>(Rows.GetData(), ScratchBuffer.GetData(), Rows.Num());
+			bIsSorted = true;
+		}
+	}
+
+	void FRowHandleArray::MakeUnique()
+	{
+		if (!bIsUnique)
+		{
+			Sort();
+
+			RowHandle* Begin = Rows.GetData();
+			RowHandle* End = Rows.GetData() + Rows.Num();
+			RowHandle* RemoveBegin = std::unique(Begin, End);
+			Rows.RemoveAt(std::distance(Rows.GetData(), RemoveBegin), End - RemoveBegin, EAllowShrinking::No);
+
+			bIsUnique = true;
+		}
+	}
+
+	void FRowHandleArray::MakeIntersection()
+	{
+		static_assert(InvalidRowHandle != std::numeric_limits<RowHandle>::max(), 
+			"FRowHandleArray::MakeUnion() needs an unused value for the row handle and currently uses the max value of a RowHandle."
+			"If InvalidRowHandle is ever changed to be this value, the unused value needs to be set to another value.");
+
+		if (!bIsUnique && Rows.Num() > 1)
+		{
+			Sort();
+
+			RowHandle* Insert = Rows.GetData();
+			RowHandle* Front = Rows.GetData();
+			RowHandle* Next = Rows.GetData() + 1;
+			RowHandle* End = Insert + Rows.Num();
+			// Can't use InvalidRowHandle as that might be a value in the array, so pick the largest row handle value as it's highly
+			// unlikely that it'll be used.
+			RowHandle LastInsert = std::numeric_limits<RowHandle>::max();
+
+			for (; Next != End; ++Next, ++Front)
+			{
+				if (*Front == *Next && LastInsert != *Next)
+				{
+					LastInsert = *Next;
+					*Insert++ = *Next;
+				}
+			}
+
+			if (Front != Insert)
+			{
+				Rows.RemoveAt(Insert - Rows.GetData(), End - Insert, EAllowShrinking::No);
+			}
+
+			bIsUnique = true;
+		}
+		else
+		{
+			Rows.Reset();
+			bIsUnique = true;
+			bIsSorted = true;
+		}
+	}
+
+	void FRowHandleArray::SortedMerge(FRowHandleArrayView AddedRows)
+	{
+		if (Rows.IsEmpty())
+		{
+			Rows.AddUninitialized(AddedRows.Num());
+			FMemory::Memcpy(Rows.GetData(), AddedRows.GetData(), AddedRows.NumBytes());
+			return;
+		}
+
+		Sort();
+
+		if (!AddedRows.IsEmpty())
+		{
+			// Check if it's a quick append at the start or end.
+			if (bIsSorted && AddedRows.IsSorted())
+			{
+				if (Rows.Last() <= *AddedRows.begin())
+				{
+					bIsUnique = bIsUnique && AddedRows.IsUnique() && Rows.Last() < *AddedRows.begin();
+					int32 StartIndex = Rows.AddUninitialized(AddedRows.Num());
+					FMemory::Memcpy(Rows.GetData() + StartIndex, AddedRows.GetData(), AddedRows.NumBytes());
+					return;
+				}
+				else if (AddedRows.Last() <= *Rows.begin())
+				{
+					bIsUnique = bIsUnique && AddedRows.IsUnique() && AddedRows.Last() < *Rows.begin();
+					Rows.Insert(AddedRows.GetData(), AddedRows.Num(), 0);
+					return;
+				}
+			}
+
+			SortedMergeInternal(AddedRows);
+		}
+	}
+
+	void FRowHandleArray::SortedMerge(const FRowHandleArray& AddedRows)
+	{
+		checkf(AddedRows.IsSorted(), TEXT("Row handle array provided for sorted merge was not sorted."));
+
+		SortedMerge(AddedRows.GetRows());
+	}
+
+	void FRowHandleArray::SortedMerge(FRowHandleArray&& AddedRows)
+	{
+		if (Rows.IsEmpty())
+		{
+			*this = MoveTemp(AddedRows);
+			return;
+		}
+
+		Sort();
+
+		if (!AddedRows.IsEmpty())
+		{
+			AddedRows.Sort();
+
+			if (Rows.Last() <= *AddedRows.Rows.begin())
+			{
+				bIsUnique = bIsUnique && AddedRows.IsUnique() && Rows.Last() < *AddedRows.Rows.begin();
+				
+				int32 StartIndex = Rows.AddUninitialized(AddedRows.Num());
+				FMemory::Memcpy(Rows.GetData() + StartIndex, AddedRows.Rows.GetData(), AddedRows.Rows.NumBytes());
+			}
+			else if (AddedRows.Rows.Last() <= *Rows.begin())
+			{
+				AddedRows.bIsUnique = bIsUnique && AddedRows.IsUnique() && AddedRows.Rows.Last() < *Rows.begin();
+
+				int32 StartIndex = AddedRows.Rows.AddUninitialized(Rows.Num());
+				FMemory::Memcpy(AddedRows.Rows.GetData() + StartIndex, Rows.GetData(), Rows.NumBytes());
+
+				*this = MoveTemp(AddedRows);
+			}
+			else if (AddedRows.Rows.Num() < Rows.Num())
+			{
+				SortedMergeInternal(AddedRows.GetRows());
+			}
+			else
+			{
+				AddedRows.SortedMergeInternal(GetRows());
+				*this = MoveTemp(AddedRows);
+			}
+		}
+	}
+
+	FRowHandleArrayView FRowHandleArray::GetRows() const
+	{
+		FRowHandleArrayView::EFlags Flags;
+		Flags |= bIsSorted ? FRowHandleArrayView::EFlags::IsSorted : FRowHandleArrayView::EFlags::None;
+		Flags |= bIsUnique ? FRowHandleArrayView::EFlags::IsUnique : FRowHandleArrayView::EFlags::None;
+		return FRowHandleArrayView(Rows, Flags);
+	}
+
+	bool FRowHandleArray::IsSorted() const
+	{
+		return bIsSorted;
+	}
+
+	bool FRowHandleArray::IsUnique() const
+	{
+		return bIsUnique;
+	}
+
+	bool FRowHandleArray::IsEmpty() const
+	{
+		return Rows.IsEmpty();
+	}
+
+	void FRowHandleArray::SortedMergeInternal(FRowHandleArrayView AddedRows)
+	{
+		// There are two paths, one is an in-place addition if the added rows fit within the existing array
+		// and one to create a new array and merge-copy the values over. The latter is done to avoid an additional copy
+		// that would happen during a resize.
+
+		int32 NewSize = Rows.Num() + AddedRows.Num();
+		if (Rows.Max() < NewSize)
+		{
+			TArray<RowHandle> NewArray;
+			NewArray.AddUninitialized(NewSize);
+			RowHandle* RowsInsert = NewArray.GetData();
+
+			const RowHandle* AddedRowsIt = AddedRows.GetData();
+			const RowHandle* AddedRowsEnd = AddedRows.GetData() + AddedRows.Num();
+			const RowHandle* RowsIt = Rows.GetData();
+			const RowHandle* RowsEnd = Rows.GetData() + Rows.Num();
+
+			// Experimented with a few different ranges and landed on 8. 4 and 8 were pretty close, but
+			// 8 edging out just a lit bit. 16 was slightly worse than 4. This will be hardware specific though.
+			while (RowsEnd - RowsIt >= 8 && AddedRowsEnd - AddedRowsIt >= 8)
+			{
+				*RowsInsert++ = *RowsIt < *AddedRowsIt ? *RowsIt++ : *AddedRowsIt++;
+				*RowsInsert++ = *RowsIt < *AddedRowsIt ? *RowsIt++ : *AddedRowsIt++;
+				*RowsInsert++ = *RowsIt < *AddedRowsIt ? *RowsIt++ : *AddedRowsIt++;
+				*RowsInsert++ = *RowsIt < *AddedRowsIt ? *RowsIt++ : *AddedRowsIt++;
+
+				*RowsInsert++ = *RowsIt < *AddedRowsIt ? *RowsIt++ : *AddedRowsIt++;
+				*RowsInsert++ = *RowsIt < *AddedRowsIt ? *RowsIt++ : *AddedRowsIt++;
+				*RowsInsert++ = *RowsIt < *AddedRowsIt ? *RowsIt++ : *AddedRowsIt++;
+				*RowsInsert++ = *RowsIt < *AddedRowsIt ? *RowsIt++ : *AddedRowsIt++;
+			}
+			while (RowsIt != RowsEnd && AddedRowsIt != AddedRowsEnd)
+			{
+				*RowsInsert++ = *RowsIt < *AddedRowsIt ? *RowsIt++ : *AddedRowsIt++;
+			}
+
+			if (RowsIt != RowsEnd)
+			{
+				FMemory::Memcpy(RowsInsert, RowsIt, (RowsEnd - RowsIt) * sizeof(RowHandle));
+			}
+			else if (AddedRowsIt != AddedRowsEnd)
+			{
+				FMemory::Memcpy(RowsInsert, AddedRowsIt, (AddedRowsEnd - AddedRowsIt) * sizeof(RowHandle));
+			}
+
+			Rows = MoveTemp(NewArray);
+		}
+		else
+		{
+			const RowHandle* AddedRowsIt = AddedRows.GetData() + AddedRows.Num() - 1;
+			const RowHandle* AddedRowsEnd = AddedRows.GetData() - 1;
+			const RowHandle* RowsIt = Rows.GetData() + Rows.Num() - 1;
+			const RowHandle* RowsEnd = Rows.GetData() - 1;
+
+			Rows.AddUninitialized(AddedRows.Num());
+			RowHandle* RowsInsert = Rows.GetData() + Rows.Num() - 1;
+
+			// Experimented with a few different ranges and landed on 8. 4 and 8 were pretty close, but
+			// 8 edging out just a lit bit. 16 was slightly worse than 4. This will be very hardware specific though.
+			while (RowsIt - RowsEnd >= 8 && AddedRowsIt - AddedRowsEnd >= 8)
+			{
+				*RowsInsert-- = *RowsIt > *AddedRowsIt ? *RowsIt-- : *AddedRowsIt--;
+				*RowsInsert-- = *RowsIt > *AddedRowsIt ? *RowsIt-- : *AddedRowsIt--;
+				*RowsInsert-- = *RowsIt > *AddedRowsIt ? *RowsIt-- : *AddedRowsIt--;
+				*RowsInsert-- = *RowsIt > *AddedRowsIt ? *RowsIt-- : *AddedRowsIt--;
+
+				*RowsInsert-- = *RowsIt > *AddedRowsIt ? *RowsIt-- : *AddedRowsIt--;
+				*RowsInsert-- = *RowsIt > *AddedRowsIt ? *RowsIt-- : *AddedRowsIt--;
+				*RowsInsert-- = *RowsIt > *AddedRowsIt ? *RowsIt-- : *AddedRowsIt--;
+				*RowsInsert-- = *RowsIt > *AddedRowsIt ? *RowsIt-- : *AddedRowsIt--;
+			}
+			while (RowsIt != RowsEnd && AddedRowsIt != AddedRowsEnd)
+			{
+				*RowsInsert-- = *RowsIt > *AddedRowsIt ? *RowsIt-- : *AddedRowsIt--;
+			}
+
+			if (AddedRowsIt != AddedRowsEnd)
+			{
+				// Insert at the front as all original values have been moved out.
+				FMemory::Memcpy(Rows.GetData(), AddedRows.GetData(), (AddedRowsIt - AddedRowsEnd) * sizeof(RowHandle));
+			}
+		}
+		bIsUnique = false; // Assume there was overlap between the existing and merged list.
+	}
+} // namespace UE::Editor::DataStorage
diff --git a/Engine/Source/Runtime/TypedElementFramework/Private/Elements/Framework/TypedElementRowHandleArrayView.cpp b/Engine/Source/Runtime/TypedElementFramework/Private/Elements/Framework/TypedElementRowHandleArrayView.cpp
new file mode 100644
--- /dev/null
+++ b/Engine/Source/Runtime/TypedElementFramework/Private/Elements/Framework/TypedElementRowHandleArrayView.cpp
@@ -0,0 +1,88 @@
+// Copyright Epic Games, Inc. All Rights Reserved.
+
+#include "Elements/Framework/TypedElementRowHandleArrayView.h"
+
+#include "Algo/BinarySearch.h"
+
+namespace UE::Editor::DataStorage
+{
+	FRowHandleArrayView::FRowHandleArrayView(const RowHandle* Rows, int32 RowCount, EFlags Flags)
+		: Rows(Rows)
+		, RowCount(RowCount)
+		, bIsSorted(EnumHasAllFlags(Flags, EFlags::IsSorted))
+		, bIsUnique(EnumHasAllFlags(Flags, EFlags::IsUnique))
+	{}
+
+	FRowHandleArrayView::FRowHandleArrayView(TConstArrayView<RowHandle> Rows, EFlags Flags)
+		: Rows(Rows.GetData())
+		, RowCount(Rows.Num())
+		, bIsSorted(EnumHasAllFlags(Flags, EFlags::IsSorted))
+		, bIsUnique(EnumHasAllFlags(Flags, EFlags::IsUnique))
+	{}
+
+	const RowHandle* FRowHandleArrayView::begin() const
+	{
+		return Rows;
+	}
+
+	const RowHandle* FRowHandleArrayView::end() const
+	{
+		return Rows + RowCount;
+	}
+
+	const RowHandle* FRowHandleArrayView::GetData() const
+	{
+		return Rows;
+	}
+
+	const RowHandle& FRowHandleArrayView::Last() const
+	{
+		checkf(RowCount > 0, TEXT("Attempting to get the last element from an empty row handle array view."));
+		return Rows[RowCount - 1];
+	}
+
+	int32 FRowHandleArrayView::Num() const
+	{
+		return RowCount;
+	}
+
+	int32 FRowHandleArrayView::NumBytes() const
+	{
+		return RowCount * sizeof(RowHandle);
+	}
+
+	bool FRowHandleArrayView::IsSorted() const
+	{
+		return bIsSorted;
+	}
+
+	bool FRowHandleArrayView::IsUnique() const
+	{
+		return bIsUnique;
+	}
+
+	bool FRowHandleArrayView::IsEmpty() const
+	{
+		return Rows == nullptr || RowCount <= 0;
+	}
+
+	bool FRowHandleArrayView::Contains(RowHandle Row) const
+	{
+		if (bIsSorted)
+		{
+			TConstArrayView<RowHandle> Range(Rows, RowCount);
+			return Algo::BinarySearch(Range, Row) != INDEX_NONE;
+		}
+		else
+		{
+			for (RowHandle Current : *this)
+			{
+				if (Current == Row)
+				{
+					return true;
+				}
+			}
+			return false;
+		}
+	}
+} // namespace UE::Editor::DataStorage
diff --git a/Engine/Source/Runtime/TypedElementFramework/Public/Elements/Framework/TypedElementRowHandleArray.h b/Engine/Source/Runtime/TypedElementFramework/Public/Elements/Framework/TypedElementRowHandleArray.h
new file mode 100644
--- /dev/null
+++ b/Engine/Source/Runtime/TypedElementFramework/Public/Elements/Framework/TypedElementRowHandleArray.h
@@ -0,0 +1,85 @@
+// Copyright Epic Games, Inc. All Rights Reserved.
+
+#pragma once
+
+#include "Containers/ContainersFwd.h"
+#include "Elements/Common/TypedElementHandles.h"
+#include "Elements/Framework/TypedElementRowHandleArrayView.h"
+
+namespace UE::Editor::DataStorage
+{
+	/** 
+	 * Array dedicated to handling row handles. This provides various unique and optimized functionality that 
+	 * makes working with a large number row handles easier and more efficient.
+	 */
+	class FRowHandleArray
+	{
+	public:
+		FRowHandleArray() = default;
+		TYPEDELEMENTFRAMEWORK_API explicit FRowHandleArray(FRowHandleArrayView InRows);
+
+		/** Returns the number of rows stored. */
+		TYPEDELEMENTFRAMEWORK_API int32 Num() const;
+
+		/** 
+		 * Reserve space for additional row handles. This allocates memory but doesn't add new row handles. If the requested count is less
+		 * than the current count then nothing happens.
+		 */
+		TYPEDELEMENTFRAMEWORK_API void Reserve(int32 Count);
+
+		/** Reduces the amount of memory used by the array to exactly fit the current number or elements. */
+		TYPEDELEMENTFRAMEWORK_API void Shrink();
+
+		/** Appends the provided row handle at the end of the array. This may result in the array becoming unsorted and/or no longer unique. */
+		TYPEDELEMENTFRAMEWORK_API void Add(RowHandle Row);
+
+		/** Append the provided list of rows to the end of the array. This will make the array unsorted and not unique. */
+		TYPEDELEMENTFRAMEWORK_API void Append(TConstArrayView<RowHandle> AdditionalRows);
+
+		/** 
+		 * Append the provided list of rows to the end of the array. If the current array and the additional rows are both sorted then
+		 * this function will try to append the additional rows and maintain the sorted status and if applicable the uniqueness state,
+		 * otherwise the array will be marked as unsorted and no longer unique.
+		 */
+		TYPEDELEMENTFRAMEWORK_API void Append(FRowHandleArrayView AdditionalRows);
+
+		/** Sorts the row handles in numerical order from lowest to highest. */
+		TYPEDELEMENTFRAMEWORK_API void Sort();
+		/**
+		 * Sorts the row handles in numerical order from lowest to highest. This version is faster as it doesn't need to allocate
+		 * a temporary intermediate array. ScratchBuffer has to be at least as large as the number of rows in the array.
+		 */
+		TYPEDELEMENTFRAMEWORK_API void Sort(TArrayView<RowHandle> ScratchBuffer);
+
+		/** Goes through the array and removes any duplicate entries. If the array isn't sorted it will be sorted first. */
+		TYPEDELEMENTFRAMEWORK_API void MakeUnique();
+		/** 
+		 * Goes through the array and only keeps a single instance of any value that appears at least once. If the array 
+		 * isn't sorted it will be sorted first. After this function completes the array only contains unique values.
+		 */
+		TYPEDELEMENTFRAMEWORK_API void MakeIntersection();
+
+		/** Merges the provided array into the array while remaining sorted. The provided array is expected to be sorted. */
+		TYPEDELEMENTFRAMEWORK_API void SortedMerge(FRowHandleArrayView AddedRows);
+		/** Merges the provided array into the array while remaining sorted. The provided array is expected to be sorted. */
+		TYPEDELEMENTFRAMEWORK_API void SortedMerge(const FRowHandleArray& AddedRows);
+		/** Merges the provided array into the array while remaining sorted. */
+		TYPEDELEMENTFRAMEWORK_API void SortedMerge(FRowHandleArray&& AddedRows);
+
+		/** Constructs a view of the array to provide access to its rows. */
+		TYPEDELEMENTFRAMEWORK_API FRowHandleArrayView GetRows() const;
+		/** Whether or not the contained rows are numerically ordered from smallest to largest. */
+		TYPEDELEMENTFRAMEWORK_API bool IsSorted() const;
+		/** Whether or all rows in the view only appear once. */
+		TYPEDELEMENTFRAMEWORK_API bool IsUnique() const;
+		/** Whether or not there are any values in the view. */
+		TYPEDELEMENTFRAMEWORK_API bool IsEmpty() const;
+
+	private:
+		void SortedMergeInternal(FRowHandleArrayView AddedRows);
+
+		TArray<RowHandle> Rows;
+		bool bIsSorted = true;
+		bool bIsUnique = true;
+	};
+} // namespace UE::Editor::DataStorage
diff --git a/Engine/Source/Runtime/TypedElementFramework/Public/Elements/Framework/TypedElementRowHandleArrayView.h b/Engine/Source/Runtime/TypedElementFramework/Public/Elements/Framework/TypedElementRowHandleArrayView.h
new file mode 100644
--- /dev/null
+++ b/Engine/Source/Runtime/TypedElementFramework/Public/Elements/Framework/TypedElementRowHandleArrayView.h
@@ -0,0 +1,79 @@
+// Copyright Epic Games, Inc. All Rights Reserved.
+
+#pragma once
+
+#include "Containers/ArrayView.h"
+#include "Elements/Common/TypedElementHandles.h"
+#include "Misc/EnumClassFlags.h"
+
+namespace UE::Editor::DataStorage
+{
+	/**
+	 * Provides a view of a list of row handles.
+	 * Note that this view is more restrictive than typical array views. For instance, an typical array view would remain valid after 
+	 * changing values the array, but will invalidate FRowHandleArrayView as the value may not be unique and/or cause the array to no
+	 * longer be sorted.
+	 */
+	class FRowHandleArrayView
+	{
+	public:
+		enum class EFlags
+		{
+			None = 0,
+			IsSorted = 1 << 0, //< The view is sorted by row handle.
+			IsUnique = 1 << 1 //< Every value in the view is guaranteed to appear only once.
+		};
+
+		FRowHandleArrayView() = default;
+		TYPEDELEMENTFRAMEWORK_API FRowHandleArrayView(const RowHandle* Rows, int32 RowCount, EFlags Flags);
+		TYPEDELEMENTFRAMEWORK_API FRowHandleArrayView(TConstArrayView<RowHandle> Rows, EFlags Flags);
+
+		/** 
+		 * Returns a pointer to the first row in the view.
+		 * This function is primarily used for compatibility to C++ features and libraries.
+		 */
+		TYPEDELEMENTFRAMEWORK_API const RowHandle* begin() const;
+		/**
+		 * Returns a pointer the address after the last row in the view. 
+		 * This function is primarily used for compatibility to C++ features and libraries.
+		 */
+		TYPEDELEMENTFRAMEWORK_API const RowHandle* end() const;
+
+		/** Returns the address of row handles in memory.*/
+		TYPEDELEMENTFRAMEWORK_API const RowHandle* GetData() const;
+		/**
+		 * If the view is not empty, this return the last row handle in the view. Calling this function on an empty view will cause an
+		 * assert.
+		 */
+		TYPEDELEMENTFRAMEWORK_API const RowHandle& Last() const;
+		/** The number of rows this view shows. */
+		TYPEDELEMENTFRAMEWORK_API int32 Num() const;
+		/**
+		 * Returns the total number of bytes that are in use by rows. This is may not be the total amount of memory in use by the 
+		 * container this view maps to.
+		 */
+		TYPEDELEMENTFRAMEWORK_API int32 NumBytes() const;
+
+		/** Whether or not the contained rows are numerically ordered from smallest to largest. */
+		TYPEDELEMENTFRAMEWORK_API bool IsSorted() const;
+		/** Whether or all rows in the view only appear once. */
+		TYPEDELEMENTFRAMEWORK_API bool IsUnique() const;
+		/** Whether or not there are any values in the view. */
+		TYPEDELEMENTFRAMEWORK_API bool IsEmpty() const;
+
+		/** 
+		 * Whether or not the view contains the provided row. If the view is sorted a much binary search can be used for a much faster 
+		 * search otherwise a slower linear search is done.
+		 */
+		TYPEDELEMENTFRAMEWORK_API bool Contains(RowHandle Row) const;
+
+	private:
+		const RowHandle* Rows = nullptr;
+		int32 RowCount = 0;
+		bool bIsSorted = false;
+		bool bIsUnique = false;
+	};
+
+	ENUM_CLASS_FLAGS(FRowHandleArrayView::EFlags);
+
+} // namespace UE::Editor::DataStorage
diff --git a/Engine/Source/Runtime/TypedElementFramework/Tests/Elements/Framework/TypedElementRowHandleArrayTests.cpp b/Engine/Source/Runtime/TypedElementFramework/Tests/Elements/Framework/TypedElementRowHandleArrayTests.cpp
new file mode 100644
--- /dev/null
+++ b/Engine/Source/Runtime/TypedElementFramework/Tests/Elements/Framework/TypedElementRowHandleArrayTests.cpp
@@ -0,0 +1,521 @@
+// Copyright Epic Games, Inc. All Rights Reserved.
+
+#if WITH_TESTS
+
+#include "Elements/Framework/TypedElementRowHandleArray.h"
+#include "Tests/TestHarnessAdapter.h"
+
+namespace UE::Editor::DataStorage::Tests::Private
+{
+	bool Compare(const UE::Editor::DataStorage::FRowHandleArray& Array, TConstArrayView<UE::Editor::DataStorage::RowHandle> Rows)
+	{
+		using namespace UE::Editor::DataStorage;
+
+		FRowHandleArrayView StoredRows = Array.GetRows();
+		if (StoredRows.Num() == Rows.Num())
+		{
+			const RowHandle* StoredFront = StoredRows.GetData();
+			const RowHandle* Front = Rows.GetData();
+			int32 Count = Rows.Num();
+
+			for (int32 Counter = 0; Counter < Count; ++Counter)
+			{
+				if (*StoredFront++ != *Front++)
+				{
+					return false;
+				}
+			}
+
+			return true;
+		}
+		else
+		{
+			return false;
+		}
+	}
+
+	static bool VerifyIsSorted(const UE::Editor::DataStorage::FRowHandleArray& Array)
+	{
+		using namespace UE::Editor::DataStorage;
+
+		const RowHandle* Previous = Array.GetRows().GetData();
+		const RowHandle* Front = Previous + 1;
+		const RowHandle* End = Previous + Array.GetRows().Num();
+		while (Front < End)
+		{
+			if (*Previous > *Front)
+			{
+				return false;
+			}
+			++Previous;
+			++Front;
+		}
+		return true;
+	}
+
+	static bool VerifyIsUnique(const UE::Editor::DataStorage::FRowHandleArray& Array)
+	{
+		using namespace UE::Editor::DataStorage;
+
+		const RowHandle* Begin = Array.GetRows().GetData();
+		const RowHandle* Previous = Begin;
+		const RowHandle* Front = Begin + 1;
+		const RowHandle* End = Begin + Array.GetRows().Num();
+		while (Front < End)
+		{
+			if (*Previous == *Front)
+			{
+				return false;
+			}
+			++Previous;
+			++Front;
+		}
+		return true;
+	}
+}
+
+TEST_CASE_NAMED(FRowHandleArray_Tests, "Editor::DataStorage::Row Handle Array (FRowHandleArray)", "[ApplicationContextMask][EngineFilter]")
+{
+	using namespace UE::Editor::DataStorage;
+	using namespace UE::Editor::DataStorage::Tests::Private;
+
+	SECTION("Empty array")
+	{
+		FRowHandleArray Array;
+		CHECK_MESSAGE(TEXT("Default for rows is not empty."), Array.GetRows().IsEmpty());
+		CHECK_MESSAGE(TEXT("Empty arrays should be considered sorted."), Array.IsSorted());
+	}
+
+	SECTION("Add first row")
+	{
+		FRowHandleArray Array;
+		Array.Add(1);
+		CHECK_MESSAGE(TEXT("Array still empty after adding row."), !Array.GetRows().IsEmpty());
+		CHECK_MESSAGE(TEXT("Arrays with one element should be considered sorted."), Array.IsSorted());
+	}
+
+	SECTION("Add second row with higher number")
+	{
+		FRowHandleArray Array;
+		Array.Add(1);
+		Array.Add(2);
+
+		CHECK_MESSAGE(TEXT("Array still empty after adding row."), !Array.GetRows().IsEmpty());
+		CHECK_MESSAGE(TEXT("Arrays should be sorted."), Array.IsSorted());
+	}
+
+	SECTION("Add second row with lower number")
+	{
+		FRowHandleArray Array;
+		Array.Add(2);
+		Array.Add(1);
+
+		CHECK_MESSAGE(TEXT("Array still empty after adding row."), !Array.GetRows().IsEmpty());
+		CHECK_MESSAGE(TEXT("Arrays should not be sorted."), !Array.IsSorted());
+	}
+
+	SECTION("Append sorted list with unique higher values")
+	{
+		FRowHandleArray Array(FRowHandleArrayView({ 1, 2, 3 }, 
+			FRowHandleArrayView::EFlags::IsSorted | FRowHandleArrayView::EFlags::IsUnique));
+
+		FRowHandleArrayView NewValues({ 4, 5, 6 },
+			FRowHandleArrayView::EFlags::IsSorted | FRowHandleArrayView::EFlags::IsUnique);
+
+		Array.Append(NewValues);
+		
+		CHECK_MESSAGE(TEXT("Sorted merge not correct."), Compare(Array, { 1, 2, 3, 4, 5, 6 }));
+		CHECK_MESSAGE(TEXT("Arrays should be sorted."), Array.IsSorted());
+		CHECK_MESSAGE(TEXT("Arrays should be unique."), Array.IsUnique());
+	}
+
+	SECTION("Append sorted list with overlapping higher values")
+	{
+		using namespace UE::Editor::DataStorage::Tests::Private;
+
+		FRowHandleArray Array(FRowHandleArrayView({ 1, 2, 3 },
+			FRowHandleArrayView::EFlags::IsSorted | FRowHandleArrayView::EFlags::IsUnique));
+
+		FRowHandleArrayView NewValues({ 3, 4, 5 },
+			FRowHandleArrayView::EFlags::IsSorted | FRowHandleArrayView::EFlags::IsUnique);
+
+		Array.Append(NewValues);
+
+		CHECK_MESSAGE(TEXT("Sorted merge not correct."), Compare(Array, { 1, 2, 3, 3, 4, 5 }));
+		CHECK_MESSAGE(TEXT("Arrays should be sorted."), Array.IsSorted());
+		CHECK_MESSAGE(TEXT("Arrays should not be unique."), !Array.IsUnique());
+	}
+
+	SECTION("Append sorted list with unsorted higher values")
+	{
+		FRowHandleArray Array(FRowHandleArrayView({ 1, 2, 3 },
+			FRowHandleArrayView::EFlags::IsSorted | FRowHandleArrayView::EFlags::IsUnique));
+
+		FRowHandleArrayView NewValues({ 5, 7, 6 }, FRowHandleArrayView::EFlags::IsUnique);
+
+		Array.Append(NewValues);
+
+		CHECK_MESSAGE(TEXT("Sorted merge not correct."), Compare(Array, { 1, 2, 3, 5, 7, 6 }));
+		CHECK_MESSAGE(TEXT("Arrays should not be sorted."), !Array.IsSorted());
+		CHECK_MESSAGE(TEXT("Arrays should not be unique."), !Array.IsUnique());
+	}
+
+	SECTION("Sort array low bits")
+	{
+		FRowHandleArray Array(FRowHandleArrayView({ 3, 1, 2 }, FRowHandleArrayView::EFlags::IsUnique));
+
+		Array.Sort();
+
+		CHECK_MESSAGE(TEXT("List incorrectly sorted."), Compare(Array, {1, 2, 3}));
+	}
+
+	SECTION("Sort array high bits")
+	{
+		FRowHandleArray Array(FRowHandleArrayView(
+			{
+				0xaabbccdd00112233,
+				0xeeff998800112233,
+				0x9988776600112233
+			}, FRowHandleArrayView::EFlags::IsUnique));
+
+		Array.Sort();
+
+		CHECK_MESSAGE(TEXT("List incorrectly sorted."), Compare(Array, 
+			{
+				0x9988776600112233,
+				0xaabbccdd00112233,
+				0xeeff998800112233
+			}));
+	}
+
+	SECTION("Sort array")
+	{
+		FRowHandleArray Array(FRowHandleArrayView(
+			{
+				0xaabbccdd00112233,
+				0xeeff998800887766,
+				0x9988776600443322
+			}, FRowHandleArrayView::EFlags::IsUnique));
+
+		Array.Sort();
+
+		CHECK_MESSAGE(TEXT("List incorrectly sorted."), Compare(Array,
+			{
+				0x9988776600443322,
+				0xaabbccdd00112233,
+				0xeeff998800887766
+			}));
+	}
+
+	SECTION("Make unique for empty array.")
+	{
+		FRowHandleArray Array;
+		Array.MakeUnique();
+		CHECK_MESSAGE(TEXT("Array expected to be unique."), Array.IsUnique());
+	}
+
+	SECTION("Make unique for array with single value.")
+	{
+		FRowHandleArray Array(FRowHandleArrayView({ 1 }, FRowHandleArrayView::EFlags::IsSorted | FRowHandleArrayView::EFlags::IsUnique));
+		Array.MakeUnique();
+		CHECK_MESSAGE(TEXT("Array expected to be unique."), Array.IsUnique());
+		CHECK_MESSAGE(TEXT("Unique array isn't matching expected results."), Compare(Array, { 1 }));
+	}
+
+	SECTION("Make unique for array with single duplicated value.")
+	{
+		FRowHandleArray Array(FRowHandleArrayView({ 1, 1 }, FRowHandleArrayView::EFlags::IsSorted));
+		Array.MakeUnique();
+		CHECK_MESSAGE(TEXT("Array expected to be unique."), Array.IsUnique());
+		CHECK_MESSAGE(TEXT("Unique array isn't matching expected results."), Compare(Array, { 1 }));
+	}
+
+	SECTION("Make unique for array with no duplicates.")
+	{
+		FRowHandleArray Array(FRowHandleArrayView({ 1, 2, 3 }, FRowHandleArrayView::EFlags::IsSorted | FRowHandleArrayView::EFlags::IsUnique));
+		Array.MakeUnique();
+		CHECK_MESSAGE(TEXT("Array expected to be unique."), Array.IsUnique());
+		CHECK_MESSAGE(TEXT("Unique array isn't matching expected results."), Compare(Array, { 1, 2, 3 }));
+	}
+
+	SECTION("Make unique for array with 1 duplicate.")
+	{
+		FRowHandleArray Array(FRowHandleArrayView({ 1, 2, 2, 3 }, FRowHandleArrayView::EFlags::IsSorted));
+		Array.MakeUnique();
+		CHECK_MESSAGE(TEXT("Array expected to be unique."), Array.IsUnique());
+		CHECK_MESSAGE(TEXT("Unique array isn't matching expected results."), Compare(Array, { 1, 2, 3 }));
+	}
+
+	SECTION("Make unique for array with multiple duplicates.")
+	{
+		FRowHandleArray Array(FRowHandleArrayView({ 1, 1, 1, 2, 2, 3, 3, 3, 4, 5, 6, 6, 6 }, FRowHandleArrayView::EFlags::IsSorted));
+		Array.MakeUnique();
+		CHECK_MESSAGE(TEXT("Array expected to be unique."), Array.IsUnique());
+		CHECK_MESSAGE(TEXT("Unique array isn't matching expected results."), Compare(Array, { 1, 2, 3, 4, 5, 6 }));
+	}
+
+	SECTION("Make unique for array with large random values.")
+	{
+		constexpr SIZE_T Count = 2000000;
+		FRandomStream Rand(0x8762ebf2);
+		FRowHandleArray Array;
+		Array.Reserve(Count);
+		for (SIZE_T Counter = 0; Counter < Count; ++Counter)
+		{
+			Array.Add(static_cast<RowHandle>(Rand.RandRange(0, 2048)) << 32 | Rand.RandRange(0, 2048));
+		}
+
+		Array.Sort();
+		Array.MakeUnique();
+		
+		CHECK_MESSAGE(TEXT("Array expected to be sorted."), VerifyIsSorted(Array));
+		CHECK_MESSAGE(TEXT("Array expected to be unique."), VerifyIsUnique(Array));
+	}
+
+	SECTION("Make intersection for empty array.")
+	{
+		FRowHandleArray Array;
+		Array.MakeIntersection();
+		CHECK_MESSAGE(TEXT("Array expected to be unique."), Array.IsUnique());
+		CHECK_MESSAGE(TEXT("Unique array is expected to be empty."), Array.IsEmpty());
+	}
+
+	SECTION("Make intersection for array with one value.")
+	{
+		FRowHandleArray Array(FRowHandleArrayView({ 1 }, FRowHandleArrayView::EFlags::IsSorted));
+		Array.MakeIntersection();
+		CHECK_MESSAGE(TEXT("Array expected to be unique."), Array.IsUnique());
+		CHECK_MESSAGE(TEXT("Unique array is expected to be empty."), Array.IsEmpty());
+	}
+
+	SECTION("Make intersection for array with one duplicate at the start.")
+	{
+		FRowHandleArray Array(FRowHandleArrayView({ 1, 1, 2, 3 }, FRowHandleArrayView::EFlags::IsSorted));
+		Array.MakeIntersection();
+		CHECK_MESSAGE(TEXT("Array expected to be unique."), Array.IsUnique());
+		CHECK_MESSAGE(TEXT("Unique array isn't matching expected results."), Compare(Array, { 1 }));
+	}
+
+	SECTION("Make intersection for array with one duplicate at the end value.")
+	{
+		FRowHandleArray Array(FRowHandleArrayView({ 1, 2, 3, 3 }, FRowHandleArrayView::EFlags::IsSorted));
+		Array.MakeIntersection();
+		CHECK_MESSAGE(TEXT("Array expected to be unique."), Array.IsUnique());
+		CHECK_MESSAGE(TEXT("Unique array isn't matching expected results."), Compare(Array, { 3 }));
+	}
+
+	SECTION("Make intersection for array with several double and unique values.")
+	{
+		FRowHandleArray Array(FRowHandleArrayView({ 1, 2, 2, 3, 4, 4, 5 }, FRowHandleArrayView::EFlags::IsSorted));
+		Array.MakeIntersection();
+		CHECK_MESSAGE(TEXT("Array expected to be unique."), Array.IsUnique());
+		CHECK_MESSAGE(TEXT("Unique array isn't matching expected results."), Compare(Array, { 2, 4 }));
+	}
+
+
+	SECTION("Make v for array with several double and unique values at the start and end.")
+	{
+		FRowHandleArray Array(FRowHandleArrayView({ 1, 1, 2, 3, 3, 4, 5, 5 }, FRowHandleArrayView::EFlags::IsSorted));
+		Array.MakeIntersection();
+		CHECK_MESSAGE(TEXT("Array expected to be unique."), Array.IsUnique());
+		CHECK_MESSAGE(TEXT("Unique array isn't matching expected results."), Compare(Array, { 1, 3, 5 }));
+	}
+
+	SECTION("Make intersection for array with the same value more than twice.")
+	{
+		FRowHandleArray Array(FRowHandleArrayView({ 1, 2, 2, 2, 2, 3 }, FRowHandleArrayView::EFlags::IsSorted));
+		Array.MakeIntersection();
+		CHECK_MESSAGE(TEXT("Array expected to be unique."), Array.IsUnique());
+		CHECK_MESSAGE(TEXT("Unique array isn't matching expected results."), Compare(Array, { 2 }));
+	}
+
+	SECTION("Make intersection for array with the same value more than twice at the start.")
+	{
+		FRowHandleArray Array(FRowHandleArrayView({ 1, 1, 1, 1, 2, 3 }, FRowHandleArrayView::EFlags::IsSorted));
+		Array.MakeIntersection();
+		CHECK_MESSAGE(TEXT("Array expected to be unique."), Array.IsUnique());
+		CHECK_MESSAGE(TEXT("Unique array isn't matching expected results."), Compare(Array, { 1 }));
+	}
+
+	SECTION("Make intersection for array with the same value more than twice at the end.")
+	{
+		FRowHandleArray Array(FRowHandleArrayView({ 1, 2, 3, 3, 3, 3 }, FRowHandleArrayView::EFlags::IsSorted));
+		Array.MakeIntersection();
+		CHECK_MESSAGE(TEXT("Array expected to be unique."), Array.IsUnique());
+		CHECK_MESSAGE(TEXT("Unique array isn't matching expected results."), Compare(Array, { 3 }));
+	}
+
+	SECTION("Make intersection for array with several values that appear more than twice at the start and end.")
+	{
+		FRowHandleArray Array(FRowHandleArrayView({ 1, 1, 1, 1, 2, 3, 3, 4, 5, 5, 5, 5, 5 }, FRowHandleArrayView::EFlags::IsSorted));
+		Array.MakeIntersection();
+		CHECK_MESSAGE(TEXT("Array expected to be unique."), Array.IsUnique());
+		CHECK_MESSAGE(TEXT("Unique array isn't matching expected results."), Compare(Array, { 1, 3, 5 }));
+	}
+
+	SECTION("Sorted merge, no duplicates, addition fully merged, without space")
+	{
+		FRowHandleArray Array(FRowHandleArrayView({ 2, 4, 6 }, FRowHandleArrayView::EFlags::IsSorted | FRowHandleArrayView::EFlags::IsUnique));
+		Array.Shrink();
+		Array.SortedMerge(FRowHandleArrayView({ 1, 3, 5, 7 }, FRowHandleArrayView::EFlags::IsSorted | FRowHandleArrayView::EFlags::IsUnique));
+		CHECK_MESSAGE(TEXT("Sorted merge not correct."), Compare(Array, { 1, 2, 3, 4, 5, 6, 7 }));
+	}
+
+	SECTION("Sorted merge, no duplicates, addition fully merged, with space")
+	{
+		FRowHandleArray Array(FRowHandleArrayView({ 1, 3, 5 }, FRowHandleArrayView::EFlags::IsSorted | FRowHandleArrayView::EFlags::IsUnique));
+		Array.Reserve(7);
+		Array.SortedMerge(FRowHandleArrayView({ 2, 4, 6, 7 }, FRowHandleArrayView::EFlags::IsSorted | FRowHandleArrayView::EFlags::IsUnique));
+		CHECK_MESSAGE(TEXT("Sorted merge not correct."), Compare(Array, { 1, 2, 3, 4, 5, 6, 7 }));
+	}
+	
+	SECTION("Sorted merge, no duplicates, addition remaining, without space")
+	{
+		FRowHandleArray Array(FRowHandleArrayView({ 1, 3, 5 }, FRowHandleArrayView::EFlags::IsSorted | FRowHandleArrayView::EFlags::IsUnique));
+		Array.Shrink();
+		Array.SortedMerge(FRowHandleArrayView({ 2, 4, 6, 7 }, FRowHandleArrayView::EFlags::IsSorted | FRowHandleArrayView::EFlags::IsUnique));
+		CHECK_MESSAGE(TEXT("Sorted merge not correct."), Compare(Array, { 1, 2, 3, 4, 5, 6, 7 }));
+	}
+
+	SECTION("Sorted merge, no duplicates, addition remaining, with space")
+	{
+		FRowHandleArray Array(FRowHandleArrayView({ 2, 4, 6 }, FRowHandleArrayView::EFlags::IsSorted | FRowHandleArrayView::EFlags::IsUnique));
+		Array.Reserve(7);
+		Array.SortedMerge(FRowHandleArrayView({ 1, 3, 5, 7 }, FRowHandleArrayView::EFlags::IsSorted | FRowHandleArrayView::EFlags::IsUnique));
+		CHECK_MESSAGE(TEXT("Sorted merge not correct."), Compare(Array, { 1, 2, 3, 4, 5, 6, 7 }));
+	}
+
+	SECTION("Sorted merge, with duplicate, without space")
+	{
+		FRowHandleArray Array(FRowHandleArrayView({ 2, 4, 6 }, FRowHandleArrayView::EFlags::IsSorted | FRowHandleArrayView::EFlags::IsUnique));
+		Array.Shrink();
+		Array.SortedMerge(FRowHandleArrayView({ 1, 4, 5, 7 }, FRowHandleArrayView::EFlags::IsSorted | FRowHandleArrayView::EFlags::IsUnique));
+		CHECK_MESSAGE(TEXT("Sorted merge not correct."), Compare(Array, { 1, 2, 4, 4, 5, 6, 7 }));
+	}
+
+	SECTION("Sorted merge, with duplicate, with space")
+	{
+		FRowHandleArray Array(FRowHandleArrayView({ 2, 4, 6 }, FRowHandleArrayView::EFlags::IsSorted | FRowHandleArrayView::EFlags::IsUnique));
+		Array.Reserve(7);
+		Array.SortedMerge(FRowHandleArrayView({ 1, 4, 5, 7 }, FRowHandleArrayView::EFlags::IsSorted | FRowHandleArrayView::EFlags::IsUnique));
+		CHECK_MESSAGE(TEXT("Sorted merge not correct."), Compare(Array, { 1, 2, 4, 4, 5, 6, 7 }));
+	}
+
+	SECTION("Sorted merge, with duplicates on source, without space")
+	{
+		FRowHandleArray Array(FRowHandleArrayView({ 2, 4, 4, 6 }, FRowHandleArrayView::EFlags::IsSorted));
+		Array.Shrink();
+		Array.SortedMerge(FRowHandleArrayView({ 1, 4, 5, 7 }, FRowHandleArrayView::EFlags::IsSorted | FRowHandleArrayView::EFlags::IsUnique));
+		CHECK_MESSAGE(TEXT("Sorted merge not correct."), Compare(Array, { 1, 2, 4, 4, 4, 5, 6, 7 }));
+	}
+
+	SECTION("Sorted merge, with duplicates in source, with space")
+	{
+		FRowHandleArray Array(FRowHandleArrayView({ 2, 4, 4, 6 }, FRowHandleArrayView::EFlags::IsSorted));
+		Array.Reserve(7);
+		Array.SortedMerge(FRowHandleArrayView({ 1, 4, 5, 7 }, FRowHandleArrayView::EFlags::IsSorted | FRowHandleArrayView::EFlags::IsUnique));
+		CHECK_MESSAGE(TEXT("Sorted merge not correct."), Compare(Array, { 1, 2, 4, 4, 4, 5, 6, 7 }));
+	}
+
+	SECTION("Sorted merge, with duplicates on merge, without space")
+	{
+		FRowHandleArray Array(FRowHandleArrayView({ 2, 4, 6 }, FRowHandleArrayView::EFlags::IsSorted | FRowHandleArrayView::EFlags::IsUnique));
+		Array.Shrink();
+		Array.SortedMerge(FRowHandleArrayView({ 1, 4, 4, 5, 7 }, FRowHandleArrayView::EFlags::IsSorted));
+		CHECK_MESSAGE(TEXT("Sorted merge not correct."), Compare(Array, { 1, 2, 4, 4, 4, 5, 6, 7 }));
+	}
+
+	SECTION("Sorted merge, with duplicates in merge, with space")
+	{
+		FRowHandleArray Array(FRowHandleArrayView({ 2, 4, 6 }, FRowHandleArrayView::EFlags::IsSorted | FRowHandleArrayView::EFlags::IsUnique));
+		Array.Reserve(7);
+		Array.SortedMerge(FRowHandleArrayView({ 1, 4, 4, 5, 7 }, FRowHandleArrayView::EFlags::IsSorted));
+		CHECK_MESSAGE(TEXT("Sorted merge not correct."), Compare(Array, { 1, 2, 4, 4, 4, 5, 6, 7 }));
+	}
+
+	SECTION("Sorted merge, with duplicates in both, without space")
+	{
+		FRowHandleArray Array(FRowHandleArrayView({ 2, 4, 4, 6 }, FRowHandleArrayView::EFlags::IsSorted));
+		Array.Shrink();
+		Array.SortedMerge(FRowHandleArrayView({ 1, 4, 4, 5, 7 }, FRowHandleArrayView::EFlags::IsSorted));
+		CHECK_MESSAGE(TEXT("Sorted merge not correct."), Compare(Array, { 1, 2, 4, 4, 4, 4, 5, 6, 7 }));
+	}
+
+	SECTION("Sorted merge, with duplicates in both, with space")
+	{
+		FRowHandleArray Array(FRowHandleArrayView({ 2, 4, 4, 6 }, FRowHandleArrayView::EFlags::IsSorted));
+		Array.Reserve(7);
+		Array.SortedMerge(FRowHandleArrayView({ 1, 4, 4, 5, 7 }, FRowHandleArrayView::EFlags::IsSorted));
+		CHECK_MESSAGE(TEXT("Sorted merge not correct."), Compare(Array, { 1, 2, 4, 4, 4, 4, 5, 6, 7 }));
+	}
+
+	SECTION("Sorted merge, all duplicates, without space")
+	{
+		FRowHandleArray Array(FRowHandleArrayView({ 2, 4, 6 }, FRowHandleArrayView::EFlags::IsSorted | FRowHandleArrayView::EFlags::IsUnique));
+		Array.Shrink();
+		Array.SortedMerge(FRowHandleArrayView({ 2, 4, 6 }, FRowHandleArrayView::EFlags::IsSorted | FRowHandleArrayView::EFlags::IsUnique));
+		CHECK_MESSAGE(TEXT("Sorted merge not correct."), Compare(Array, { 2, 2, 4, 4, 6, 6 }));
+	}
+
+	SECTION("Sorted merge, all duplicates, with space")
+	{
+		FRowHandleArray Array(FRowHandleArrayView({ 2, 4, 6 }, FRowHandleArrayView::EFlags::IsSorted | FRowHandleArrayView::EFlags::IsUnique));
+		Array.Reserve(6);
+		Array.SortedMerge(FRowHandleArrayView({ 2, 4, 6 }, FRowHandleArrayView::EFlags::IsSorted | FRowHandleArrayView::EFlags::IsUnique));
+		CHECK_MESSAGE(TEXT("Sorted merge not correct."), Compare(Array, { 2, 2, 4, 4, 6, 6 }));
+	}
+
+	SECTION("Sorted merge, random range, without space")
+	{
+		constexpr SIZE_T Count = 128;
+		FRandomStream Rand(0xdeadbeef);
+		TArray<RowHandle> TestRange;
+		TestRange.Reserve(Count);
+		for (SIZE_T Counter = 0; Counter < Count; ++Counter)
+		{
+			TestRange.Add(Rand.GetUnsignedInt());
+		}
+		
+		TConstArrayView<RowHandle> Source(TestRange.GetData(), 52);
+		FRowHandleArray Array(FRowHandleArrayView(Source, FRowHandleArrayView::EFlags::None));
+		Array.Sort();
+		Array.Shrink();
+
+		TConstArrayView<RowHandle> Merge(TestRange.GetData() + 52, Count - 52);
+		FRowHandleArray MergeArray(FRowHandleArrayView(Merge, FRowHandleArrayView::EFlags::None));
+		MergeArray.Sort();
+
+		Array.SortedMerge(MergeArray);
+		
+		TestRange.Sort();
+		CHECK_MESSAGE(TEXT("Sorted merge not correct."), Compare(Array, TestRange));
+	}
+
+	SECTION("Sorted merge, random range, with space")
+	{
+		constexpr SIZE_T Count = 128;
+		FRandomStream Rand(0xdeadbeef);
+		TArray<RowHandle> TestRange;
+		TestRange.Reserve(Count);
+		for (SIZE_T Counter = 0; Counter < Count; ++Counter)
+		{
+			TestRange.Add(Rand.GetUnsignedInt());
+		}
+		
+		TConstArrayView<RowHandle> Source(TestRange.GetData(), 52);
+		FRowHandleArray Array(FRowHandleArrayView(Source, FRowHandleArrayView::EFlags::None));
+		Array.Sort();
+		Array.Reserve(Count);
+
+		TConstArrayView<RowHandle> Merge(TestRange.GetData() + 52, Count - 52);
+		FRowHandleArray MergeArray(FRowHandleArrayView(Merge, FRowHandleArrayView::EFlags::None));
+		MergeArray.Sort();
+		
+		Array.SortedMerge(MergeArray);
+		
+		TestRange.Sort();
+		CHECK_MESSAGE(TEXT("Sorted merge not correct."), Compare(Array, TestRange));
+	}
+}
+
+#endif // #if WITH_TESTS
