Fix for warning reported by static analysis.

diff --git a/Engine/Source/Runtime/TypedElementFramework/Private/Elements/Framework/TypedElementRowHandleArray.cpp b/Engine/Source/Runtime/TypedElementFramework/Private/Elements/Framework/TypedElementRowHandleArray.cpp
--- a/Engine/Source/Runtime/TypedElementFramework/Private/Elements/Framework/TypedElementRowHandleArray.cpp
+++ b/Engine/Source/Runtime/TypedElementFramework/Private/Elements/Framework/TypedElementRowHandleArray.cpp
@@ -239,7 +239,7 @@ namespace UE::Editor::DataStorage
 
 	FRowHandleArrayView FRowHandleArray::GetRows() const
 	{
-		FRowHandleArrayView::EFlags Flags;
+		FRowHandleArrayView::EFlags Flags = FRowHandleArrayView::EFlags::None;
 		Flags |= bIsSorted ? FRowHandleArrayView::EFlags::IsSorted : FRowHandleArrayView::EFlags::None;
 		Flags |= bIsUnique ? FRowHandleArrayView::EFlags::IsUnique : FRowHandleArrayView::EFlags::None;
 		return FRowHandleArrayView(Rows, Flags);
