Initial step for the TEDS Query Stack.

This change list adds the TEDS Query Stack module. The Query Stack can be used to composite nodes together that each have a unique ability that when combined creates a stack to handle complex row handle manipulation.
The first version includes two nodes:
- Row view node: a node that takes a collection of row handles and allows them to be used in the query stack.
- Row merge node: a node that has zero or more parent nodes that it combines together inot a new list of nodes. Merging can happen in the following ways:
- Append: The parents are appended to each other without any further manipulation.
- Sorted: The parents are combined together into a sorted list.
- Unique: The parents are combined together into a sorted list and only one instance of any row handle is kept.
- Repeated: The parents are combined together into a sorted list and only row handles are kept that appear at least twice in the result.

diff --git a/Engine/Plugins/Experimental/EditorDataStorage/Source/TedsCore/Private/TypedElementDatabaseCommandBuffer.cpp b/Engine/Plugins/Experimental/EditorDataStorage/Source/TedsCore/Private/TypedElementDatabaseCommandBuffer.cpp
--- a/Engine/Plugins/Experimental/EditorDataStorage/Source/TedsCore/Private/TypedElementDatabaseCommandBuffer.cpp
+++ b/Engine/Plugins/Experimental/EditorDataStorage/Source/TedsCore/Private/TypedElementDatabaseCommandBuffer.cpp
@@ -201,19 +201,11 @@ namespace UE::Editor::DataStorage::Legacy
 
 	bool FCommandBuffer::Execute_IsRowAvailable(const FMassEntityManager& MassEntityManager, RowHandle Row)
 	{
-		if (Row == InvalidRowHandle)
-		{
-			return false;
-		}
 		return MassEntityManager.IsEntityValid(FMassEntityHandle::FromNumber(Row));
 	}
 
 	bool FCommandBuffer::Execute_IsRowAssigned(const FMassEntityManager& MassEntityManager, RowHandle Row)
 	{
-		if (Row == InvalidRowHandle)
-		{
-			return false;
-		}
 		return MassEntityManager.IsEntityActive(FMassEntityHandle::FromNumber(Row));
 	}
 
diff --git a/Engine/Plugins/Experimental/EditorDataStorageFeatures/EditorDataStorageFeatures.uplugin b/Engine/Plugins/Experimental/EditorDataStorageFeatures/EditorDataStorageFeatures.uplugin
--- a/Engine/Plugins/Experimental/EditorDataStorageFeatures/EditorDataStorageFeatures.uplugin
+++ b/Engine/Plugins/Experimental/EditorDataStorageFeatures/EditorDataStorageFeatures.uplugin
@@ -16,16 +16,20 @@
 	"Installed": false,
 	"Modules": [
 		{
-			"Name": "TedsRevisionControl",
+			"Name": "TedsAlerts",
 			"Type": "Editor"
 		},
 		{
-			"Name": "TedsAlerts",
-			"Type":  "Editor"
+			"Name": "TedsAssetData",
+			"Type": "Editor"
 		},
 		{
-			"Name": "TedsAssetData",
-			"Type":  "Editor"
+			"Name": "TedsContentBrowser",
+			"Type": "Editor"
+		},
+		{
+			"Name": "TedsDebugger",
+			"Type": "Editor"
 		},
 		{
 			"Name": "TedsOutliner",
@@ -36,23 +40,23 @@
 			"Type": "Editor"
 		},
 		{
-			"Name": "TedsDebugger",
+			"Name": "TedsQueryStack",
 			"Type": "Editor"
 		},
 		{
-			"Name": "TedsSettings",
+			"Name": "TedsRevisionControl",
 			"Type": "Editor"
 		},
 		{
-			"Name": "TedsTableViewer",
+			"Name": "TedsSettings",
 			"Type": "Editor"
 		},
 		{
-			"Name": "TedsContentBrowser",
+			"Name": "TedsStyling",
 			"Type": "Editor"
 		},
 		{
-			"Name": "TedsStyling",
+			"Name": "TedsTableViewer",
 			"Type": "Editor"
 		}
 	],
diff --git a/Engine/Plugins/Experimental/EditorDataStorageFeatures/Source/TedsQueryStack/Private/TedsQueryStackModule.cpp b/Engine/Plugins/Experimental/EditorDataStorageFeatures/Source/TedsQueryStack/Private/TedsQueryStackModule.cpp
new file mode 100644
--- /dev/null
+++ b/Engine/Plugins/Experimental/EditorDataStorageFeatures/Source/TedsQueryStack/Private/TedsQueryStackModule.cpp
@@ -0,0 +1,5 @@
+// Copyright Epic Games, Inc. All Rights Reserved.
+
+#include "Modules/ModuleManager.h"
+
+IMPLEMENT_MODULE(FDefaultModuleImpl, TedsQueryStack)
diff --git a/Engine/Plugins/Experimental/EditorDataStorageFeatures/Source/TedsQueryStack/Private/TedsRowMergeNode.cpp b/Engine/Plugins/Experimental/EditorDataStorageFeatures/Source/TedsQueryStack/Private/TedsRowMergeNode.cpp
new file mode 100644
--- /dev/null
+++ b/Engine/Plugins/Experimental/EditorDataStorageFeatures/Source/TedsQueryStack/Private/TedsRowMergeNode.cpp
@@ -0,0 +1,84 @@
+// Copyright Epic Games, Inc. All Rights Reserved.
+
+#include "TedsRowMergeNode.h"
+#include "Elements/Framework/TypedElementRowHandleArray.h"
+
+namespace UE::Editor::DataStorage::QueryStack
+{
+	FRowMergeNode::FRowMergeNode(TConstArrayView<TSharedPtr<IRowNode>> InParents, EMergeApproach InMergeApproach)
+		: MergeApproach(InMergeApproach)
+	{
+		for (const TSharedPtr<IRowNode>& Parent : InParents)
+		{
+			Parents.Emplace<FParentInfo>({ .Parent = Parent, .Revision = Parent->GetRevision()});
+		}
+		Merge();
+	}
+
+	IRowNode::RevisionId FRowMergeNode::GetRevision() const
+	{
+		return Revision;
+	}
+
+	void FRowMergeNode::Update()
+	{
+		bool bRebuild = false;
+		for (FParentInfo& ParentInfo : Parents)
+		{
+			ParentInfo.Parent->Update();
+			RevisionId ParentRevision = ParentInfo.Parent->GetRevision();
+			if (ParentInfo.Revision != ParentRevision)
+			{
+				bRebuild = true;
+				ParentInfo.Revision = ParentRevision;
+				break;
+			}
+		}
+
+		if (bRebuild)
+		{
+			Merge();
+			Revision++;
+		}
+	}
+
+	FRowHandleArrayView FRowMergeNode::GetRows() const
+	{
+		return Rows.GetRows();
+	}
+
+	void FRowMergeNode::Merge()
+	{
+		Rows.Reset();
+
+		int32 TotalCount = 0;
+		for (FParentInfo& Parent : Parents)
+		{
+			TotalCount += Parent.Parent->GetRows().Num();
+		}
+		Rows.Reserve(TotalCount);
+
+		for (FParentInfo& Parent : Parents)
+		{
+			Rows.Append(Parent.Parent->GetRows());
+		}
+
+		switch (MergeApproach)
+		{
+		case EMergeApproach::Append:
+			break;
+		case EMergeApproach::Sorted:
+			Rows.Sort();
+			break;
+		case EMergeApproach::Unique:
+			Rows.MakeUnique();
+			break;
+		case EMergeApproach::Repeating:
+			Rows.ReduceToDuplicates();
+			break;
+		default:
+			checkf(false, TEXT("Unexpected merge type in FQueryStack_RowMerge."));
+			break;
+		}
+	}
+} // namespace UE::Editor::DataStorage::QueryStack
diff --git a/Engine/Plugins/Experimental/EditorDataStorageFeatures/Source/TedsQueryStack/Private/TedsRowViewNode.cpp b/Engine/Plugins/Experimental/EditorDataStorageFeatures/Source/TedsQueryStack/Private/TedsRowViewNode.cpp
new file mode 100644
--- /dev/null
+++ b/Engine/Plugins/Experimental/EditorDataStorageFeatures/Source/TedsQueryStack/Private/TedsRowViewNode.cpp
@@ -0,0 +1,29 @@
+// Copyright Epic Games, Inc. All Rights Reserved.
+
+#include "TedsRowViewNode.h"
+
+namespace UE::Editor::DataStorage::QueryStack
+{
+	FRowViewNode::FRowViewNode(FRowHandleArrayView Rows)
+		: Rows(Rows)
+	{}
+
+	void FRowViewNode::MarkDirty()
+	{
+		Revision++;
+	}
+
+	INode::RevisionId FRowViewNode::GetRevision() const
+	{
+		return Revision;
+	}
+
+	void FRowViewNode::Update() 
+	{
+	}
+
+	FRowHandleArrayView FRowViewNode::GetRows() const
+	{
+		return Rows;
+	}
+} // namespace UE::Editor::DataStorage::QueryStack
diff --git a/Engine/Plugins/Experimental/EditorDataStorageFeatures/Source/TedsQueryStack/Public/TedsQueryStackInterfaces.h b/Engine/Plugins/Experimental/EditorDataStorageFeatures/Source/TedsQueryStack/Public/TedsQueryStackInterfaces.h
new file mode 100644
--- /dev/null
+++ b/Engine/Plugins/Experimental/EditorDataStorageFeatures/Source/TedsQueryStack/Public/TedsQueryStackInterfaces.h
@@ -0,0 +1,38 @@
+// Copyright Epic Games, Inc. All Rights Reserved.
+
+#pragma once
+
+#include "Containers/ContainersFwd.h"
+#include "Elements/Common/TypedElementHandles.h"
+#include "Elements/Framework/TypedElementRowHandleArrayView.h"
+
+namespace UE::Editor::DataStorage::QueryStack
+{
+	/** Base interface for a all query stack nodes. */
+	class INode
+	{
+	public:
+		using RevisionId = uint32;
+
+		virtual ~INode() = default;
+		
+		virtual RevisionId GetRevision() const = 0;
+		virtual void Update() = 0;
+	};
+
+	/** Query stack node that works on queries handles. These nodes are typically run in some fashion to be turned into a row node. */
+	class IQueryNode : public INode
+	{
+	public:
+		virtual ~IQueryNode() = default;
+		virtual QueryHandle GetQuery() const = 0;
+	};
+
+	/** Query stack node that works on row handles. */
+	class IRowNode : public INode
+	{
+	public:
+		virtual ~IRowNode() = default;
+		virtual FRowHandleArrayView GetRows() const = 0;
+	};
+} // UE::Editor::DataStorage::QueryStack
diff --git a/Engine/Plugins/Experimental/EditorDataStorageFeatures/Source/TedsQueryStack/Public/TedsRowMergeNode.h b/Engine/Plugins/Experimental/EditorDataStorageFeatures/Source/TedsQueryStack/Public/TedsRowMergeNode.h
new file mode 100644
--- /dev/null
+++ b/Engine/Plugins/Experimental/EditorDataStorageFeatures/Source/TedsQueryStack/Public/TedsRowMergeNode.h
@@ -0,0 +1,44 @@
+// Copyright Epic Games, Inc. All Rights Reserved.
+
+#pragma once
+
+#include "Containers/ContainersFwd.h"
+#include "Elements/Framework/TypedElementRowHandleArray.h"
+#include "TedsQueryStackInterfaces.h"
+#include "Templates/SharedPointer.h"
+
+namespace UE::Editor::DataStorage::QueryStack
+{
+	class FRowMergeNode : public IRowNode
+	{
+	public:
+		enum class EMergeApproach
+		{
+			Append, //< The rows in each parent are added at the end of the final array.
+			Sorted, //< The rows in each parent are combined together in a final sorted array.
+			Unique, //< The rows in each parent are combined then sorted and all duplicates are removed from the final array.
+			Repeating, //< The rows in each parent are combined then sorted and only rows that appear in at least 2 parents are kept.
+		};
+
+		TEDSQUERYSTACK_API FRowMergeNode(TConstArrayView<TSharedPtr<IRowNode>> InParents, EMergeApproach InMergeApproach);
+		virtual ~FRowMergeNode() override = default;
+
+		TEDSQUERYSTACK_API virtual RevisionId GetRevision() const override;
+		TEDSQUERYSTACK_API virtual void Update() override;
+		TEDSQUERYSTACK_API virtual FRowHandleArrayView GetRows() const override;
+
+	private:
+		void Merge();
+
+		struct FParentInfo
+		{
+			TSharedPtr<IRowNode> Parent;
+			RevisionId Revision = 0;
+		};
+
+		TArray<FParentInfo> Parents;
+		FRowHandleArray Rows;
+		RevisionId Revision = 0;
+		EMergeApproach MergeApproach;
+	};
+} // namespace UE::Editor::DataStorage::QueryStack
diff --git a/Engine/Plugins/Experimental/EditorDataStorageFeatures/Source/TedsQueryStack/Public/TedsRowViewNode.h b/Engine/Plugins/Experimental/EditorDataStorageFeatures/Source/TedsQueryStack/Public/TedsRowViewNode.h
new file mode 100644
--- /dev/null
+++ b/Engine/Plugins/Experimental/EditorDataStorageFeatures/Source/TedsQueryStack/Public/TedsRowViewNode.h
@@ -0,0 +1,31 @@
+// Copyright Epic Games, Inc. All Rights Reserved.
+
+#pragma once
+
+#include "Containers/ContainersFwd.h"
+#include "TedsQueryStackInterfaces.h"
+
+namespace UE::Editor::DataStorage::QueryStack
+{
+	/**
+	 * Stores a view to a list of rows. The container that view is pointing to needs to be kept alive for
+	 * as long as this Query Stack node is alive.
+	 */
+	class FRowViewNode : public IRowNode
+	{
+	public:
+		TEDSQUERYSTACK_API FRowViewNode() = default;
+		TEDSQUERYSTACK_API explicit FRowViewNode(FRowHandleArrayView Rows);
+		virtual ~FRowViewNode() override = default;
+
+		TEDSQUERYSTACK_API void MarkDirty();
+
+		TEDSQUERYSTACK_API virtual RevisionId GetRevision() const override;
+		TEDSQUERYSTACK_API virtual void Update() override;
+		TEDSQUERYSTACK_API virtual FRowHandleArrayView GetRows() const override;
+
+	private:
+		FRowHandleArrayView Rows;
+		RevisionId Revision = 0;
+	};
+} // namespace UE::Editor::DataStorage::QueryStack
diff --git a/Engine/Plugins/Experimental/EditorDataStorageFeatures/Source/TedsQueryStack/TedsQueryStack.Build.cs b/Engine/Plugins/Experimental/EditorDataStorageFeatures/Source/TedsQueryStack/TedsQueryStack.Build.cs
new file mode 100644
--- /dev/null
+++ b/Engine/Plugins/Experimental/EditorDataStorageFeatures/Source/TedsQueryStack/TedsQueryStack.Build.cs
@@ -0,0 +1,26 @@
+// Copyright Epic Games, Inc. All Rights Reserved.
+
+using UnrealBuildTool;
+
+public class TedsQueryStack : ModuleRules
+{
+	public TedsQueryStack(ReadOnlyTargetRules Target) : base(Target)
+	{
+		PCHUsage = ModuleRules.PCHUsageMode.UseExplicitOrSharedPCHs;
+
+		if (Target.bBuildEditor)
+		{
+			PublicIncludePaths.AddRange(new string[] {});
+			PrivateIncludePaths.AddRange(new string[] {});
+
+			PublicDependencyModuleNames.AddRange(
+				new string[]
+				{
+					"Core",
+					"TypedElementFramework",
+				});
+
+			DynamicallyLoadedModuleNames.AddRange(new string[] {});
+		}
+	}
+}
diff --git a/Engine/Plugins/Experimental/EditorDataStorageFeatures/Source/TedsQueryStack/Tests/TedsRowMergeNodeTests.cpp b/Engine/Plugins/Experimental/EditorDataStorageFeatures/Source/TedsQueryStack/Tests/TedsRowMergeNodeTests.cpp
new file mode 100644
--- /dev/null
+++ b/Engine/Plugins/Experimental/EditorDataStorageFeatures/Source/TedsQueryStack/Tests/TedsRowMergeNodeTests.cpp
@@ -0,0 +1,150 @@
+// Copyright Epic Games, Inc. All Rights Reserved.
+
+#if WITH_TESTS
+
+#include "TedsRowMergeNode.h"
+#include "TedsRowViewNode.h"
+#include "Tests/TestHarnessAdapter.h"
+
+TEST_CASE_NAMED(TEDS_QueryStack_RowMergeNode_Tests, "Editor::DataStorage::QueryStack::FRowMergeNode", "[ApplicationContextMask][EngineFilter]")
+{
+	using namespace UE::Editor::DataStorage;
+	using namespace UE::Editor::DataStorage::QueryStack;
+
+	SECTION("No parents")
+	{
+		FRowMergeNode Node({}, FRowMergeNode::EMergeApproach::Append);
+		CHECK_EQUALS(TEXT("Revision"), Node.GetRevision(), 0);
+		CHECK_MESSAGE(TEXT("Default for rows is not empty."), Node.GetRows().IsEmpty());
+		
+		Node.Update();
+	}
+
+	SECTION("Append two nodes")
+	{
+		FRowHandleArrayView Values = FRowHandleArrayView( { 1, 2, 3 }, 
+			FRowHandleArrayView::EFlags::IsSorted | FRowHandleArrayView::EFlags::IsUnique);
+		TSharedPtr<IRowNode> View0 = MakeShared<FRowViewNode>(Values);
+		TSharedPtr<IRowNode> View1 = MakeShared<FRowViewNode>(Values);
+		
+		FRowMergeNode Node({View0, View1}, FRowMergeNode::EMergeApproach::Append);
+		
+		FRowHandleArrayView Rows = Node.GetRows();
+		CHECK_MESSAGE(TEXT("Expected rows to not be sorted."), !Rows.IsSorted());
+		CHECK_EQUALS(TEXT("Size"), Rows.Num(), 6);
+		CHECK_EQUALS(TEXT("Rows[0]"), Rows[0], 1llu);
+		CHECK_EQUALS(TEXT("Rows[1]"), Rows[1], 2llu);
+		CHECK_EQUALS(TEXT("Rows[2]"), Rows[2], 3llu);
+		CHECK_EQUALS(TEXT("Rows[3]"), Rows[3], 1llu);
+		CHECK_EQUALS(TEXT("Rows[4]"), Rows[4], 2llu);
+		CHECK_EQUALS(TEXT("Rows[5]"), Rows[5], 3llu);
+	}
+
+	SECTION("Merge two nodes sorted")
+	{
+		FRowHandleArrayView Values = FRowHandleArrayView( { 1, 2, 3 },
+			FRowHandleArrayView::EFlags::IsSorted | FRowHandleArrayView::EFlags::IsUnique);
+		TSharedPtr<IRowNode> View0 = MakeShared<FRowViewNode>(Values);
+		TSharedPtr<IRowNode> View1 = MakeShared<FRowViewNode>(Values);
+		
+		FRowMergeNode Node({ View0, View1 }, FRowMergeNode::EMergeApproach::Sorted);
+
+		FRowHandleArrayView Rows = Node.GetRows();
+		CHECK_MESSAGE(TEXT("Expected rows to be sorted."), Rows.IsSorted());
+		CHECK_EQUALS(TEXT("Size"), Rows.Num(), 6);
+		CHECK_EQUALS(TEXT("Rows[0]"), Rows[0], 1llu);
+		CHECK_EQUALS(TEXT("Rows[1]"), Rows[1], 1llu);
+		CHECK_EQUALS(TEXT("Rows[2]"), Rows[2], 2llu);
+		CHECK_EQUALS(TEXT("Rows[3]"), Rows[3], 2llu);
+		CHECK_EQUALS(TEXT("Rows[4]"), Rows[4], 3llu);
+		CHECK_EQUALS(TEXT("Rows[5]"), Rows[5], 3llu);
+	}
+
+	SECTION("Uniquely merge two nodes")
+	{
+		FRowHandleArrayView Values = FRowHandleArrayView( { 1, 2, 3 },
+			FRowHandleArrayView::EFlags::IsSorted | FRowHandleArrayView::EFlags::IsUnique);
+		TSharedPtr<IRowNode> View0 = MakeShared<FRowViewNode>(Values);
+		TSharedPtr<IRowNode> View1 = MakeShared<FRowViewNode>(Values);
+
+		FRowMergeNode Node({ View0, View1 }, FRowMergeNode::EMergeApproach::Unique);
+
+		FRowHandleArrayView Rows = Node.GetRows();
+		CHECK_MESSAGE(TEXT("Expected rows to be sorted."), Rows.IsSorted());
+		CHECK_EQUALS(TEXT("Size"), Rows.Num(), 3);
+		CHECK_EQUALS(TEXT("Rows[0]"), Rows[0], 1llu);
+		CHECK_EQUALS(TEXT("Rows[1]"), Rows[1], 2llu);
+		CHECK_EQUALS(TEXT("Rows[2]"), Rows[2], 3llu);
+	}
+
+	SECTION("Repeating of data in two nodes")
+	{
+		FRowHandleArrayView Values0 = FRowHandleArrayView( { 1, 2, 3 },
+			FRowHandleArrayView::EFlags::IsSorted | FRowHandleArrayView::EFlags::IsUnique);
+		FRowHandleArrayView Values1 = FRowHandleArrayView( { 2, 3, 4 },
+			FRowHandleArrayView::EFlags::IsSorted | FRowHandleArrayView::EFlags::IsUnique);
+		TSharedPtr<IRowNode> View0 = MakeShared<FRowViewNode>(Values0);
+		TSharedPtr<IRowNode> View1 = MakeShared<FRowViewNode>(Values1);
+
+		FRowMergeNode Node({ View0, View1 }, FRowMergeNode::EMergeApproach::Repeating);
+
+		FRowHandleArrayView Rows = Node.GetRows();
+		CHECK_MESSAGE(TEXT("Expected rows to be sorted."), Rows.IsSorted());
+		CHECK_EQUALS(TEXT("Size"), Rows.Num(), 2);
+		CHECK_EQUALS(TEXT("Rows[0]"), Rows[0], 2llu);
+		CHECK_EQUALS(TEXT("Rows[1]"), Rows[1], 3llu);
+	}
+
+	SECTION("Repeating of data in three nodes")
+	{
+		FRowHandleArrayView Values0 = FRowHandleArrayView({ 1, 2, 3 },
+			FRowHandleArrayView::EFlags::IsSorted | FRowHandleArrayView::EFlags::IsUnique);
+		FRowHandleArrayView Values1 = FRowHandleArrayView({ 2, 3, 4 },
+			FRowHandleArrayView::EFlags::IsSorted | FRowHandleArrayView::EFlags::IsUnique);
+		FRowHandleArrayView Values2 = FRowHandleArrayView({ 1, 3 },
+			FRowHandleArrayView::EFlags::IsSorted | FRowHandleArrayView::EFlags::IsUnique);
+		TSharedPtr<IRowNode> View0 = MakeShared<FRowViewNode>(Values0);
+		TSharedPtr<IRowNode> View1 = MakeShared<FRowViewNode>(Values1);
+		TSharedPtr<IRowNode> View2 = MakeShared<FRowViewNode>(Values2);
+
+		FRowMergeNode Node({ View0, View1, View2 }, FRowMergeNode::EMergeApproach::Repeating);
+
+		FRowHandleArrayView Rows = Node.GetRows();
+		CHECK_MESSAGE(TEXT("Expected rows to be sorted."), Rows.IsSorted());
+		CHECK_EQUALS(TEXT("Size"), Rows.Num(), 3);
+		CHECK_EQUALS(TEXT("Rows[0]"), Rows[0], 1llu);
+		CHECK_EQUALS(TEXT("Rows[1]"), Rows[1], 2llu);
+		CHECK_EQUALS(TEXT("Rows[2]"), Rows[2], 3llu);
+	}
+
+	SECTION("Updating multiple nodes")
+	{
+		RowHandle Values[] = {1, 2, 3};
+
+		static constexpr int32 ViewCount = 42;
+		TArray<TSharedPtr<FRowViewNode>> Views;
+		Views.Reserve(ViewCount);
+		TArray<TSharedPtr<IRowNode>> Nodes;
+		Nodes.Reserve(ViewCount);
+		for (int32 Count = 0; Count < ViewCount; ++Count)
+		{
+			TSharedPtr<FRowViewNode> Node = MakeShared<FRowViewNode>(
+				FRowHandleArrayView(Values, FRowHandleArrayView::EFlags::IsSorted | FRowHandleArrayView::EFlags::IsUnique));
+			Views.Add(Node);
+			Nodes.Add(Node);
+		}
+
+		FRowMergeNode Node(Nodes, FRowMergeNode::EMergeApproach::Sorted);
+
+		Views[5]->MarkDirty();
+		Node.Update();
+
+		Views[14]->MarkDirty();
+		Node.Update();
+
+		Views[32]->MarkDirty();
+		Node.Update();
+	}
+}
+
+#endif // #if WITH_TESTS
diff --git a/Engine/Plugins/Experimental/EditorDataStorageFeatures/Source/TedsQueryStack/Tests/TedsRowViewNodeTests.cpp b/Engine/Plugins/Experimental/EditorDataStorageFeatures/Source/TedsQueryStack/Tests/TedsRowViewNodeTests.cpp
new file mode 100644
--- /dev/null
+++ b/Engine/Plugins/Experimental/EditorDataStorageFeatures/Source/TedsQueryStack/Tests/TedsRowViewNodeTests.cpp
@@ -0,0 +1,44 @@
+// Copyright Epic Games, Inc. All Rights Reserved.
+
+#if WITH_TESTS
+
+#include "TedsRowViewNode.h"
+#include "Tests/TestHarnessAdapter.h"
+
+TEST_CASE_NAMED(TEDS_QueryStack_RowViewNode_Tests, "Editor::DataStorage::QueryStack::FRowViewNode", "[ApplicationContextMask][EngineFilter]")
+{
+	using namespace UE::Editor::DataStorage;
+	using namespace UE::Editor::DataStorage::QueryStack;
+
+	SECTION("Empty view")
+	{
+		FRowViewNode View;
+		CHECK_EQUALS(TEXT("Revision"), View.GetRevision(), 0);
+		CHECK_MESSAGE(TEXT("Default for rows is not empty."), View.GetRows().IsEmpty());
+	}
+
+	SECTION("View with a few row handles")
+	{
+		FRowHandleArrayView RowHandles({ 1, 2, 3 }, FRowHandleArrayView::EFlags::IsSorted | FRowHandleArrayView::EFlags::IsUnique);
+		FRowViewNode View(RowHandles);
+		CHECK_EQUALS(TEXT("Revision"), View.GetRevision(), 0);
+		
+		FRowHandleArrayView Rows = View.GetRows();
+		CHECK_MESSAGE(TEXT("Expected rows to be sorted."), Rows.IsSorted());
+		CHECK_EQUALS(TEXT("Size"), Rows.Num(), 3);
+		CHECK_EQUALS(TEXT("Rows[0]"), Rows[0], 1llu);
+		CHECK_EQUALS(TEXT("Rows[1]"), Rows[1], 2llu);
+		CHECK_EQUALS(TEXT("Rows[2]"), Rows[2], 3llu);
+	}
+
+	SECTION("Mark dirty")
+	{
+		FRowViewNode View;
+		CHECK_EQUALS(TEXT("Revision"), View.GetRevision(), 0);
+		
+		View.MarkDirty();
+		CHECK_EQUALS(TEXT("Revision"), View.GetRevision(), 1);
+	}
+}
+
+#endif // #if WITH_TESTS
diff --git a/Engine/Source/Runtime/TypedElementFramework/Private/Elements/Framework/TypedElementRowHandleArray.cpp b/Engine/Source/Runtime/TypedElementFramework/Private/Elements/Framework/TypedElementRowHandleArray.cpp
--- a/Engine/Source/Runtime/TypedElementFramework/Private/Elements/Framework/TypedElementRowHandleArray.cpp
+++ b/Engine/Source/Runtime/TypedElementFramework/Private/Elements/Framework/TypedElementRowHandleArray.cpp
@@ -58,7 +58,7 @@ namespace UE::Editor::DataStorage
 			if (!Rows.IsEmpty())
 			{
 				RowHandle CurrentLast = Rows.Last();
-				RowHandle FirstNew = *AdditionalRows.begin();
+				RowHandle FirstNew = AdditionalRows.First();
 				bIsSorted = bIsSorted && AdditionalRows.IsSorted() && (CurrentLast <= FirstNew);
 				bIsUnique = bIsSorted && bIsUnique && AdditionalRows.IsUnique() && (CurrentLast < FirstNew);
 			}
@@ -73,6 +73,16 @@ namespace UE::Editor::DataStorage
 		}
 	}
 
+	void FRowHandleArray::Reset(int32 NewSize)
+	{
+		Rows.Reset(NewSize);
+	}
+
+	void FRowHandleArray::Empty(int32 Slack)
+	{
+		Rows.Empty(Slack);
+	}
+
 	void FRowHandleArray::Sort()
 	{
 		if (!bIsSorted)
@@ -111,10 +121,10 @@ namespace UE::Editor::DataStorage
 		}
 	}
 
-	void FRowHandleArray::MakeIntersection()
+	void FRowHandleArray::ReduceToDuplicates()
 	{
 		static_assert(InvalidRowHandle != std::numeric_limits<RowHandle>::max(), 
-			"FRowHandleArray::MakeUnion() needs an unused value for the row handle and currently uses the max value of a RowHandle."
+			"FRowHandleArray::MakeIntersection() needs an unused value for the row handle and currently uses the max value of a RowHandle."
 			"If InvalidRowHandle is ever changed to be this value, the unused value needs to be set to another value.");
 
 		if (!bIsUnique && Rows.Num() > 1)
diff --git a/Engine/Source/Runtime/TypedElementFramework/Private/Elements/Framework/TypedElementRowHandleArrayView.cpp b/Engine/Source/Runtime/TypedElementFramework/Private/Elements/Framework/TypedElementRowHandleArrayView.cpp
--- a/Engine/Source/Runtime/TypedElementFramework/Private/Elements/Framework/TypedElementRowHandleArrayView.cpp
+++ b/Engine/Source/Runtime/TypedElementFramework/Private/Elements/Framework/TypedElementRowHandleArrayView.cpp
@@ -30,11 +30,23 @@ namespace UE::Editor::DataStorage
 		return Rows + RowCount;
 	}
 
+	RowHandle FRowHandleArrayView::operator[](uint32 Index) const
+	{
+		checkf(static_cast<int32>(Index) < RowCount && RowCount > 0, TEXT("Index (%u) out of bounds (%i)."), Index, RowCount);
+		return Rows[Index];
+	}
+
 	const RowHandle* FRowHandleArrayView::GetData() const
 	{
 		return Rows;
 	}
 
+	const RowHandle& FRowHandleArrayView::First() const
+	{
+		checkf(RowCount > 0, TEXT("Attempting to get the first element from an empty row handle array view."));
+		return Rows[0];
+	}
+
 	const RowHandle& FRowHandleArrayView::Last() const
 	{
 		checkf(RowCount > 0, TEXT("Attempting to get the last element from an empty row handle array view."));
diff --git a/Engine/Source/Runtime/TypedElementFramework/Public/Elements/Framework/TypedElementRowHandleArray.h b/Engine/Source/Runtime/TypedElementFramework/Public/Elements/Framework/TypedElementRowHandleArray.h
--- a/Engine/Source/Runtime/TypedElementFramework/Public/Elements/Framework/TypedElementRowHandleArray.h
+++ b/Engine/Source/Runtime/TypedElementFramework/Public/Elements/Framework/TypedElementRowHandleArray.h
@@ -43,6 +43,18 @@ namespace UE::Editor::DataStorage
 		 */
 		TYPEDELEMENTFRAMEWORK_API void Append(FRowHandleArrayView AdditionalRows);
 
+		/**
+		 * Removes all row handles, but keeps the allocated memory around. If the provided size is larger than the number of rows
+		 * that fit in the container then extra memory will be allocated so there's enough space for the new number of rows.
+		 */
+		TYPEDELEMENTFRAMEWORK_API void Reset(int32 NewSize = 0);
+
+		/**
+		 * Removes all row handles and releases all allocated memory. If the provided slack size is not zero then memory is
+		 * reserved for at least the provided number of rows.
+		 */
+		TYPEDELEMENTFRAMEWORK_API void Empty(int32 Slack = 0);
+
 		/** Sorts the row handles in numerical order from lowest to highest. */
 		TYPEDELEMENTFRAMEWORK_API void Sort();
 		/**
@@ -54,10 +66,10 @@ namespace UE::Editor::DataStorage
 		/** Goes through the array and removes any duplicate entries. If the array isn't sorted it will be sorted first. */
 		TYPEDELEMENTFRAMEWORK_API void MakeUnique();
 		/** 
-		 * Goes through the array and only keeps a single instance of any value that appears at least once. If the array 
+		 * Goes through the array and only keeps a single instance of any value that appears twice or more. If the array 
 		 * isn't sorted it will be sorted first. After this function completes the array only contains unique values.
 		 */
-		TYPEDELEMENTFRAMEWORK_API void MakeIntersection();
+		TYPEDELEMENTFRAMEWORK_API void ReduceToDuplicates();
 
 		/** Merges the provided array into the array while remaining sorted. The provided array is expected to be sorted. */
 		TYPEDELEMENTFRAMEWORK_API void SortedMerge(FRowHandleArrayView AddedRows);
diff --git a/Engine/Source/Runtime/TypedElementFramework/Public/Elements/Framework/TypedElementRowHandleArrayView.h b/Engine/Source/Runtime/TypedElementFramework/Public/Elements/Framework/TypedElementRowHandleArrayView.h
--- a/Engine/Source/Runtime/TypedElementFramework/Public/Elements/Framework/TypedElementRowHandleArrayView.h
+++ b/Engine/Source/Runtime/TypedElementFramework/Public/Elements/Framework/TypedElementRowHandleArrayView.h
@@ -39,8 +39,15 @@ namespace UE::Editor::DataStorage
 		 */
 		TYPEDELEMENTFRAMEWORK_API const RowHandle* end() const;
 
+		/** Retrieves the row handle at the provided index or asserts if the index is out of bounds. */
+		TYPEDELEMENTFRAMEWORK_API RowHandle operator[](uint32 Index) const;
 		/** Returns the address of row handles in memory.*/
 		TYPEDELEMENTFRAMEWORK_API const RowHandle* GetData() const;
+		/**
+		 * If the view is not empty, this return the last row handle in the view. Calling this function on an empty view will cause an
+		 * assert.
+		 */
+		TYPEDELEMENTFRAMEWORK_API const RowHandle& First() const;
 		/**
 		 * If the view is not empty, this return the last row handle in the view. Calling this function on an empty view will cause an
 		 * assert.
diff --git a/Engine/Source/Runtime/TypedElementFramework/Tests/Elements/Framework/TypedElementRowHandleArrayTests.cpp b/Engine/Source/Runtime/TypedElementFramework/Tests/Elements/Framework/TypedElementRowHandleArrayTests.cpp
--- a/Engine/Source/Runtime/TypedElementFramework/Tests/Elements/Framework/TypedElementRowHandleArrayTests.cpp
+++ b/Engine/Source/Runtime/TypedElementFramework/Tests/Elements/Framework/TypedElementRowHandleArrayTests.cpp
@@ -116,10 +116,12 @@ TEST_CASE_NAMED(FRowHandleArray_Tests, "Editor::DataStorage::Row Handle Array (F
 
 	SECTION("Append sorted list with unique higher values")
 	{
-		FRowHandleArray Array(FRowHandleArrayView({ 1, 2, 3 }, 
+		RowHandle ArrayData0[] = { 1, 2, 3 };
+		FRowHandleArray Array(FRowHandleArrayView(ArrayData0, 
 			FRowHandleArrayView::EFlags::IsSorted | FRowHandleArrayView::EFlags::IsUnique));
 
-		FRowHandleArrayView NewValues({ 4, 5, 6 },
+		RowHandle ArrayData1[] = { 4, 5, 6 };
+		FRowHandleArrayView NewValues(ArrayData1,
 			FRowHandleArrayView::EFlags::IsSorted | FRowHandleArrayView::EFlags::IsUnique);
 
 		Array.Append(NewValues);
@@ -133,10 +135,12 @@ TEST_CASE_NAMED(FRowHandleArray_Tests, "Editor::DataStorage::Row Handle Array (F
 	{
 		using namespace UE::Editor::DataStorage::Tests::Private;
 
-		FRowHandleArray Array(FRowHandleArrayView({ 1, 2, 3 },
+		RowHandle ArrayData0[] = { 1, 2, 3 };
+		FRowHandleArray Array(FRowHandleArrayView(ArrayData0,
 			FRowHandleArrayView::EFlags::IsSorted | FRowHandleArrayView::EFlags::IsUnique));
 
-		FRowHandleArrayView NewValues({ 3, 4, 5 },
+		RowHandle ArrayData1[] = { 3, 4, 5 };
+		FRowHandleArrayView NewValues(ArrayData1,
 			FRowHandleArrayView::EFlags::IsSorted | FRowHandleArrayView::EFlags::IsUnique);
 
 		Array.Append(NewValues);
@@ -148,10 +152,12 @@ TEST_CASE_NAMED(FRowHandleArray_Tests, "Editor::DataStorage::Row Handle Array (F
 
 	SECTION("Append sorted list with unsorted higher values")
 	{
-		FRowHandleArray Array(FRowHandleArrayView({ 1, 2, 3 },
+		RowHandle ArrayData0[] = { 1, 2, 3 };
+		FRowHandleArray Array(FRowHandleArrayView(ArrayData0,
 			FRowHandleArrayView::EFlags::IsSorted | FRowHandleArrayView::EFlags::IsUnique));
 
-		FRowHandleArrayView NewValues({ 5, 7, 6 }, FRowHandleArrayView::EFlags::IsUnique);
+		RowHandle ArrayData1[] = { 5, 7, 6};
+		FRowHandleArrayView NewValues(ArrayData1, FRowHandleArrayView::EFlags::IsUnique);
 
 		Array.Append(NewValues);
 
@@ -171,12 +177,13 @@ TEST_CASE_NAMED(FRowHandleArray_Tests, "Editor::DataStorage::Row Handle Array (F
 
 	SECTION("Sort array high bits")
 	{
-		FRowHandleArray Array(FRowHandleArrayView(
-			{
-				0xaabbccdd00112233,
-				0xeeff998800112233,
-				0x9988776600112233
-			}, FRowHandleArrayView::EFlags::IsUnique));
+		RowHandle ArrayData[] = 
+		{
+			0xaabbccdd00112233,
+			0xeeff998800112233,
+			0x9988776600112233
+		};
+		FRowHandleArray Array(FRowHandleArrayView(ArrayData, FRowHandleArrayView::EFlags::IsUnique));
 
 		Array.Sort();
 
@@ -190,12 +197,13 @@ TEST_CASE_NAMED(FRowHandleArray_Tests, "Editor::DataStorage::Row Handle Array (F
 
 	SECTION("Sort array")
 	{
-		FRowHandleArray Array(FRowHandleArrayView(
-			{
+		RowHandle ArrayData[] =
+		{
 				0xaabbccdd00112233,
 				0xeeff998800887766,
 				0x9988776600443322
-			}, FRowHandleArrayView::EFlags::IsUnique));
+		};
+		FRowHandleArray Array(FRowHandleArrayView(ArrayData, FRowHandleArrayView::EFlags::IsUnique));
 
 		Array.Sort();
 
@@ -216,7 +224,8 @@ TEST_CASE_NAMED(FRowHandleArray_Tests, "Editor::DataStorage::Row Handle Array (F
 
 	SECTION("Make unique for array with single value.")
 	{
-		FRowHandleArray Array(FRowHandleArrayView({ 1 }, FRowHandleArrayView::EFlags::IsSorted | FRowHandleArrayView::EFlags::IsUnique));
+		RowHandle ArrayData[] = { 1 };
+		FRowHandleArray Array(FRowHandleArrayView(ArrayData, FRowHandleArrayView::EFlags::IsSorted | FRowHandleArrayView::EFlags::IsUnique));
 		Array.MakeUnique();
 		CHECK_MESSAGE(TEXT("Array expected to be unique."), Array.IsUnique());
 		CHECK_MESSAGE(TEXT("Unique array isn't matching expected results."), Compare(Array, { 1 }));
@@ -224,7 +233,8 @@ TEST_CASE_NAMED(FRowHandleArray_Tests, "Editor::DataStorage::Row Handle Array (F
 
 	SECTION("Make unique for array with single duplicated value.")
 	{
-		FRowHandleArray Array(FRowHandleArrayView({ 1, 1 }, FRowHandleArrayView::EFlags::IsSorted));
+		RowHandle ArrayData[] = { 1, 1 };
+		FRowHandleArray Array(FRowHandleArrayView(ArrayData, FRowHandleArrayView::EFlags::IsSorted));
 		Array.MakeUnique();
 		CHECK_MESSAGE(TEXT("Array expected to be unique."), Array.IsUnique());
 		CHECK_MESSAGE(TEXT("Unique array isn't matching expected results."), Compare(Array, { 1 }));
@@ -232,7 +242,8 @@ TEST_CASE_NAMED(FRowHandleArray_Tests, "Editor::DataStorage::Row Handle Array (F
 
 	SECTION("Make unique for array with no duplicates.")
 	{
-		FRowHandleArray Array(FRowHandleArrayView({ 1, 2, 3 }, FRowHandleArrayView::EFlags::IsSorted | FRowHandleArrayView::EFlags::IsUnique));
+		RowHandle ArrayData[] = { 1, 2, 3 };
+		FRowHandleArray Array(FRowHandleArrayView(ArrayData, FRowHandleArrayView::EFlags::IsSorted | FRowHandleArrayView::EFlags::IsUnique));
 		Array.MakeUnique();
 		CHECK_MESSAGE(TEXT("Array expected to be unique."), Array.IsUnique());
 		CHECK_MESSAGE(TEXT("Unique array isn't matching expected results."), Compare(Array, { 1, 2, 3 }));
@@ -240,7 +251,8 @@ TEST_CASE_NAMED(FRowHandleArray_Tests, "Editor::DataStorage::Row Handle Array (F
 
 	SECTION("Make unique for array with 1 duplicate.")
 	{
-		FRowHandleArray Array(FRowHandleArrayView({ 1, 2, 2, 3 }, FRowHandleArrayView::EFlags::IsSorted));
+		RowHandle ArrayData[] = { 1, 2, 2, 3 };
+		FRowHandleArray Array(FRowHandleArrayView(ArrayData, FRowHandleArrayView::EFlags::IsSorted));
 		Array.MakeUnique();
 		CHECK_MESSAGE(TEXT("Array expected to be unique."), Array.IsUnique());
 		CHECK_MESSAGE(TEXT("Unique array isn't matching expected results."), Compare(Array, { 1, 2, 3 }));
@@ -248,7 +260,8 @@ TEST_CASE_NAMED(FRowHandleArray_Tests, "Editor::DataStorage::Row Handle Array (F
 
 	SECTION("Make unique for array with multiple duplicates.")
 	{
-		FRowHandleArray Array(FRowHandleArrayView({ 1, 1, 1, 2, 2, 3, 3, 3, 4, 5, 6, 6, 6 }, FRowHandleArrayView::EFlags::IsSorted));
+		RowHandle ArrayData[] = { 1, 1, 1, 2, 2, 3, 3, 3, 4, 5, 6, 6, 6 };
+		FRowHandleArray Array(FRowHandleArrayView(ArrayData, FRowHandleArrayView::EFlags::IsSorted));
 		Array.MakeUnique();
 		CHECK_MESSAGE(TEXT("Array expected to be unique."), Array.IsUnique());
 		CHECK_MESSAGE(TEXT("Unique array isn't matching expected results."), Compare(Array, { 1, 2, 3, 4, 5, 6 }));
@@ -272,42 +285,46 @@ TEST_CASE_NAMED(FRowHandleArray_Tests, "Editor::DataStorage::Row Handle Array (F
 		CHECK_MESSAGE(TEXT("Array expected to be unique."), VerifyIsUnique(Array));
 	}
 
-	SECTION("Make intersection for empty array.")
+	SECTION("Reduce to duplicates for empty array.")
 	{
 		FRowHandleArray Array;
-		Array.MakeIntersection();
+		Array.ReduceToDuplicates();
 		CHECK_MESSAGE(TEXT("Array expected to be unique."), Array.IsUnique());
 		CHECK_MESSAGE(TEXT("Unique array is expected to be empty."), Array.IsEmpty());
 	}
 
-	SECTION("Make intersection for array with one value.")
+	SECTION("Reduce to duplicates for array with one value.")
 	{
-		FRowHandleArray Array(FRowHandleArrayView({ 1 }, FRowHandleArrayView::EFlags::IsSorted));
-		Array.MakeIntersection();
+		RowHandle ArrayData[] = { 1 };
+		FRowHandleArray Array(FRowHandleArrayView(ArrayData, FRowHandleArrayView::EFlags::IsSorted));
+		Array.ReduceToDuplicates();
 		CHECK_MESSAGE(TEXT("Array expected to be unique."), Array.IsUnique());
 		CHECK_MESSAGE(TEXT("Unique array is expected to be empty."), Array.IsEmpty());
 	}
 
-	SECTION("Make intersection for array with one duplicate at the start.")
+	SECTION("Reduce to duplicates for array with one duplicate at the start.")
 	{
-		FRowHandleArray Array(FRowHandleArrayView({ 1, 1, 2, 3 }, FRowHandleArrayView::EFlags::IsSorted));
-		Array.MakeIntersection();
+		RowHandle ArrayData[] = { 1, 1, 2, 3 };
+		FRowHandleArray Array(FRowHandleArrayView(ArrayData, FRowHandleArrayView::EFlags::IsSorted));
+		Array.ReduceToDuplicates();
 		CHECK_MESSAGE(TEXT("Array expected to be unique."), Array.IsUnique());
 		CHECK_MESSAGE(TEXT("Unique array isn't matching expected results."), Compare(Array, { 1 }));
 	}
 
-	SECTION("Make intersection for array with one duplicate at the end value.")
+	SECTION("Reduce to duplicates for array with one duplicate at the end value.")
 	{
-		FRowHandleArray Array(FRowHandleArrayView({ 1, 2, 3, 3 }, FRowHandleArrayView::EFlags::IsSorted));
-		Array.MakeIntersection();
+		RowHandle ArrayData[] = { 1, 2, 3, 3 };
+		FRowHandleArray Array(FRowHandleArrayView(ArrayData, FRowHandleArrayView::EFlags::IsSorted));
+		Array.ReduceToDuplicates();
 		CHECK_MESSAGE(TEXT("Array expected to be unique."), Array.IsUnique());
 		CHECK_MESSAGE(TEXT("Unique array isn't matching expected results."), Compare(Array, { 3 }));
 	}
 
-	SECTION("Make intersection for array with several double and unique values.")
+	SECTION("Reduce to duplicates for array with several double and unique values.")
 	{
-		FRowHandleArray Array(FRowHandleArrayView({ 1, 2, 2, 3, 4, 4, 5 }, FRowHandleArrayView::EFlags::IsSorted));
-		Array.MakeIntersection();
+		RowHandle ArrayData[] = { 1, 2, 2, 3, 4, 4, 5 };
+		FRowHandleArray Array(FRowHandleArrayView(ArrayData, FRowHandleArrayView::EFlags::IsSorted));
+		Array.ReduceToDuplicates();
 		CHECK_MESSAGE(TEXT("Array expected to be unique."), Array.IsUnique());
 		CHECK_MESSAGE(TEXT("Unique array isn't matching expected results."), Compare(Array, { 2, 4 }));
 	}
@@ -315,153 +332,201 @@ TEST_CASE_NAMED(FRowHandleArray_Tests, "Editor::DataStorage::Row Handle Array (F
 
 	SECTION("Make v for array with several double and unique values at the start and end.")
 	{
-		FRowHandleArray Array(FRowHandleArrayView({ 1, 1, 2, 3, 3, 4, 5, 5 }, FRowHandleArrayView::EFlags::IsSorted));
-		Array.MakeIntersection();
+		RowHandle ArrayData[] = { 1, 1, 2, 3, 3, 4, 5, 5 };
+		FRowHandleArray Array(FRowHandleArrayView(ArrayData, FRowHandleArrayView::EFlags::IsSorted));
+		Array.ReduceToDuplicates();
 		CHECK_MESSAGE(TEXT("Array expected to be unique."), Array.IsUnique());
 		CHECK_MESSAGE(TEXT("Unique array isn't matching expected results."), Compare(Array, { 1, 3, 5 }));
 	}
 
-	SECTION("Make intersection for array with the same value more than twice.")
+	SECTION("Reduce to duplicates for array with the same value more than twice.")
 	{
-		FRowHandleArray Array(FRowHandleArrayView({ 1, 2, 2, 2, 2, 3 }, FRowHandleArrayView::EFlags::IsSorted));
-		Array.MakeIntersection();
+		RowHandle ArrayData[] = { 1, 2, 2, 2, 2, 3 };
+		FRowHandleArray Array(FRowHandleArrayView(ArrayData, FRowHandleArrayView::EFlags::IsSorted));
+		Array.ReduceToDuplicates();
 		CHECK_MESSAGE(TEXT("Array expected to be unique."), Array.IsUnique());
 		CHECK_MESSAGE(TEXT("Unique array isn't matching expected results."), Compare(Array, { 2 }));
 	}
 
-	SECTION("Make intersection for array with the same value more than twice at the start.")
+	SECTION("Reduce to duplicates for array with the same value more than twice at the start.")
 	{
-		FRowHandleArray Array(FRowHandleArrayView({ 1, 1, 1, 1, 2, 3 }, FRowHandleArrayView::EFlags::IsSorted));
-		Array.MakeIntersection();
+		RowHandle ArrayData[] = { 1, 1, 1, 1, 2, 3 };
+		FRowHandleArray Array(FRowHandleArrayView(ArrayData, FRowHandleArrayView::EFlags::IsSorted));
+		Array.ReduceToDuplicates();
 		CHECK_MESSAGE(TEXT("Array expected to be unique."), Array.IsUnique());
 		CHECK_MESSAGE(TEXT("Unique array isn't matching expected results."), Compare(Array, { 1 }));
 	}
 
-	SECTION("Make intersection for array with the same value more than twice at the end.")
+	SECTION("Reduce to duplicates for array with the same value more than twice at the end.")
 	{
-		FRowHandleArray Array(FRowHandleArrayView({ 1, 2, 3, 3, 3, 3 }, FRowHandleArrayView::EFlags::IsSorted));
-		Array.MakeIntersection();
+		RowHandle ArrayData[] = { 1, 2, 3, 3, 3, 3 };
+		FRowHandleArray Array(FRowHandleArrayView(ArrayData, FRowHandleArrayView::EFlags::IsSorted));
+		Array.ReduceToDuplicates();
 		CHECK_MESSAGE(TEXT("Array expected to be unique."), Array.IsUnique());
 		CHECK_MESSAGE(TEXT("Unique array isn't matching expected results."), Compare(Array, { 3 }));
 	}
 
-	SECTION("Make intersection for array with several values that appear more than twice at the start and end.")
+	SECTION("Reduce to duplicates for array with several values that appear more than twice at the start and end.")
 	{
-		FRowHandleArray Array(FRowHandleArrayView({ 1, 1, 1, 1, 2, 3, 3, 4, 5, 5, 5, 5, 5 }, FRowHandleArrayView::EFlags::IsSorted));
-		Array.MakeIntersection();
+		RowHandle ArrayData[] = { 1, 1, 1, 1, 2, 3, 3, 4, 5, 5, 5, 5, 5 };
+		FRowHandleArray Array(FRowHandleArrayView(ArrayData, FRowHandleArrayView::EFlags::IsSorted));
+		Array.ReduceToDuplicates();
 		CHECK_MESSAGE(TEXT("Array expected to be unique."), Array.IsUnique());
 		CHECK_MESSAGE(TEXT("Unique array isn't matching expected results."), Compare(Array, { 1, 3, 5 }));
 	}
 
 	SECTION("Sorted merge, no duplicates, addition fully merged, without space")
 	{
-		FRowHandleArray Array(FRowHandleArrayView({ 2, 4, 6 }, FRowHandleArrayView::EFlags::IsSorted | FRowHandleArrayView::EFlags::IsUnique));
+		RowHandle ArrayData0[] = { 2, 4, 6 };
+		RowHandle ArrayData1[] = { 1, 3, 5, 7 };
+
+		FRowHandleArray Array(FRowHandleArrayView(ArrayData0, FRowHandleArrayView::EFlags::IsSorted | FRowHandleArrayView::EFlags::IsUnique));
 		Array.Shrink();
-		Array.SortedMerge(FRowHandleArrayView({ 1, 3, 5, 7 }, FRowHandleArrayView::EFlags::IsSorted | FRowHandleArrayView::EFlags::IsUnique));
+		Array.SortedMerge(FRowHandleArrayView(ArrayData1, FRowHandleArrayView::EFlags::IsSorted | FRowHandleArrayView::EFlags::IsUnique));
 		CHECK_MESSAGE(TEXT("Sorted merge not correct."), Compare(Array, { 1, 2, 3, 4, 5, 6, 7 }));
 	}
 
 	SECTION("Sorted merge, no duplicates, addition fully merged, with space")
 	{
-		FRowHandleArray Array(FRowHandleArrayView({ 1, 3, 5 }, FRowHandleArrayView::EFlags::IsSorted | FRowHandleArrayView::EFlags::IsUnique));
+		RowHandle ArrayData0[] = { 1, 3, 5 };
+		RowHandle ArrayData1[] = { 2, 4, 6, 7 };
+
+		FRowHandleArray Array(FRowHandleArrayView(ArrayData0, FRowHandleArrayView::EFlags::IsSorted | FRowHandleArrayView::EFlags::IsUnique));
 		Array.Reserve(7);
-		Array.SortedMerge(FRowHandleArrayView({ 2, 4, 6, 7 }, FRowHandleArrayView::EFlags::IsSorted | FRowHandleArrayView::EFlags::IsUnique));
+		Array.SortedMerge(FRowHandleArrayView(ArrayData1, FRowHandleArrayView::EFlags::IsSorted | FRowHandleArrayView::EFlags::IsUnique));
 		CHECK_MESSAGE(TEXT("Sorted merge not correct."), Compare(Array, { 1, 2, 3, 4, 5, 6, 7 }));
 	}
 	
 	SECTION("Sorted merge, no duplicates, addition remaining, without space")
 	{
-		FRowHandleArray Array(FRowHandleArrayView({ 1, 3, 5 }, FRowHandleArrayView::EFlags::IsSorted | FRowHandleArrayView::EFlags::IsUnique));
+		RowHandle ArrayData0[] = { 1, 3, 5 };
+		RowHandle ArrayData1[] = { 2, 4, 6, 7 };
+
+		FRowHandleArray Array(FRowHandleArrayView(ArrayData0, FRowHandleArrayView::EFlags::IsSorted | FRowHandleArrayView::EFlags::IsUnique));
 		Array.Shrink();
-		Array.SortedMerge(FRowHandleArrayView({ 2, 4, 6, 7 }, FRowHandleArrayView::EFlags::IsSorted | FRowHandleArrayView::EFlags::IsUnique));
+		Array.SortedMerge(FRowHandleArrayView(ArrayData1, FRowHandleArrayView::EFlags::IsSorted | FRowHandleArrayView::EFlags::IsUnique));
 		CHECK_MESSAGE(TEXT("Sorted merge not correct."), Compare(Array, { 1, 2, 3, 4, 5, 6, 7 }));
 	}
 
 	SECTION("Sorted merge, no duplicates, addition remaining, with space")
 	{
-		FRowHandleArray Array(FRowHandleArrayView({ 2, 4, 6 }, FRowHandleArrayView::EFlags::IsSorted | FRowHandleArrayView::EFlags::IsUnique));
+		RowHandle ArrayData0[] = { 1, 3, 5 };
+		RowHandle ArrayData1[] = { 2, 4, 6, 7 };
+
+		FRowHandleArray Array(FRowHandleArrayView(ArrayData0, FRowHandleArrayView::EFlags::IsSorted | FRowHandleArrayView::EFlags::IsUnique));
 		Array.Reserve(7);
-		Array.SortedMerge(FRowHandleArrayView({ 1, 3, 5, 7 }, FRowHandleArrayView::EFlags::IsSorted | FRowHandleArrayView::EFlags::IsUnique));
+		Array.SortedMerge(FRowHandleArrayView(ArrayData1, FRowHandleArrayView::EFlags::IsSorted | FRowHandleArrayView::EFlags::IsUnique));
 		CHECK_MESSAGE(TEXT("Sorted merge not correct."), Compare(Array, { 1, 2, 3, 4, 5, 6, 7 }));
 	}
 
 	SECTION("Sorted merge, with duplicate, without space")
 	{
-		FRowHandleArray Array(FRowHandleArrayView({ 2, 4, 6 }, FRowHandleArrayView::EFlags::IsSorted | FRowHandleArrayView::EFlags::IsUnique));
+		RowHandle ArrayData0[] = { 2, 4, 6 };
+		RowHandle ArrayData1[] = { 1, 4, 5, 7 };
+
+		FRowHandleArray Array(FRowHandleArrayView(ArrayData0, FRowHandleArrayView::EFlags::IsSorted | FRowHandleArrayView::EFlags::IsUnique));
 		Array.Shrink();
-		Array.SortedMerge(FRowHandleArrayView({ 1, 4, 5, 7 }, FRowHandleArrayView::EFlags::IsSorted | FRowHandleArrayView::EFlags::IsUnique));
+		Array.SortedMerge(FRowHandleArrayView(ArrayData1, FRowHandleArrayView::EFlags::IsSorted | FRowHandleArrayView::EFlags::IsUnique));
 		CHECK_MESSAGE(TEXT("Sorted merge not correct."), Compare(Array, { 1, 2, 4, 4, 5, 6, 7 }));
 	}
 
 	SECTION("Sorted merge, with duplicate, with space")
 	{
-		FRowHandleArray Array(FRowHandleArrayView({ 2, 4, 6 }, FRowHandleArrayView::EFlags::IsSorted | FRowHandleArrayView::EFlags::IsUnique));
+		RowHandle ArrayData0[] = { 2, 4, 6 };
+		RowHandle ArrayData1[] = { 1, 4, 5, 7 };
+
+		FRowHandleArray Array(FRowHandleArrayView(ArrayData0, FRowHandleArrayView::EFlags::IsSorted | FRowHandleArrayView::EFlags::IsUnique));
 		Array.Reserve(7);
-		Array.SortedMerge(FRowHandleArrayView({ 1, 4, 5, 7 }, FRowHandleArrayView::EFlags::IsSorted | FRowHandleArrayView::EFlags::IsUnique));
+		Array.SortedMerge(FRowHandleArrayView(ArrayData1, FRowHandleArrayView::EFlags::IsSorted | FRowHandleArrayView::EFlags::IsUnique));
 		CHECK_MESSAGE(TEXT("Sorted merge not correct."), Compare(Array, { 1, 2, 4, 4, 5, 6, 7 }));
 	}
 
 	SECTION("Sorted merge, with duplicates on source, without space")
 	{
-		FRowHandleArray Array(FRowHandleArrayView({ 2, 4, 4, 6 }, FRowHandleArrayView::EFlags::IsSorted));
+		RowHandle ArrayData0[] = { 2, 4, 4, 6 };
+		RowHandle ArrayData1[] = { 1, 4, 5, 7 };
+
+		FRowHandleArray Array(FRowHandleArrayView(ArrayData0, FRowHandleArrayView::EFlags::IsSorted));
 		Array.Shrink();
-		Array.SortedMerge(FRowHandleArrayView({ 1, 4, 5, 7 }, FRowHandleArrayView::EFlags::IsSorted | FRowHandleArrayView::EFlags::IsUnique));
+		Array.SortedMerge(FRowHandleArrayView(ArrayData1, FRowHandleArrayView::EFlags::IsSorted | FRowHandleArrayView::EFlags::IsUnique));
 		CHECK_MESSAGE(TEXT("Sorted merge not correct."), Compare(Array, { 1, 2, 4, 4, 4, 5, 6, 7 }));
 	}
 
 	SECTION("Sorted merge, with duplicates in source, with space")
 	{
-		FRowHandleArray Array(FRowHandleArrayView({ 2, 4, 4, 6 }, FRowHandleArrayView::EFlags::IsSorted));
+		RowHandle ArrayData0[] = { 2, 4, 4, 6 };
+		RowHandle ArrayData1[] = { 1, 4, 5, 7 };
+
+		FRowHandleArray Array(FRowHandleArrayView(ArrayData0, FRowHandleArrayView::EFlags::IsSorted));
 		Array.Reserve(7);
-		Array.SortedMerge(FRowHandleArrayView({ 1, 4, 5, 7 }, FRowHandleArrayView::EFlags::IsSorted | FRowHandleArrayView::EFlags::IsUnique));
+		Array.SortedMerge(FRowHandleArrayView(ArrayData1, FRowHandleArrayView::EFlags::IsSorted | FRowHandleArrayView::EFlags::IsUnique));
 		CHECK_MESSAGE(TEXT("Sorted merge not correct."), Compare(Array, { 1, 2, 4, 4, 4, 5, 6, 7 }));
 	}
 
 	SECTION("Sorted merge, with duplicates on merge, without space")
 	{
-		FRowHandleArray Array(FRowHandleArrayView({ 2, 4, 6 }, FRowHandleArrayView::EFlags::IsSorted | FRowHandleArrayView::EFlags::IsUnique));
+		RowHandle ArrayData0[] = { 2, 4, 6 };
+		RowHandle ArrayData1[] = { 1, 4, 4, 5, 7 };
+
+		FRowHandleArray Array(FRowHandleArrayView(ArrayData0, FRowHandleArrayView::EFlags::IsSorted | FRowHandleArrayView::EFlags::IsUnique));
 		Array.Shrink();
-		Array.SortedMerge(FRowHandleArrayView({ 1, 4, 4, 5, 7 }, FRowHandleArrayView::EFlags::IsSorted));
+		Array.SortedMerge(FRowHandleArrayView(ArrayData1, FRowHandleArrayView::EFlags::IsSorted));
 		CHECK_MESSAGE(TEXT("Sorted merge not correct."), Compare(Array, { 1, 2, 4, 4, 4, 5, 6, 7 }));
 	}
 
 	SECTION("Sorted merge, with duplicates in merge, with space")
 	{
-		FRowHandleArray Array(FRowHandleArrayView({ 2, 4, 6 }, FRowHandleArrayView::EFlags::IsSorted | FRowHandleArrayView::EFlags::IsUnique));
+		RowHandle ArrayData0[] = { 2, 4, 6 };
+		RowHandle ArrayData1[] = { 1, 4, 4, 5, 7 };
+
+
+		FRowHandleArray Array(FRowHandleArrayView(ArrayData0, FRowHandleArrayView::EFlags::IsSorted | FRowHandleArrayView::EFlags::IsUnique));
 		Array.Reserve(7);
-		Array.SortedMerge(FRowHandleArrayView({ 1, 4, 4, 5, 7 }, FRowHandleArrayView::EFlags::IsSorted));
+		Array.SortedMerge(FRowHandleArrayView(ArrayData1, FRowHandleArrayView::EFlags::IsSorted));
 		CHECK_MESSAGE(TEXT("Sorted merge not correct."), Compare(Array, { 1, 2, 4, 4, 4, 5, 6, 7 }));
 	}
 
 	SECTION("Sorted merge, with duplicates in both, without space")
 	{
-		FRowHandleArray Array(FRowHandleArrayView({ 2, 4, 4, 6 }, FRowHandleArrayView::EFlags::IsSorted));
+		RowHandle ArrayData0[] = { 2, 4, 4, 6 };
+		RowHandle ArrayData1[] = { 1, 4, 4, 5, 7 };
+
+		FRowHandleArray Array(FRowHandleArrayView(ArrayData0, FRowHandleArrayView::EFlags::IsSorted));
 		Array.Shrink();
-		Array.SortedMerge(FRowHandleArrayView({ 1, 4, 4, 5, 7 }, FRowHandleArrayView::EFlags::IsSorted));
+		Array.SortedMerge(FRowHandleArrayView(ArrayData1, FRowHandleArrayView::EFlags::IsSorted));
 		CHECK_MESSAGE(TEXT("Sorted merge not correct."), Compare(Array, { 1, 2, 4, 4, 4, 4, 5, 6, 7 }));
 	}
 
 	SECTION("Sorted merge, with duplicates in both, with space")
 	{
-		FRowHandleArray Array(FRowHandleArrayView({ 2, 4, 4, 6 }, FRowHandleArrayView::EFlags::IsSorted));
+		RowHandle ArrayData0[] = { 2, 4, 4, 6 };
+		RowHandle ArrayData1[] = { 1, 4, 4, 5, 7 };
+
+		FRowHandleArray Array(FRowHandleArrayView(ArrayData0, FRowHandleArrayView::EFlags::IsSorted));
 		Array.Reserve(7);
-		Array.SortedMerge(FRowHandleArrayView({ 1, 4, 4, 5, 7 }, FRowHandleArrayView::EFlags::IsSorted));
+		Array.SortedMerge(FRowHandleArrayView(ArrayData1, FRowHandleArrayView::EFlags::IsSorted));
 		CHECK_MESSAGE(TEXT("Sorted merge not correct."), Compare(Array, { 1, 2, 4, 4, 4, 4, 5, 6, 7 }));
 	}
 
 	SECTION("Sorted merge, all duplicates, without space")
 	{
-		FRowHandleArray Array(FRowHandleArrayView({ 2, 4, 6 }, FRowHandleArrayView::EFlags::IsSorted | FRowHandleArrayView::EFlags::IsUnique));
+		RowHandle ArrayData0[] = { 2, 4, 6 };
+		RowHandle ArrayData1[] = { 2, 4, 6 };
+		
+		FRowHandleArray Array(FRowHandleArrayView(ArrayData0, FRowHandleArrayView::EFlags::IsSorted | FRowHandleArrayView::EFlags::IsUnique));
 		Array.Shrink();
-		Array.SortedMerge(FRowHandleArrayView({ 2, 4, 6 }, FRowHandleArrayView::EFlags::IsSorted | FRowHandleArrayView::EFlags::IsUnique));
+		Array.SortedMerge(FRowHandleArrayView(ArrayData1, FRowHandleArrayView::EFlags::IsSorted | FRowHandleArrayView::EFlags::IsUnique));
 		CHECK_MESSAGE(TEXT("Sorted merge not correct."), Compare(Array, { 2, 2, 4, 4, 6, 6 }));
 	}
 
 	SECTION("Sorted merge, all duplicates, with space")
 	{
-		FRowHandleArray Array(FRowHandleArrayView({ 2, 4, 6 }, FRowHandleArrayView::EFlags::IsSorted | FRowHandleArrayView::EFlags::IsUnique));
+		RowHandle ArrayData0[] = { 2, 4, 6 };
+		RowHandle ArrayData1[] = { 2, 4, 6 };
+
+		FRowHandleArray Array(FRowHandleArrayView(ArrayData0, FRowHandleArrayView::EFlags::IsSorted | FRowHandleArrayView::EFlags::IsUnique));
 		Array.Reserve(6);
-		Array.SortedMerge(FRowHandleArrayView({ 2, 4, 6 }, FRowHandleArrayView::EFlags::IsSorted | FRowHandleArrayView::EFlags::IsUnique));
+		Array.SortedMerge(FRowHandleArrayView(ArrayData1, FRowHandleArrayView::EFlags::IsSorted | FRowHandleArrayView::EFlags::IsUnique));
 		CHECK_MESSAGE(TEXT("Sorted merge not correct."), Compare(Array, { 2, 2, 4, 4, 6, 6 }));
 	}
 
