[TEDS] Temporarily removing UHT code generation for dynamic columns
* Follow up to fix newly added column declarations for integration purposes

diff --git a/Engine/Plugins/Fab/Source/Fab/Private/Teds/FabMyFolderIntegration.h b/Engine/Plugins/Fab/Source/Fab/Private/Teds/FabMyFolderIntegration.h
--- a/Engine/Plugins/Fab/Source/Fab/Private/Teds/FabMyFolderIntegration.h
+++ b/Engine/Plugins/Fab/Source/Fab/Private/Teds/FabMyFolderIntegration.h
@@ -30,6 +30,8 @@ USTRUCT(meta = (DisplayName = "Fab distribution method", EditorDataStorage_Dynam
 struct FFabDistributionMethodTag : public FEditorDataStorageTag
 {
 	GENERATED_BODY()
+
+	struct EditorDataStorage_DynamicColumnTemplate{};
 };
 
 USTRUCT(meta = (DisplayName = "Fab object"))
diff --git a/Engine/Source/Runtime/TypedElementFramework/Public/Elements/Columns/TypedElementWebColumns.h b/Engine/Source/Runtime/TypedElementFramework/Public/Elements/Columns/TypedElementWebColumns.h
--- a/Engine/Source/Runtime/TypedElementFramework/Public/Elements/Columns/TypedElementWebColumns.h
+++ b/Engine/Source/Runtime/TypedElementFramework/Public/Elements/Columns/TypedElementWebColumns.h
@@ -28,6 +28,8 @@ struct FEditorDataStorageWebImageColumn final : public FEditorDataStorageColumn
 {
 	GENERATED_BODY()
 
+	struct EditorDataStorage_DynamicColumnTemplate{};
+
 	/** URL to the image in string form. Can be used with for instance with WebImage(Cache) to create a local image. */
 	UPROPERTY()
 	FString UrlString;
