[Truncation Warnings] Fix and enable truncation warnings in TypedElementFramework and TypedElementRuntime modules

#rnx

diff --git a/Engine/Source/Runtime/TypedElementFramework/Private/Elements/Common/TypedElementMapKey.cpp b/Engine/Source/Runtime/TypedElementFramework/Private/Elements/Common/TypedElementMapKey.cpp
--- a/Engine/Source/Runtime/TypedElementFramework/Private/Elements/Common/TypedElementMapKey.cpp
+++ b/Engine/Source/Runtime/TypedElementFramework/Private/Elements/Common/TypedElementMapKey.cpp
@@ -19,7 +19,7 @@ namespace UE::Editor::DataStorage
 		}
 		uint64 operator()(const FStringView& Value)
 		{
-			return CityHash64(reinterpret_cast<const char*>(Value.GetData()), Value.NumBytes());
+			return CityHash64(reinterpret_cast<const char*>(Value.GetData()), static_cast<uint32>(Value.NumBytes()));
 		}
 		uint64 operator()(const FName& Value)	{ return Value.ToUnstableInt(); }
 		uint64 operator()(const FName* Value)	{ return Value->ToUnstableInt(); }
diff --git a/Engine/Source/Runtime/TypedElementFramework/Private/Elements/Framework/TypedElementRowHandleArray.cpp b/Engine/Source/Runtime/TypedElementFramework/Private/Elements/Framework/TypedElementRowHandleArray.cpp
--- a/Engine/Source/Runtime/TypedElementFramework/Private/Elements/Framework/TypedElementRowHandleArray.cpp
+++ b/Engine/Source/Runtime/TypedElementFramework/Private/Elements/Framework/TypedElementRowHandleArray.cpp
@@ -115,7 +115,7 @@ namespace UE::Editor::DataStorage
 			RowHandle* Begin = Rows.GetData();
 			RowHandle* End = Rows.GetData() + Rows.Num();
 			RowHandle* RemoveBegin = std::unique(Begin, End);
-			Rows.RemoveAt(std::distance(Rows.GetData(), RemoveBegin), End - RemoveBegin, EAllowShrinking::No);
+			Rows.RemoveAt(static_cast<int32>(std::distance(Rows.GetData(), RemoveBegin)), static_cast<int32>(End - RemoveBegin), EAllowShrinking::No);
 
 			bIsUnique = true;
 		}
@@ -150,7 +150,7 @@ namespace UE::Editor::DataStorage
 
 			if (Front != Insert)
 			{
-				Rows.RemoveAt(Insert - Rows.GetData(), End - Insert, EAllowShrinking::No);
+				Rows.RemoveAt(static_cast<int32>(Insert - Rows.GetData()), static_cast<int32>(End - Insert), EAllowShrinking::No);
 			}
 
 			bIsUnique = true;
diff --git a/Engine/Source/Runtime/TypedElementFramework/TypedElementFramework.Build.cs b/Engine/Source/Runtime/TypedElementFramework/TypedElementFramework.Build.cs
--- a/Engine/Source/Runtime/TypedElementFramework/TypedElementFramework.Build.cs
+++ b/Engine/Source/Runtime/TypedElementFramework/TypedElementFramework.Build.cs
@@ -8,6 +8,9 @@ public class TypedElementFramework : ModuleRules
 {
 	public TypedElementFramework(ReadOnlyTargetRules Target) : base(Target)
 	{
+		// Enable truncation warnings in this module
+		CppCompileWarningSettings.UnsafeTypeCastWarningLevel = WarningLevel.Warning;
+
 		PublicIncludePaths.Add(Path.Combine(ModuleDirectory, "Tests"));
 
 		PrivateDependencyModuleNames.AddAll(
diff --git a/Engine/Source/Runtime/TypedElementRuntime/TypedElementRuntime.Build.cs b/Engine/Source/Runtime/TypedElementRuntime/TypedElementRuntime.Build.cs
--- a/Engine/Source/Runtime/TypedElementRuntime/TypedElementRuntime.Build.cs
+++ b/Engine/Source/Runtime/TypedElementRuntime/TypedElementRuntime.Build.cs
@@ -7,6 +7,9 @@ public class TypedElementRuntime : ModuleRules
 {
 	public TypedElementRuntime(ReadOnlyTargetRules Target) : base(Target)
 	{
+		// Enable truncation warnings in this module
+		CppCompileWarningSettings.UnsafeTypeCastWarningLevel = WarningLevel.Warning;
+
 		PublicDependencyModuleNames.AddRange(
 			new string[]
 			{
