[Truncation Warnings] Use std::distance consistently for pointer maths

#rnx

diff --git a/Engine/Source/Runtime/TypedElementFramework/Private/Elements/Framework/TypedElementRowHandleArray.cpp b/Engine/Source/Runtime/TypedElementFramework/Private/Elements/Framework/TypedElementRowHandleArray.cpp
--- a/Engine/Source/Runtime/TypedElementFramework/Private/Elements/Framework/TypedElementRowHandleArray.cpp
+++ b/Engine/Source/Runtime/TypedElementFramework/Private/Elements/Framework/TypedElementRowHandleArray.cpp
@@ -115,7 +115,7 @@ namespace UE::Editor::DataStorage
 			RowHandle* Begin = Rows.GetData();
 			RowHandle* End = Rows.GetData() + Rows.Num();
 			RowHandle* RemoveBegin = std::unique(Begin, End);
-			Rows.RemoveAt(static_cast<int32>(std::distance(Rows.GetData(), RemoveBegin)), static_cast<int32>(End - RemoveBegin), EAllowShrinking::No);
+			Rows.RemoveAt(static_cast<int32>(std::distance(Rows.GetData(), RemoveBegin)), static_cast<int32>(std::distance(RemoveBegin, End)), EAllowShrinking::No);
 
 			bIsUnique = true;
 		}
@@ -150,7 +150,7 @@ namespace UE::Editor::DataStorage
 
 			if (Front != Insert)
 			{
-				Rows.RemoveAt(static_cast<int32>(Insert - Rows.GetData()), static_cast<int32>(End - Insert), EAllowShrinking::No);
+				Rows.RemoveAt(static_cast<int32>(std::distance(Rows.GetData(), Insert)), static_cast<int32>(std::distance(Insert, End)), EAllowShrinking::No);
 			}
 
 			bIsUnique = true;
