TEDS Column Utils: Add Util functions for dynamic columns

diff --git a/Engine/Plugins/Experimental/EditorDataStorage/Source/TedsCore/Private/DynamicColumnGenerator.cpp b/Engine/Plugins/Experimental/EditorDataStorage/Source/TedsCore/Private/DynamicColumnGenerator.cpp
--- a/Engine/Plugins/Experimental/EditorDataStorage/Source/TedsCore/Private/DynamicColumnGenerator.cpp
+++ b/Engine/Plugins/Experimental/EditorDataStorage/Source/TedsCore/Private/DynamicColumnGenerator.cpp
@@ -5,6 +5,7 @@
 #include "AssetDefinitionDefault.h"
 #include "TypedElementDataStorageSharedColumn.h"
 #include "Elements/Common/TypedElementCommonTypes.h"
+#include "Elements/Framework/TypedElementColumnUtils.h"
 #include "Kismet2/ReloadUtilities.h"
 #include "Queries/TypedElementExtendedQueryStore.h"
 #include "UObject/MetaData.h"
@@ -58,7 +59,7 @@ namespace UE::Editor::DataStorage
 		};
 		auto HasTemplateMetadata = [](const UScriptStruct& Template)
 		{
-			return Template.FindMetaData(TEXT("EditorDataStorage_DynamicColumnTemplate")) != nullptr;
+			return UE::Editor::DataStorage::ColumnUtils::IsDynamicTemplate(&Template);
 		};
 		bool bCanUseTemplate = true;
 		if (!IsValidType(Template))
@@ -133,7 +134,7 @@ namespace UE::Editor::DataStorage
 		}
 	}
 
-	void FDynamicColumnGenerator::ForEachDynamicColumn(const UScriptStruct& Template, TFunctionRef<void(const FDynamicColumnGeneratorInfo&)> Callback)
+	void FDynamicColumnGenerator::ForEachDynamicColumn(const UScriptStruct& Template, TFunctionRef<void(const FDynamicColumnGeneratorInfo&)> Callback) const
 	{
 		UE_MT_SCOPED_READ_ACCESS(AccessDetector);
 
diff --git a/Engine/Plugins/Experimental/EditorDataStorage/Source/TedsCore/Private/DynamicColumnGenerator.h b/Engine/Plugins/Experimental/EditorDataStorage/Source/TedsCore/Private/DynamicColumnGenerator.h
--- a/Engine/Plugins/Experimental/EditorDataStorage/Source/TedsCore/Private/DynamicColumnGenerator.h
+++ b/Engine/Plugins/Experimental/EditorDataStorage/Source/TedsCore/Private/DynamicColumnGenerator.h
@@ -55,7 +55,7 @@ namespace UE::Editor::DataStorage
 		 */
 		FDynamicColumnGeneratorInfo GenerateColumn(const UScriptStruct& Template, const FName& Identifier);
 
-		void ForEachDynamicColumn(const UScriptStruct& Template, TFunctionRef<void(const FDynamicColumnGeneratorInfo&)> Callback);
+		void ForEachDynamicColumn(const UScriptStruct& Template, TFunctionRef<void(const FDynamicColumnGeneratorInfo&)> Callback) const;
 
 		bool IsDynamicTemplate(const UScriptStruct& InCandidate) const;
 		const UScriptStruct* FindByTemplateId(const UScriptStruct& Template, const FName& Identifier) const;
diff --git a/Engine/Source/Runtime/TypedElementFramework/Private/Elements/Common/TypedElementQueryConditions.cpp b/Engine/Source/Runtime/TypedElementFramework/Private/Elements/Common/TypedElementQueryConditions.cpp
--- a/Engine/Source/Runtime/TypedElementFramework/Private/Elements/Common/TypedElementQueryConditions.cpp
+++ b/Engine/Source/Runtime/TypedElementFramework/Private/Elements/Common/TypedElementQueryConditions.cpp
@@ -3,14 +3,12 @@
 #include "Elements/Common/TypedElementQueryConditions.h"
 
 #include "RHIStats.h"
+#include "Elements/Framework/TypedElementColumnUtils.h"
 #include "Elements/Interfaces/TypedElementDataStorageInterface.h"
 #include "UObject/Class.h"
 
 namespace UE::Editor::DataStorage::Queries
 {
-	static const FName NAME_DynamicTemplateMetadata(TEXT("EditorDataStorage_DynamicColumnTemplate"));
-	static const FName NAME_DerivedFromDynamicTemplateMetadata(TEXT("EditorDataStorage_DerivedFromDynamicTemplate"));
-	
 	FEditorStorageQueryConditionCompileContext::FEditorStorageQueryConditionCompileContext(IEditorDataStorageProvider* InDataStorage)
 		: DataStorage(InDataStorage)
 	{
@@ -92,13 +90,12 @@ namespace UE::Editor::DataStorage::Queries
 		TBitArray<TInlineAllocator<512>> DerivedFromDynamicTemplate;
 #if WITH_EDITORONLY_DATA
 		DerivedFromDynamicTemplate.SetNum(AvailableColumns.Num(), false);
-		const FName DerivedFromDynamicTemplateMetadata = NAME_DerivedFromDynamicTemplateMetadata;
 		for (int32 AvailableColumnIndex = 0; AvailableColumnIndex < AvailableColumns.Num(); ++AvailableColumnIndex)
 		{
 			const TWeakObjectPtr<const UScriptStruct>& AvailableColumn = Projection(AvailableColumns[AvailableColumnIndex]);
 			if (const UScriptStruct* AvailableColumnStruct = AvailableColumn.Get())
 			{
-				DerivedFromDynamicTemplate[AvailableColumnIndex] = AvailableColumnStruct->FindMetaData(DerivedFromDynamicTemplateMetadata) != nullptr;
+				DerivedFromDynamicTemplate[AvailableColumnIndex] = ColumnUtils::IsDerivedFromDynamicTemplate(AvailableColumnStruct);
 			}
 		}
 #endif
@@ -389,7 +386,7 @@ namespace UE::Editor::DataStorage::Queries
 			else
 			{
 #if WITH_EDITORONLY_DATA
-				bool bIsDynamicTemplate = Columns[ColumnIndex]->FindMetaData(NAME_DynamicTemplateMetadata) != nullptr;
+				bool bIsDynamicTemplate = ColumnUtils::IsDynamicTemplate(Columns[ColumnIndex].Get());
 				if (bIsDynamicTemplate)
 				{
 					ColumnFlags[ColumnIndex] |= EColumnFlags::DynamicColumnTemplate;
diff --git a/Engine/Source/Runtime/TypedElementFramework/Private/Elements/Framework/TypedElementColumnUtils.cpp b/Engine/Source/Runtime/TypedElementFramework/Private/Elements/Framework/TypedElementColumnUtils.cpp
--- a/Engine/Source/Runtime/TypedElementFramework/Private/Elements/Framework/TypedElementColumnUtils.cpp
+++ b/Engine/Source/Runtime/TypedElementFramework/Private/Elements/Framework/TypedElementColumnUtils.cpp
@@ -5,8 +5,14 @@
 #include "UObject/Class.h"
 #include "UObject/UnrealType.h"
 
-namespace TypedElement::ColumnUtils
+namespace UE::Editor::DataStorage::ColumnUtils
 {
+	namespace Private
+	{
+		static const FName NAME_DynamicTemplateMetadata(TEXT("EditorDataStorage_DynamicColumnTemplate"));
+		static const FName NAME_DerivedFromDynamicTemplateMetadata(TEXT("EditorDataStorage_DerivedFromDynamicTemplate"));
+	}
+	
 	void SetColumnValue(void* ColumnData, const UScriptStruct* ColumnType, FName ArgumentName, const FString& ArgumentValue)
 	{
 		FProperty* Property = ColumnType->FindPropertyByName(ArgumentName);
@@ -27,4 +33,33 @@ namespace TypedElement::ColumnUtils
 			SetColumnValue(ColumnData, ColumnType, Arg.Name, Arg.Value);
 		}
 	}
+
+	bool IsDynamicTemplate(const UScriptStruct* InColumn)
+	{
+#if WITH_EDITORONLY_DATA
+		return InColumn->HasMetaData(Private::NAME_DynamicTemplateMetadata);
+#else
+		return false;
+#endif
+	}
+
+	bool IsDerivedFromDynamicTemplate(const UScriptStruct* InColumn)
+	{
+#if WITH_EDITORONLY_DATA
+		return InColumn->HasMetaData(Private::NAME_DerivedFromDynamicTemplateMetadata);
+#else
+		return false;
+#endif
+	}
+
+	FName GetDynamicColumnIdentifier(const UScriptStruct* InColumn)
+	{
+#if WITH_EDITORONLY_DATA
+		if (const FString* FoundMetaData = InColumn->FindMetaData(Private::NAME_DerivedFromDynamicTemplateMetadata))
+		{
+			return FName(*FoundMetaData);
+		}
+#endif
+		return NAME_None;
+	}
 } // namespace TypedElement::ColumnUtils
\ No newline at end of file
diff --git a/Engine/Source/Runtime/TypedElementFramework/Public/Elements/Framework/TypedElementColumnUtils.h b/Engine/Source/Runtime/TypedElementFramework/Public/Elements/Framework/TypedElementColumnUtils.h
--- a/Engine/Source/Runtime/TypedElementFramework/Public/Elements/Framework/TypedElementColumnUtils.h
+++ b/Engine/Source/Runtime/TypedElementFramework/Public/Elements/Framework/TypedElementColumnUtils.h
@@ -7,7 +7,7 @@
 
 class UScriptStruct;
 
-namespace TypedElement::ColumnUtils
+namespace UE::Editor::DataStorage::ColumnUtils
 {
 	struct Argument
 	{
@@ -42,4 +42,22 @@ namespace TypedElement::ColumnUtils
 	{
 		SetColumnValues(&Column, ColumnType::StaticStruct(), Arguments);
 	}
-} // namespace TypedElement::ColumnUtils
\ No newline at end of file
+
+	/** Utility function to check if a column type is a dynamic template. */
+	TYPEDELEMENTFRAMEWORK_API bool IsDynamicTemplate(const UScriptStruct* InColumn);
+
+	/** Utility function to check if a column type is derived from a dynamic template. */
+	TYPEDELEMENTFRAMEWORK_API bool IsDerivedFromDynamicTemplate(const UScriptStruct* InColumn);
+
+	/** Utility function to get the identifier for a dynamic column, NAME_None if invalid */
+	TYPEDELEMENTFRAMEWORK_API FName GetDynamicColumnIdentifier(const UScriptStruct* InColumn);
+}
+
+// Deprecated old namespace with an alias to the new namespace
+namespace TypedElement
+{
+	namespace UE_DEPRECATED(5.6, "Use UE::Editor::DataStorage::ColumnUtils instead") ColumnUtils
+	{
+		using namespace UE::Editor::DataStorage::ColumnUtils;
+	}
+}
