Removed the type aliasing for RowHandle from TEDS UI's interface.

The aliases were added in order to keep older versions compiling and to avoid a single large review. This is another step in removing the temporary aliases.

diff --git a/Engine/Plugins/Experimental/EditorDataStorage/Source/TedsUI/Private/Widgets/ExportedTextWidget.cpp b/Engine/Plugins/Experimental/EditorDataStorage/Source/TedsUI/Private/Widgets/ExportedTextWidget.cpp
--- a/Engine/Plugins/Experimental/EditorDataStorage/Source/TedsUI/Private/Widgets/ExportedTextWidget.cpp
+++ b/Engine/Plugins/Experimental/EditorDataStorage/Source/TedsUI/Private/Widgets/ExportedTextWidget.cpp
@@ -163,8 +163,8 @@ FString FExportedTextWidgetConstructor::CreateWidgetDisplayName(IEditorDataStora
 	}
 }
 
-TSharedPtr<SWidget> FExportedTextWidgetConstructor::ConstructFinalWidget(RowHandle Row, IEditorDataStorageProvider* DataStorage,
-	IEditorDataStorageUiProvider* DataStorageUi, const UE::Editor::DataStorage::FMetaDataView& Arguments)
+TSharedPtr<SWidget> FExportedTextWidgetConstructor::ConstructFinalWidget(UE::Editor::DataStorage::RowHandle Row, 
+	IEditorDataStorageProvider* DataStorage, IEditorDataStorageUiProvider* DataStorageUi, const UE::Editor::DataStorage::FMetaDataView& Arguments)
 {
 	FTypedElementScriptStructTypeInfoColumn& TypeInfoColumn = *DataStorage->GetColumn<FTypedElementScriptStructTypeInfoColumn>(Row);
 
diff --git a/Engine/Plugins/Experimental/EditorDataStorage/Source/TedsUI/Private/Widgets/LabelWidget.cpp b/Engine/Plugins/Experimental/EditorDataStorage/Source/TedsUI/Private/Widgets/LabelWidget.cpp
--- a/Engine/Plugins/Experimental/EditorDataStorage/Source/TedsUI/Private/Widgets/LabelWidget.cpp
+++ b/Engine/Plugins/Experimental/EditorDataStorage/Source/TedsUI/Private/Widgets/LabelWidget.cpp
@@ -65,8 +65,8 @@ TConstArrayView<const UScriptStruct*> FLabelWidgetConstructor::GetAdditionalColu
 }
 
 TSharedPtr<SWidget> FLabelWidgetConstructor::CreateWidget(IEditorDataStorageProvider* DataStorage,
-	IEditorDataStorageUiProvider* DataStorageUi, RowHandle TargetRow, RowHandle WidgetRow,
-	const UE::Editor::DataStorage::FMetaDataView& Arguments)
+	IEditorDataStorageUiProvider* DataStorageUi, UE::Editor::DataStorage::RowHandle TargetRow, 
+	UE::Editor::DataStorage::RowHandle WidgetRow, const UE::Editor::DataStorage::FMetaDataView& Arguments)
 {
 	UE::Editor::DataStorage::FAttributeBinder Binder(TargetRow, DataStorage);
 
diff --git a/Engine/Plugins/Experimental/EditorDataStorage/Source/TedsUI/Private/Widgets/SlateColorWidget.cpp b/Engine/Plugins/Experimental/EditorDataStorage/Source/TedsUI/Private/Widgets/SlateColorWidget.cpp
--- a/Engine/Plugins/Experimental/EditorDataStorage/Source/TedsUI/Private/Widgets/SlateColorWidget.cpp
+++ b/Engine/Plugins/Experimental/EditorDataStorage/Source/TedsUI/Private/Widgets/SlateColorWidget.cpp
@@ -54,9 +54,8 @@ FSlateColorWidgetConstructor::FSlateColorWidgetConstructor()
 
 
 
-TSharedPtr<SWidget> FSlateColorWidgetConstructor::CreateWidget(IEditorDataStorageProvider* DataStorage,
-                                                               IEditorDataStorageUiProvider* DataStorageUi, RowHandle TargetRow, RowHandle WidgetRow,
-                                                               const UE::Editor::DataStorage::FMetaDataView& Arguments)
+TSharedPtr<SWidget> FSlateColorWidgetConstructor::CreateWidget(IEditorDataStorageProvider* DataStorage, IEditorDataStorageUiProvider* DataStorageUi, 
+	UE::Editor::DataStorage::RowHandle TargetRow, UE::Editor::DataStorage::RowHandle WidgetRow, const UE::Editor::DataStorage::FMetaDataView& Arguments)
 {
 	UE::Editor::DataStorage::FAttributeBinder Binder(TargetRow, DataStorage);
 
diff --git a/Engine/Plugins/Experimental/EditorDataStorage/Source/TedsUI/Private/Widgets/SlateColorWidget.h b/Engine/Plugins/Experimental/EditorDataStorage/Source/TedsUI/Private/Widgets/SlateColorWidget.h
--- a/Engine/Plugins/Experimental/EditorDataStorage/Source/TedsUI/Private/Widgets/SlateColorWidget.h
+++ b/Engine/Plugins/Experimental/EditorDataStorage/Source/TedsUI/Private/Widgets/SlateColorWidget.h
@@ -36,7 +36,7 @@ public:
 	virtual TSharedPtr<SWidget> CreateWidget(
 		IEditorDataStorageProvider* DataStorage,
 		IEditorDataStorageUiProvider* DataStorageUi,
-		RowHandle TargetRow,
-		RowHandle WidgetRow, 
+		UE::Editor::DataStorage::RowHandle TargetRow,
+		UE::Editor::DataStorage::RowHandle WidgetRow, 
 		const UE::Editor::DataStorage::FMetaDataView& Arguments) override;
 };
\ No newline at end of file
diff --git a/Engine/Plugins/Experimental/EditorDataStorage/Source/TedsUI/Public/Widgets/ExportedTextWidget.h b/Engine/Plugins/Experimental/EditorDataStorage/Source/TedsUI/Public/Widgets/ExportedTextWidget.h
--- a/Engine/Plugins/Experimental/EditorDataStorage/Source/TedsUI/Public/Widgets/ExportedTextWidget.h
+++ b/Engine/Plugins/Experimental/EditorDataStorage/Source/TedsUI/Public/Widgets/ExportedTextWidget.h
@@ -40,7 +40,7 @@ public:
 		IEditorDataStorageProvider* DataStorage, UE::Editor::DataStorage::RowHandle Row) const override;
 
 	virtual TSharedPtr<SWidget> ConstructFinalWidget(
-		RowHandle Row, /** The row the widget will be stored in. */
+		UE::Editor::DataStorage::RowHandle Row, /** The row the widget will be stored in. */
 		IEditorDataStorageProvider* DataStorage,
 		IEditorDataStorageUiProvider* DataStorageUi,
 		const UE::Editor::DataStorage::FMetaDataView& Arguments) override;
diff --git a/Engine/Plugins/Experimental/EditorDataStorageFeatures/Source/TedsOutliner/Private/Widgets/OutlinerLabelWidget.cpp b/Engine/Plugins/Experimental/EditorDataStorageFeatures/Source/TedsOutliner/Private/Widgets/OutlinerLabelWidget.cpp
--- a/Engine/Plugins/Experimental/EditorDataStorageFeatures/Source/TedsOutliner/Private/Widgets/OutlinerLabelWidget.cpp
+++ b/Engine/Plugins/Experimental/EditorDataStorageFeatures/Source/TedsOutliner/Private/Widgets/OutlinerLabelWidget.cpp
@@ -119,8 +119,8 @@ FOutlinerLabelWidgetConstructor::FOutlinerLabelWidgetConstructor()
 }
 
 TSharedPtr<SWidget> FOutlinerLabelWidgetConstructor::CreateWidget(IEditorDataStorageProvider* DataStorage,
-	IEditorDataStorageUiProvider* DataStorageUi, RowHandle TargetRow, RowHandle WidgetRow,
-	const UE::Editor::DataStorage::FMetaDataView& Arguments)
+	IEditorDataStorageUiProvider* DataStorageUi, UE::Editor::DataStorage::RowHandle TargetRow, 
+	UE::Editor::DataStorage::RowHandle WidgetRow, const UE::Editor::DataStorage::FMetaDataView& Arguments)
 {
 	using namespace UE::Editor::DataStorage;
 	if(DataStorage->IsRowAvailable(TargetRow))
@@ -173,8 +173,8 @@ TSharedPtr<SWidget> FOutlinerLabelWidgetConstructor::CreateWidget(IEditorDataSto
 }
 
 TSharedRef<SWidget> FOutlinerLabelWidgetConstructor::CreateLabel(IEditorDataStorageProvider* DataStorage,
-	IEditorDataStorageUiProvider* DataStorageUi, RowHandle TargetRow, RowHandle WidgetRow,
-	const UE::Editor::DataStorage::FMetaDataView& Arguments)
+	IEditorDataStorageUiProvider* DataStorageUi, UE::Editor::DataStorage::RowHandle TargetRow, 
+	UE::Editor::DataStorage::RowHandle WidgetRow, const UE::Editor::DataStorage::FMetaDataView& Arguments)
 {
 	using namespace UE::Editor::DataStorage;
 	TSharedRef<SWidget> Result = SNullWidget::NullWidget;
diff --git a/Engine/Plugins/Experimental/EditorDataStorageFeatures/Source/TedsOutliner/Private/Widgets/OutlinerLabelWidget.h b/Engine/Plugins/Experimental/EditorDataStorageFeatures/Source/TedsOutliner/Private/Widgets/OutlinerLabelWidget.h
--- a/Engine/Plugins/Experimental/EditorDataStorageFeatures/Source/TedsOutliner/Private/Widgets/OutlinerLabelWidget.h
+++ b/Engine/Plugins/Experimental/EditorDataStorageFeatures/Source/TedsOutliner/Private/Widgets/OutlinerLabelWidget.h
@@ -44,6 +44,6 @@ public:
 protected:
 
 	TSharedRef<SWidget> CreateLabel(IEditorDataStorageProvider* DataStorage,
-	IEditorDataStorageUiProvider* DataStorageUi, RowHandle TargetRow, RowHandle WidgetRow,
-	const UE::Editor::DataStorage::FMetaDataView& Arguments);
+		IEditorDataStorageUiProvider* DataStorageUi, UE::Editor::DataStorage::RowHandle TargetRow, 
+		UE::Editor::DataStorage::RowHandle WidgetRow, const UE::Editor::DataStorage::FMetaDataView& Arguments);
 };
\ No newline at end of file
diff --git a/Engine/Plugins/Experimental/EditorDataStorageFeatures/Source/TedsTableViewer/Private/Widgets/UObjectLabelWidget.cpp b/Engine/Plugins/Experimental/EditorDataStorageFeatures/Source/TedsTableViewer/Private/Widgets/UObjectLabelWidget.cpp
--- a/Engine/Plugins/Experimental/EditorDataStorageFeatures/Source/TedsTableViewer/Private/Widgets/UObjectLabelWidget.cpp
+++ b/Engine/Plugins/Experimental/EditorDataStorageFeatures/Source/TedsTableViewer/Private/Widgets/UObjectLabelWidget.cpp
@@ -34,8 +34,8 @@ FUObjectLabelWidgetConstructor::FUObjectLabelWidgetConstructor(const UScriptStru
 }
 
 TSharedPtr<SWidget> FUObjectLabelWidgetConstructor::CreateWidget(IEditorDataStorageProvider* DataStorage,
-	IEditorDataStorageUiProvider* DataStorageUi, RowHandle TargetRow, RowHandle WidgetRow,
-	const UE::Editor::DataStorage::FMetaDataView& Arguments)
+	IEditorDataStorageUiProvider* DataStorageUi, UE::Editor::DataStorage::RowHandle TargetRow, 
+	UE::Editor::DataStorage::RowHandle WidgetRow, const UE::Editor::DataStorage::FMetaDataView& Arguments)
 {
 	using namespace UE::Editor::DataStorage;
 	if(DataStorage->IsRowAvailable(TargetRow))
diff --git a/Engine/Source/Runtime/TypedElementFramework/Private/Elements/Interfaces/TypedElementDataStorageUiInterface.cpp b/Engine/Source/Runtime/TypedElementFramework/Private/Elements/Interfaces/TypedElementDataStorageUiInterface.cpp
--- a/Engine/Source/Runtime/TypedElementFramework/Private/Elements/Interfaces/TypedElementDataStorageUiInterface.cpp
+++ b/Engine/Source/Runtime/TypedElementFramework/Private/Elements/Interfaces/TypedElementDataStorageUiInterface.cpp
@@ -54,7 +54,7 @@ TConstArrayView<const UScriptStruct*> FTypedElementWidgetConstructor::GetAdditio
 }
 
 FString FTypedElementWidgetConstructor::CreateWidgetDisplayName(
-	IEditorDataStorageProvider* DataStorage, RowHandle Row) const
+	IEditorDataStorageProvider* DataStorage, UE::Editor::DataStorage::RowHandle Row) const
 {
 	switch (MatchedColumnTypes.Num())
 	{
@@ -106,7 +106,7 @@ FString FTypedElementWidgetConstructor::CreateWidgetDisplayName(
 }
 
 TSharedPtr<SWidget> FTypedElementWidgetConstructor::ConstructFinalWidget(
-	RowHandle Row,
+	UE::Editor::DataStorage::RowHandle Row,
 	IEditorDataStorageProvider* DataStorage,
 	IEditorDataStorageUiProvider* DataStorageUi,
 	const UE::Editor::DataStorage::FMetaDataView& Arguments)
@@ -151,7 +151,7 @@ TSharedPtr<SWidget> FTypedElementWidgetConstructor::ConstructFinalWidget(
 }
 
 TSharedPtr<SWidget> FTypedElementWidgetConstructor::Construct(
-	RowHandle Row,
+	UE::Editor::DataStorage::RowHandle Row,
 	IEditorDataStorageProvider* DataStorage,
 	IEditorDataStorageUiProvider* DataStorageUi,
 	const UE::Editor::DataStorage::FMetaDataView& Arguments)
@@ -189,7 +189,7 @@ TSharedPtr<SWidget> FTypedElementWidgetConstructor::CreateWidget(
 	return CreateWidget(Arguments);
 }
 
-bool FTypedElementWidgetConstructor::SetColumns(IEditorDataStorageProvider* DataStorage, RowHandle Row)
+bool FTypedElementWidgetConstructor::SetColumns(IEditorDataStorageProvider* DataStorage, UE::Editor::DataStorage::RowHandle Row)
 {
 	return true;
 }
@@ -214,13 +214,13 @@ FString FTypedElementWidgetConstructor::DescribeColumnType(const UScriptStruct*
 bool FTypedElementWidgetConstructor::FinalizeWidget(
 	IEditorDataStorageProvider* DataStorage,
 	IEditorDataStorageUiProvider* DataStorageUi,
-	RowHandle Row,
+	UE::Editor::DataStorage::RowHandle Row,
 	const TSharedPtr<SWidget>& Widget)
 {
 	return true;
 }
 
-void FTypedElementWidgetConstructor::AddDefaultWidgetColumns(RowHandle Row, IEditorDataStorageProvider* DataStorage) const
+void FTypedElementWidgetConstructor::AddDefaultWidgetColumns(UE::Editor::DataStorage::RowHandle Row, IEditorDataStorageProvider* DataStorage) const
 {
 	const FString WidgetLabel(CreateWidgetDisplayName(DataStorage, Row));
 	DataStorage->AddColumn(Row, FTypedElementLabelColumn{.Label = WidgetLabel} );
@@ -239,7 +239,8 @@ void FTypedElementWidgetConstructor::AddDefaultWidgetColumns(RowHandle Row, IEdi
 	}
 }
 
-FTypedElementWidgetConstructor::RowHandle FTypedElementWidgetConstructor::GetTargetRow(IEditorDataStorageProvider* DataStorage, RowHandle WidgetRow) const
+UE::Editor::DataStorage::RowHandle FTypedElementWidgetConstructor::GetTargetRow(
+	IEditorDataStorageProvider* DataStorage, UE::Editor::DataStorage::RowHandle WidgetRow) const
 {
 	using namespace UE::Editor::DataStorage;
 	RowHandle TargetRow = InvalidRowHandle;
@@ -267,7 +268,7 @@ TSharedPtr<SWidget> FSimpleWidgetConstructor::CreateWidget(IEditorDataStoragePro
 	return nullptr;
 }
 
-bool FSimpleWidgetConstructor::SetColumns(IEditorDataStorageProvider* DataStorage, RowHandle Row)
+bool FSimpleWidgetConstructor::SetColumns(IEditorDataStorageProvider* DataStorage, UE::Editor::DataStorage::RowHandle Row)
 {
 	return FTypedElementWidgetConstructor::SetColumns(DataStorage, Row);
 }
@@ -279,13 +280,13 @@ TSharedPtr<SWidget> FSimpleWidgetConstructor::CreateWidget(const UE::Editor::Dat
 }
 
 bool FSimpleWidgetConstructor::FinalizeWidget(IEditorDataStorageProvider* DataStorage, IEditorDataStorageUiProvider* DataStorageUi,
-	RowHandle Row, const TSharedPtr<SWidget>& Widget)
+	UE::Editor::DataStorage::RowHandle Row, const TSharedPtr<SWidget>& Widget)
 {
 	// This function is not needed anymore and only exists so derived classes cannot derive from it anymore
 	return true;
 }
 
-TSharedPtr<SWidget> FSimpleWidgetConstructor::Construct(RowHandle WidgetRow, IEditorDataStorageProvider* DataStorage,
+TSharedPtr<SWidget> FSimpleWidgetConstructor::Construct(UE::Editor::DataStorage::RowHandle WidgetRow, IEditorDataStorageProvider* DataStorage,
 	IEditorDataStorageUiProvider* DataStorageUi, const UE::Editor::DataStorage::FMetaDataView& Arguments)
 {
 	const UE::Editor::DataStorage::RowHandle TargetRow = GetTargetRow(DataStorage, WidgetRow);
diff --git a/Engine/Source/Runtime/TypedElementFramework/Public/Elements/Interfaces/TypedElementDataStorageUiInterface.h b/Engine/Source/Runtime/TypedElementFramework/Public/Elements/Interfaces/TypedElementDataStorageUiInterface.h
--- a/Engine/Source/Runtime/TypedElementFramework/Public/Elements/Interfaces/TypedElementDataStorageUiInterface.h
+++ b/Engine/Source/Runtime/TypedElementFramework/Public/Elements/Interfaces/TypedElementDataStorageUiInterface.h
@@ -32,8 +32,6 @@ struct FTypedElementWidgetConstructor
 {
 	GENERATED_BODY()
 
-	using RowHandle = UE::Editor::DataStorage::RowHandle;
-
 public:
 	TYPEDELEMENTFRAMEWORK_API explicit FTypedElementWidgetConstructor(const UScriptStruct* InTypeInfo);
 	explicit FTypedElementWidgetConstructor(EForceInit) {} //< For compatibility and shouldn't be directly used.
@@ -61,7 +59,7 @@ public:
 	 * Individual widget constructors can override this function with a name specific to them.
 	 */
 	TYPEDELEMENTFRAMEWORK_API virtual FString CreateWidgetDisplayName(
-		IEditorDataStorageProvider* DataStorage, RowHandle Row) const;
+		IEditorDataStorageProvider* DataStorage, UE::Editor::DataStorage::RowHandle Row) const;
 	
 	/**
 	 *	Calls Construct() to create the internal widget, and then stores it in a container before returning.
@@ -72,7 +70,7 @@ public:
 	 *	@see Construct
 	 */
 	TYPEDELEMENTFRAMEWORK_API virtual TSharedPtr<SWidget> ConstructFinalWidget(
-		RowHandle Row, /** The row the widget will be stored in. */
+		UE::Editor::DataStorage::RowHandle Row, /** The row the widget will be stored in. */
 		IEditorDataStorageProvider* DataStorage,
 		IEditorDataStorageUiProvider* DataStorageUi,
 		const UE::Editor::DataStorage::FMetaDataView& Arguments);
@@ -85,7 +83,7 @@ public:
 	 * In most cases, you want to call ConstructFinalWidget to create the actual widget.
 	 */
 	TYPEDELEMENTFRAMEWORK_API virtual TSharedPtr<SWidget> Construct(
-		RowHandle Row, /** The row the widget will be stored in. */
+		UE::Editor::DataStorage::RowHandle Row, /** The row the widget will be stored in. */
 		IEditorDataStorageProvider* DataStorage,
 		IEditorDataStorageUiProvider* DataStorageUi,
 		const UE::Editor::DataStorage::FMetaDataView& Arguments);
@@ -101,7 +99,7 @@ protected:
 		UE::Editor::DataStorage::RowHandle WidgetRow, 
 		const UE::Editor::DataStorage::FMetaDataView& Arguments);
 	/** Set any values in columns if needed. The columns provided through GetAdditionalColumnsList() will have already been created. */
-	TYPEDELEMENTFRAMEWORK_API virtual bool SetColumns(IEditorDataStorageProvider* DataStorage, RowHandle Row);
+	TYPEDELEMENTFRAMEWORK_API virtual bool SetColumns(IEditorDataStorageProvider* DataStorage, UE::Editor::DataStorage::RowHandle Row);
 	
 	/** Creates a (friendly) name for the provided column type. */
 	TYPEDELEMENTFRAMEWORK_API virtual FString DescribeColumnType(const UScriptStruct* ColumnType) const;
@@ -113,17 +111,18 @@ protected:
 	TYPEDELEMENTFRAMEWORK_API virtual bool FinalizeWidget(
 		IEditorDataStorageProvider* DataStorage,
 		IEditorDataStorageUiProvider* DataStorageUi,
-		RowHandle Row,
+		UE::Editor::DataStorage::RowHandle Row,
 		const TSharedPtr<SWidget>& Widget);
 
 	/** Add the default misc columns we want a widget row to have. */
-	TYPEDELEMENTFRAMEWORK_API void AddDefaultWidgetColumns(RowHandle Row, IEditorDataStorageProvider* DataStorage) const;
+	TYPEDELEMENTFRAMEWORK_API void AddDefaultWidgetColumns(UE::Editor::DataStorage::RowHandle Row, IEditorDataStorageProvider* DataStorage) const;
 
 	/**
 	 * Helper function to get the actual target row with the data the widget is operating on (if applicable). Returns InvalidRowHandle if there is no
 	 * target row
 	 */
-	TYPEDELEMENTFRAMEWORK_API RowHandle GetTargetRow(IEditorDataStorageProvider* DataStorage, RowHandle WidgetRow) const;
+	TYPEDELEMENTFRAMEWORK_API UE::Editor::DataStorage::RowHandle GetTargetRow(
+		IEditorDataStorageProvider* DataStorage, UE::Editor::DataStorage::RowHandle WidgetRow) const;
 
 protected:
 
@@ -173,7 +172,7 @@ struct FSimpleWidgetConstructor : public FTypedElementWidgetConstructor
 	 * @param WidgetRow The row that contains information about the widget itself
 	 * @return Whether any columns were added
 	 */
-	TYPEDELEMENTFRAMEWORK_API virtual bool SetColumns(IEditorDataStorageProvider* DataStorage, RowHandle WidgetRow) override;
+	TYPEDELEMENTFRAMEWORK_API virtual bool SetColumns(IEditorDataStorageProvider* DataStorage, UE::Editor::DataStorage::RowHandle WidgetRow) override;
 
 protected:
 
@@ -186,14 +185,14 @@ protected:
 	TYPEDELEMENTFRAMEWORK_API virtual bool FinalizeWidget(
 		IEditorDataStorageProvider* DataStorage,
 		IEditorDataStorageUiProvider* DataStorageUi,
-		RowHandle Row,
+		UE::Editor::DataStorage::RowHandle Row,
 		const TSharedPtr<SWidget>& Widget) override final;
 	
 	/**
 	 * Helper function to call SetColumns() and CreateWidget(), you should not need to override this for simple widget constructors
 	 */
 	TYPEDELEMENTFRAMEWORK_API virtual TSharedPtr<SWidget> Construct(
-		RowHandle WidgetRow, 
+		UE::Editor::DataStorage::RowHandle WidgetRow,
 		IEditorDataStorageProvider* DataStorage,
 		IEditorDataStorageUiProvider* DataStorageUi,
 		const UE::Editor::DataStorage::FMetaDataView& Arguments) override final;
@@ -211,8 +210,6 @@ struct TStructOpsTypeTraits<FTypedElementWidgetConstructor> : public TStructOpsT
 
 class IEditorDataStorageUiProvider : public IModularFeature
 {
-	using RowHandle = UE::Editor::DataStorage::RowHandle;
-
 public:
 	enum class EPurposeType : uint8
 	{
@@ -243,7 +240,7 @@ public:
 		SingleMatch
 	};
 
-	using WidgetCreatedCallback = TFunctionRef<void(const TSharedRef<SWidget>& NewWidget, RowHandle Row)>;
+	using WidgetCreatedCallback = TFunctionRef<void(const TSharedRef<SWidget>& NewWidget, UE::Editor::DataStorage::RowHandle Row)>;
 	using WidgetConstructorCallback = TFunctionRef<bool(TUniquePtr<FTypedElementWidgetConstructor>, TConstArrayView<TWeakObjectPtr<const UScriptStruct>>)>;
 	using WidgetPurposeCallback = TFunctionRef<void(FName, EPurposeType, const FText&)>;
 
@@ -342,7 +339,7 @@ public:
 	 * The provided row will be used to store the widget information on. If columns have already been added to the row, the 
 	 * constructor is free to use that to configure the widget. Arguments are used by the constructor to configure the widget.
 	 */
-	virtual TSharedPtr<SWidget> ConstructWidget(RowHandle Row, FTypedElementWidgetConstructor& Constructor,
+	virtual TSharedPtr<SWidget> ConstructWidget(UE::Editor::DataStorage::RowHandle Row, FTypedElementWidgetConstructor& Constructor,
 		const UE::Editor::DataStorage::FMetaDataView& Arguments) = 0;
 
 	/** Calls the provided callback for all known registered widget purposes. */
