Fixed the TEDS Query Builder not being able to bind to a referenced lambda.

diff --git a/Engine/Source/Runtime/TypedElementFramework/Public/Elements/Framework/TypedElementQueryBuilder.inl b/Engine/Source/Runtime/TypedElementFramework/Public/Elements/Framework/TypedElementQueryBuilder.inl
--- a/Engine/Source/Runtime/TypedElementFramework/Public/Elements/Framework/TypedElementQueryBuilder.inl
+++ b/Engine/Source/Runtime/TypedElementFramework/Public/Elements/Framework/TypedElementQueryBuilder.inl
@@ -779,7 +779,10 @@ namespace UE::Editor::DataStorage::Queries
 		template<SourceQueryContext SourceContext, bool ValidateColumns, typename Functor>
 		void BindQueryFunction(QueryCallback& Function, Functor&& Callback)
 		{
-			BindQueryFunction_Expand<SourceContext, ValidateColumns>(Function, std::forward<Functor>(Callback), &Functor::operator());
+			BindQueryFunction_Expand<SourceContext, ValidateColumns>(
+				Function,
+				std::forward<Functor>(Callback),
+				&std::remove_cv_t<std::remove_reference_t<Functor>>::operator());
 		}
 
 		
