[TEDS Outliner] Fixed TEDS outliner columns reporting missing localization, changing the CreateWidgetDisplayName function to return FText, and general cleanup

diff --git a/Engine/Plugins/Experimental/EditorDataStorage/Source/TedsUI/Private/Widgets/ExportedTextWidget.cpp b/Engine/Plugins/Experimental/EditorDataStorage/Source/TedsUI/Private/Widgets/ExportedTextWidget.cpp
--- a/Engine/Plugins/Experimental/EditorDataStorage/Source/TedsUI/Private/Widgets/ExportedTextWidget.cpp
+++ b/Engine/Plugins/Experimental/EditorDataStorage/Source/TedsUI/Private/Widgets/ExportedTextWidget.cpp
@@ -150,7 +150,7 @@ const UE::Editor::DataStorage::Queries::FConditions* FExportedTextWidgetConstruc
 	return nullptr;
 }
 
-FString FExportedTextWidgetConstructor::CreateWidgetDisplayName(IEditorDataStorageProvider* DataStorage,
+FText FExportedTextWidgetConstructor::CreateWidgetDisplayNameText(IEditorDataStorageProvider* DataStorage,
 	UE::Editor::DataStorage::RowHandle Row) const
 {
 	if (FTypedElementScriptStructTypeInfoColumn* TypeInfoColumn = DataStorage->GetColumn<FTypedElementScriptStructTypeInfoColumn>(Row))
diff --git a/Engine/Plugins/Experimental/EditorDataStorage/Source/TedsUI/Public/Widgets/ExportedTextWidget.h b/Engine/Plugins/Experimental/EditorDataStorage/Source/TedsUI/Public/Widgets/ExportedTextWidget.h
--- a/Engine/Plugins/Experimental/EditorDataStorage/Source/TedsUI/Public/Widgets/ExportedTextWidget.h
+++ b/Engine/Plugins/Experimental/EditorDataStorage/Source/TedsUI/Public/Widgets/ExportedTextWidget.h
@@ -36,7 +36,8 @@ public:
 
 	virtual TConstArrayView<const UScriptStruct*> GetAdditionalColumnsList() const override;
 	virtual const UE::Editor::DataStorage::Queries::FConditions* GetQueryConditions() const override;
-	virtual FString CreateWidgetDisplayName(
+
+	virtual FText CreateWidgetDisplayNameText(
 		IEditorDataStorageProvider* DataStorage, UE::Editor::DataStorage::RowHandle Row) const override;
 
 	virtual TSharedPtr<SWidget> ConstructFinalWidget(
diff --git a/Engine/Plugins/Experimental/EditorDataStorageFeatures/Source/TedsOutliner/Private/Compatibility/SceneOutlinerTedsBridge.cpp b/Engine/Plugins/Experimental/EditorDataStorageFeatures/Source/TedsOutliner/Private/Compatibility/SceneOutlinerTedsBridge.cpp
--- a/Engine/Plugins/Experimental/EditorDataStorageFeatures/Source/TedsOutliner/Private/Compatibility/SceneOutlinerTedsBridge.cpp
+++ b/Engine/Plugins/Experimental/EditorDataStorageFeatures/Source/TedsOutliner/Private/Compatibility/SceneOutlinerTedsBridge.cpp
@@ -543,9 +543,12 @@ void FSceneOutlinerTedsBridge::AssignQuery(UE::Editor::DataStorage::QueryHandle
 					OutlinerPinned->RemoveColumn(FallbackColumn);
 
 					FName NameId = UE::Editor::DataStorage::TableViewerUtils::FindLongestMatchingName(ColumnTypes, IndexOffset);
+					FText DisplayName = CellConstructor.Get()->CreateWidgetDisplayNameText(Storage);
 					AddedColumns.Add(NameId);
 					OutlinerPinned->AddColumn(NameId,
-						FSceneOutlinerColumnInfo(ESceneOutlinerColumnVisibility::Visible, ColumnPriority,
+						FSceneOutlinerColumnInfo(
+							ESceneOutlinerColumnVisibility::Visible, 
+							ColumnPriority,
 							FCreateSceneOutlinerColumn::CreateLambda(
 								[this, Query, MetaDataView, NameId, &ColumnTypes, CellConstructor, &OutlinerPinned, FallbackColumn](ISceneOutliner&)
 								{
@@ -556,7 +559,10 @@ void FSceneOutlinerTedsBridge::AssignQuery(UE::Editor::DataStorage::QueryHandle
 										TArray<TWeakObjectPtr<const UScriptStruct>>(ColumnTypes.GetData(), ColumnTypes.Num()), 
 										MoveTemp(HeaderConstructor), CellConstructor, FallbackColumn, OutlinerPinned, Dealiaser);
 
-								})
+								}),
+							true,
+							TOptional<float>(),
+							DisplayName
 						)
 					);
 					++IndexOffset;
diff --git a/Engine/Plugins/Experimental/EditorDataStorageFeatures/Source/TedsTableViewer/Private/Widgets/SRowDetails.cpp b/Engine/Plugins/Experimental/EditorDataStorageFeatures/Source/TedsTableViewer/Private/Widgets/SRowDetails.cpp
--- a/Engine/Plugins/Experimental/EditorDataStorageFeatures/Source/TedsTableViewer/Private/Widgets/SRowDetails.cpp
+++ b/Engine/Plugins/Experimental/EditorDataStorageFeatures/Source/TedsTableViewer/Private/Widgets/SRowDetails.cpp
@@ -162,8 +162,8 @@ namespace UE::Editor::DataStorage
 		if (ColumnName == Widgets::Private::NameColumn)
 		{
 			return SNew(STextBlock)
-					.Text(FText::FromString(Item->WidgetConstructor->CreateWidgetDisplayName(
-							DataStorage, Item->WidgetRow)));
+					.Text(Item->WidgetConstructor->CreateWidgetDisplayNameText(
+							DataStorage, Item->WidgetRow));
 		}
 		else if (ColumnName == Widgets::Private::DataColumn)
 		{
diff --git a/Engine/Source/Editor/SceneOutliner/Private/SSceneOutliner.cpp b/Engine/Source/Editor/SceneOutliner/Private/SSceneOutliner.cpp
--- a/Engine/Source/Editor/SceneOutliner/Private/SSceneOutliner.cpp
+++ b/Engine/Source/Editor/SceneOutliner/Private/SSceneOutliner.cpp
@@ -1199,11 +1199,11 @@ int32 SSceneOutliner::GetFilterCount() const
 	return Filters->Num();
 }
 
-void SSceneOutliner::AddColumn(FName ColumnId, const FSceneOutlinerColumnInfo& ColumInfo)
+void SSceneOutliner::AddColumn(FName ColumnId, const FSceneOutlinerColumnInfo& ColumnInfo)
 {
 	if (!SharedData->ColumnMap.Contains(ColumnId))
 	{
-		SharedData->ColumnMap.Add(ColumnId, ColumInfo);
+		SharedData->ColumnMap.Add(ColumnId, ColumnInfo);
 
 		// Get the new sorted list of columns to make sure this is added in the right position
 		TArray<FName> SortedColumnIDs;
@@ -1216,7 +1216,7 @@ void SSceneOutliner::AddColumn(FName ColumnId, const FSceneOutlinerColumnInfo& C
 			ColumnVisibilities = OutlinerConfig->ColumnVisibilities;
 		}
 
-		AddColumn_Internal(ColumnId, ColumInfo, ColumnVisibilities, SortedColumnIDs.Find(ColumnId));
+		AddColumn_Internal(ColumnId, ColumnInfo, ColumnVisibilities, SortedColumnIDs.Find(ColumnId));
 	}
 }
 
diff --git a/Engine/Source/Editor/SceneOutliner/Private/SceneOutlinerModule.cpp b/Engine/Source/Editor/SceneOutliner/Private/SceneOutlinerModule.cpp
--- a/Engine/Source/Editor/SceneOutliner/Private/SceneOutlinerModule.cpp
+++ b/Engine/Source/Editor/SceneOutliner/Private/SceneOutlinerModule.cpp
@@ -227,7 +227,6 @@ void FSceneOutlinerModule::CreateActorBrowserColumns(FSceneOutlinerInitializatio
 
 	InInitOptions.UseDefaultColumns();
 	
-	InInitOptions.UseDefaultColumns();
 	InInitOptions.ColumnMap.Add(FSceneOutlinerBuiltInColumnTypes::Gutter(), FSceneOutlinerColumnInfo(ESceneOutlinerColumnVisibility::Visible, 0, FCreateSceneOutlinerColumn(), true, TOptional<float>(), FSceneOutlinerBuiltInColumnTypes::Gutter_Localized()));
 	InInitOptions.ColumnMap.Add(FSceneOutlinerBuiltInColumnTypes::ActorInfo(), FSceneOutlinerColumnInfo(ESceneOutlinerColumnVisibility::Visible, 20, FCreateSceneOutlinerColumn(), true, TOptional<float>(), FSceneOutlinerBuiltInColumnTypes::ActorInfo_Localized()));
 
diff --git a/Engine/Source/Editor/SceneOutliner/Public/SSceneOutliner.h b/Engine/Source/Editor/SceneOutliner/Public/SSceneOutliner.h
--- a/Engine/Source/Editor/SceneOutliner/Public/SSceneOutliner.h
+++ b/Engine/Source/Editor/SceneOutliner/Public/SSceneOutliner.h
@@ -359,9 +359,9 @@ public:
 	 * Add or replace a column of the scene outliner
 	 * Note: The column id must match the id of the column returned by the factory
 	 * @param ColumnId The id of the column to add
-	 * @param ColumInfo The struct that contains the information on how to present and retrieve the column
+	 * @param ColumnInfo The struct that contains the information on how to present and retrieve the column
 	 */
-	virtual void AddColumn(FName ColumnId, const FSceneOutlinerColumnInfo& ColumInfo) override;
+	virtual void AddColumn(FName ColumnId, const FSceneOutlinerColumnInfo& ColumnInfo) override;
 
 	/**
 	 * Remove a column of the scene outliner
diff --git a/Engine/Source/Runtime/TypedElementFramework/Private/Elements/Interfaces/TypedElementDataStorageUiInterface.cpp b/Engine/Source/Runtime/TypedElementFramework/Private/Elements/Interfaces/TypedElementDataStorageUiInterface.cpp
--- a/Engine/Source/Runtime/TypedElementFramework/Private/Elements/Interfaces/TypedElementDataStorageUiInterface.cpp
+++ b/Engine/Source/Runtime/TypedElementFramework/Private/Elements/Interfaces/TypedElementDataStorageUiInterface.cpp
@@ -9,6 +9,8 @@
 #include "Elements/Framework/TypedElementDataStorageWidget.h"
 #include "Widgets/SNullWidget.h"
 
+#define LOCTEXT_NAMESPACE "FTypedElementWidgetConstructor"
+
 FTypedElementWidgetConstructor::FTypedElementWidgetConstructor(const UScriptStruct* InTypeInfo)
 	: TypeInfo(InTypeInfo)
 {
@@ -55,23 +57,29 @@ TConstArrayView<const UScriptStruct*> FTypedElementWidgetConstructor::GetAdditio
 
 FString FTypedElementWidgetConstructor::CreateWidgetDisplayName(
 	IEditorDataStorageProvider* DataStorage, UE::Editor::DataStorage::RowHandle Row) const
+{
+	return CreateWidgetDisplayNameText(DataStorage, Row).ToString();
+}
+
+FText FTypedElementWidgetConstructor::CreateWidgetDisplayNameText(
+	IEditorDataStorageProvider* DataStorage, UE::Editor::DataStorage::RowHandle Row) const
 {
 	switch (MatchedColumnTypes.Num())
 	{
 	case 0:
-		return FString(TEXT("TEDS Column"));
+		return LOCTEXT("TEDSColumn", "TEDS Column");
 	case 1:
 		return DescribeColumnType(MatchedColumnTypes[0].Get());
 	default:
 	{
-		FString LongestMatchString = DescribeColumnType(MatchedColumnTypes[0].Get());
+		FString LongestMatchString = DescribeColumnType(MatchedColumnTypes[0].Get()).ToString();
 		FStringView LongestMatch = LongestMatchString;
 		const TWeakObjectPtr<const UScriptStruct>* It = MatchedColumnTypes.GetData();
 		const TWeakObjectPtr<const UScriptStruct>* ItEnd = It + MatchedColumnTypes.Num();
 		++It; // Skip the first entry as that's already set.
 		for (; It != ItEnd; ++It)
 		{
-			FString NextMatchText = DescribeColumnType(It->Get());
+			FString NextMatchText = DescribeColumnType(It->Get()).ToString();
 			FStringView NextMatch = NextMatchText;
 
 			int32 MatchSize = 0;
@@ -97,10 +105,10 @@ FString FTypedElementWidgetConstructor::CreateWidgetDisplayName(
 			else
 			{
 				// There are not enough characters in the string that match. Just return the name of the first column
-				return LongestMatchString;
+				return FText::FromString(LongestMatchString);
 			}
 		}
-		return FString(LongestMatch);
+		return FText::FromString(FString(LongestMatch));
 	}
 	};
 }
@@ -194,20 +202,17 @@ bool FTypedElementWidgetConstructor::SetColumns(IEditorDataStorageProvider* Data
 	return true;
 }
 
-FString FTypedElementWidgetConstructor::DescribeColumnType(const UScriptStruct* ColumnType) const
+FText FTypedElementWidgetConstructor::DescribeColumnType(const UScriptStruct* ColumnType) const
 {
-	static const FName DisplayNameName(TEXT("DisplayName"));
-
 #if WITH_EDITOR
 	if (ColumnType)
 	{
-		const FString* Name = ColumnType->FindMetaData(DisplayNameName);
-		return Name ? *Name : ColumnType->GetDisplayNameText().ToString();
+		return ColumnType->GetDisplayNameText();
 	}
 	else
 #endif
 	{
-		return FString(TEXT("<Invalid>"));
+		return LOCTEXT("Invalid", "<Invalid>");
 	}
 }
 
@@ -222,7 +227,7 @@ bool FTypedElementWidgetConstructor::FinalizeWidget(
 
 void FTypedElementWidgetConstructor::AddDefaultWidgetColumns(UE::Editor::DataStorage::RowHandle Row, IEditorDataStorageProvider* DataStorage) const
 {
-	const FString WidgetLabel(CreateWidgetDisplayName(DataStorage, Row));
+	const FString WidgetLabel(CreateWidgetDisplayNameText(DataStorage, Row).ToString());
 	DataStorage->AddColumn(Row, FTypedElementLabelColumn{.Label = WidgetLabel} );
 
 	// We don't want to display any second level widgets (widgets for widgets and so on...) in UI because they will cause the table viewer to
@@ -306,3 +311,4 @@ TSharedPtr<SWidget> FSimpleWidgetConstructor::Construct(UE::Editor::DataStorage:
 	return Widget;
 }
 
+#undef LOCTEXT_NAMESPACE
\ No newline at end of file
diff --git a/Engine/Source/Runtime/TypedElementFramework/Public/Elements/Interfaces/TypedElementDataStorageUiInterface.h b/Engine/Source/Runtime/TypedElementFramework/Public/Elements/Interfaces/TypedElementDataStorageUiInterface.h
--- a/Engine/Source/Runtime/TypedElementFramework/Public/Elements/Interfaces/TypedElementDataStorageUiInterface.h
+++ b/Engine/Source/Runtime/TypedElementFramework/Public/Elements/Interfaces/TypedElementDataStorageUiInterface.h
@@ -52,14 +52,27 @@ public:
 	/** Returns a list of additional columns the widget requires to be added to its rows. */
 	TYPEDELEMENTFRAMEWORK_API virtual TConstArrayView<const UScriptStruct*> GetAdditionalColumnsList() const;
 
-	/** 
+	/**
 	 * Returns a friendly name for the data the created widget represents.
 	 * By default the associated column is used. If there are multiple columns associated with the constructor
 	 * the default implementation will attempt to find the longest common starting string for all the columns.
 	 * Individual widget constructors can override this function with a name specific to them.
 	 */
+	UE_DEPRECATED(5.6, "CreateWidgetDisplayName is deprecated, please use CreateWidgetDisplayNameText instead, which will return FText rather than an FString.")
 	TYPEDELEMENTFRAMEWORK_API virtual FString CreateWidgetDisplayName(
-		IEditorDataStorageProvider* DataStorage, UE::Editor::DataStorage::RowHandle Row) const;
+		IEditorDataStorageProvider* DataStorage, UE::Editor::DataStorage::RowHandle Row) const final;
+
+	/** 
+	 * Returns an FText of the friendly name for the data the created widget represents.
+	 * By default the associated column is used. If there are multiple columns associated with the constructor
+	 * the default implementation will attempt to find the longest common starting string for all the columns.
+	 * Individual widget constructors can override this function with a name specific to them.
+	 * 
+	 * If no row is specified, it will default to an InvalidRowHandle, assuming that we are asking for a display 
+	 * name for the whole widget and not a specific instance.
+	 */
+	TYPEDELEMENTFRAMEWORK_API virtual FText CreateWidgetDisplayNameText(
+		IEditorDataStorageProvider* DataStorage, UE::Editor::DataStorage::RowHandle Row = UE::Editor::DataStorage::InvalidRowHandle) const;
 	
 	/**
 	 *	Calls Construct() to create the internal widget, and then stores it in a container before returning.
@@ -102,7 +115,7 @@ protected:
 	TYPEDELEMENTFRAMEWORK_API virtual bool SetColumns(IEditorDataStorageProvider* DataStorage, UE::Editor::DataStorage::RowHandle Row);
 	
 	/** Creates a (friendly) name for the provided column type. */
-	TYPEDELEMENTFRAMEWORK_API virtual FString DescribeColumnType(const UScriptStruct* ColumnType) const;
+	TYPEDELEMENTFRAMEWORK_API virtual FText DescribeColumnType(const UScriptStruct* ColumnType) const;
 
 	/** 
 	 * Last opportunity to configure anything in the widget or the row. This step can be needed to initialize widgets with data stored
