[TEDS] Cleanup aliases, move core interface classes into namespace, and fix up one class name that wasn't following standards

diff --git a/Engine/Source/Runtime/Engine/Private/Elements/Framework/EngineElementsLibrary.cpp b/Engine/Source/Runtime/Engine/Private/Elements/Framework/EngineElementsLibrary.cpp
--- a/Engine/Source/Runtime/Engine/Private/Elements/Framework/EngineElementsLibrary.cpp
+++ b/Engine/Source/Runtime/Engine/Private/Elements/Framework/EngineElementsLibrary.cpp
@@ -382,7 +382,7 @@ void UEngineElementsLibrary::DestroyActorElement(const AActor* InActor, TTypedEl
 void UEngineElementsLibrary::RegisterActorElement(const AActor* InActor)
 {
 	using namespace UE::Editor::DataStorage;
-	IEditorDataStorageCompatibilityProvider* Storage = GetMutableDataStorageFeature<IEditorDataStorageCompatibilityProvider>(CompatibilityFeatureName);
+	ICompatibilityProvider* Storage = GetMutableDataStorageFeature<ICompatibilityProvider>(CompatibilityFeatureName);
 	if (Storage)
 	{
 		Storage->AddCompatibleObject(const_cast<AActor*>(InActor));
@@ -392,7 +392,7 @@ void UEngineElementsLibrary::RegisterActorElement(const AActor* InActor)
 void UEngineElementsLibrary::UnregisterActorElement(const AActor* InActor)
 {
 	using namespace UE::Editor::DataStorage;
-	if (IEditorDataStorageCompatibilityProvider* Storage = GetMutableDataStorageFeature<IEditorDataStorageCompatibilityProvider>(CompatibilityFeatureName))
+	if (ICompatibilityProvider* Storage = GetMutableDataStorageFeature<ICompatibilityProvider>(CompatibilityFeatureName))
 	{
 		Storage->RemoveCompatibleObject(const_cast<AActor*>(InActor));
 	}
diff --git a/Engine/Source/Runtime/Engine/Private/Level.cpp b/Engine/Source/Runtime/Engine/Private/Level.cpp
--- a/Engine/Source/Runtime/Engine/Private/Level.cpp
+++ b/Engine/Source/Runtime/Engine/Private/Level.cpp
@@ -1060,8 +1060,8 @@ void ULevel::AddLoadedActors(const TArray<AActor*>& ActorList, const FTransform*
 	}
 
 	using namespace UE::Editor::DataStorage;
-	IEditorDataStorageCompatibilityProvider* TedsCompat =
-					GetMutableDataStorageFeature<IEditorDataStorageCompatibilityProvider>(CompatibilityFeatureName);
+	ICompatibilityProvider* TedsCompat =
+					GetMutableDataStorageFeature<ICompatibilityProvider>(CompatibilityFeatureName);
 
 	// Finalize actors
 	for (AActor* Actor : ActorsQueue)
@@ -1184,8 +1184,8 @@ void ULevel::RemoveLoadedActors(const TArray<AActor*>& ActorList, const FTransfo
 	OnLoadedActorRemovedFromLevelPreEvent.Broadcast(ActorsQueue);
 	
 	using namespace UE::Editor::DataStorage;
-	IEditorDataStorageCompatibilityProvider* TedsCompat =
-					GetMutableDataStorageFeature<IEditorDataStorageCompatibilityProvider>(CompatibilityFeatureName);
+	ICompatibilityProvider* TedsCompat =
+					GetMutableDataStorageFeature<ICompatibilityProvider>(CompatibilityFeatureName);
 
 	for (AActor* Actor : ActorsQueue)
 	{
diff --git a/Engine/Source/Runtime/TypedElementFramework/Private/Elements/Common/TypedElementQueryConditions.cpp b/Engine/Source/Runtime/TypedElementFramework/Private/Elements/Common/TypedElementQueryConditions.cpp
--- a/Engine/Source/Runtime/TypedElementFramework/Private/Elements/Common/TypedElementQueryConditions.cpp
+++ b/Engine/Source/Runtime/TypedElementFramework/Private/Elements/Common/TypedElementQueryConditions.cpp
@@ -9,7 +9,7 @@
 
 namespace UE::Editor::DataStorage::Queries
 {
-	FEditorStorageQueryConditionCompileContext::FEditorStorageQueryConditionCompileContext(IEditorDataStorageProvider* InDataStorage)
+	FEditorStorageQueryConditionCompileContext::FEditorStorageQueryConditionCompileContext(ICoreProvider* InDataStorage)
 		: DataStorage(InDataStorage)
 	{
 	}
diff --git a/Engine/Source/Runtime/TypedElementFramework/Private/Elements/Framework/TypedElementAttributeBinding.cpp b/Engine/Source/Runtime/TypedElementFramework/Private/Elements/Framework/TypedElementAttributeBinding.cpp
--- a/Engine/Source/Runtime/TypedElementFramework/Private/Elements/Framework/TypedElementAttributeBinding.cpp
+++ b/Engine/Source/Runtime/TypedElementFramework/Private/Elements/Framework/TypedElementAttributeBinding.cpp
@@ -7,12 +7,12 @@
 namespace UE::Editor::DataStorage
 {
 	FAttributeBinder::FAttributeBinder(RowHandle InTargetRow)
-		: FAttributeBinder(InTargetRow, GetMutableDataStorageFeature<IEditorDataStorageProvider>(StorageFeatureName))
+		: FAttributeBinder(InTargetRow, GetMutableDataStorageFeature<ICoreProvider>(StorageFeatureName))
 	{
 	}
 
 	FAttributeBinder::FAttributeBinder(
-		RowHandle InTargetRow, IEditorDataStorageProvider* InDataStorage)
+		RowHandle InTargetRow, ICoreProvider* InDataStorage)
 		: TargetRow(InTargetRow)
 		, DataStorage(InDataStorage)
 	{
diff --git a/Engine/Source/Runtime/TypedElementFramework/Private/Elements/Framework/TypedElementRegistry.cpp b/Engine/Source/Runtime/TypedElementFramework/Private/Elements/Framework/TypedElementRegistry.cpp
--- a/Engine/Source/Runtime/TypedElementFramework/Private/Elements/Framework/TypedElementRegistry.cpp
+++ b/Engine/Source/Runtime/TypedElementFramework/Private/Elements/Framework/TypedElementRegistry.cpp
@@ -65,40 +65,40 @@ UTypedElementRegistry* UTypedElementRegistry::GetInstance()
 	return Instance.Get();
 }
 
-IEditorDataStorageProvider* UTypedElementRegistry::GetMutableDataStorage()
+UE::Editor::DataStorage::ICoreProvider* UTypedElementRegistry::GetMutableDataStorage()
 {
 	using namespace UE::Editor::DataStorage;
-	return GetMutableDataStorageFeature<IEditorDataStorageProvider>(StorageFeatureName);
+	return GetMutableDataStorageFeature<ICoreProvider>(StorageFeatureName);
 }
 
-const IEditorDataStorageProvider* UTypedElementRegistry::GetDataStorage() const
+const UE::Editor::DataStorage::ICoreProvider* UTypedElementRegistry::GetDataStorage() const
 {
 	using namespace UE::Editor::DataStorage;
-	return GetDataStorageFeature<IEditorDataStorageProvider>(StorageFeatureName);
+	return GetDataStorageFeature<ICoreProvider>(StorageFeatureName);
 }
 
-IEditorDataStorageCompatibilityProvider* UTypedElementRegistry::GetMutableDataStorageCompatibility()
+UE::Editor::DataStorage::ICompatibilityProvider* UTypedElementRegistry::GetMutableDataStorageCompatibility()
 {
 	using namespace UE::Editor::DataStorage;
-	return GetMutableDataStorageFeature<IEditorDataStorageCompatibilityProvider>(CompatibilityFeatureName);
+	return GetMutableDataStorageFeature<ICompatibilityProvider>(CompatibilityFeatureName);
 }
 
-const IEditorDataStorageCompatibilityProvider* UTypedElementRegistry::GetDataStorageCompatibility() const
+const UE::Editor::DataStorage::ICompatibilityProvider* UTypedElementRegistry::GetDataStorageCompatibility() const
 {
 	using namespace UE::Editor::DataStorage;
-	return GetDataStorageFeature<IEditorDataStorageCompatibilityProvider>(CompatibilityFeatureName);
+	return GetDataStorageFeature<ICompatibilityProvider>(CompatibilityFeatureName);
 }
 
-IEditorDataStorageUiProvider* UTypedElementRegistry::GetMutableDataStorageUi()
+UE::Editor::DataStorage::IUiProvider* UTypedElementRegistry::GetMutableDataStorageUi()
 {
 	using namespace UE::Editor::DataStorage;
-	return GetMutableDataStorageFeature<IEditorDataStorageUiProvider>(UiFeatureName);
+	return GetMutableDataStorageFeature<IUiProvider>(UiFeatureName);
 }
 
-const IEditorDataStorageUiProvider* UTypedElementRegistry::GetDataStorageUi() const
+const UE::Editor::DataStorage::IUiProvider* UTypedElementRegistry::GetDataStorageUi() const
 {
 	using namespace UE::Editor::DataStorage;
-	return GetDataStorageFeature<IEditorDataStorageUiProvider>(UiFeatureName);
+	return GetDataStorageFeature<IUiProvider>(UiFeatureName);
 }
 
 bool UTypedElementRegistry::AreDataStorageInterfacesSet() const
diff --git a/Engine/Source/Runtime/TypedElementFramework/Private/Elements/Framework/TypedElementtAttributeBindingText.cpp b/Engine/Source/Runtime/TypedElementFramework/Private/Elements/Framework/TypedElementtAttributeBindingText.cpp
--- a/Engine/Source/Runtime/TypedElementFramework/Private/Elements/Framework/TypedElementtAttributeBindingText.cpp
+++ b/Engine/Source/Runtime/TypedElementFramework/Private/Elements/Framework/TypedElementtAttributeBindingText.cpp
@@ -5,7 +5,7 @@
 namespace UE::Editor::DataStorage
 {
 	FTextAttributeFormatted::FTextAttributeFormatted(
-		FTextFormat InFormat, RowHandle InTargetRow, IEditorDataStorageProvider* InDataStorage)
+		FTextFormat InFormat, RowHandle InTargetRow, ICoreProvider* InDataStorage)
 		: Format(MoveTemp(InFormat))
 		, TargetRow(InTargetRow)
 		, DataStorage(InDataStorage)
diff --git a/Engine/Source/Runtime/TypedElementFramework/Private/Elements/Interfaces/TypedElementDataStorageUiInterface.cpp b/Engine/Source/Runtime/TypedElementFramework/Private/Elements/Interfaces/TypedElementDataStorageUiInterface.cpp
--- a/Engine/Source/Runtime/TypedElementFramework/Private/Elements/Interfaces/TypedElementDataStorageUiInterface.cpp
+++ b/Engine/Source/Runtime/TypedElementFramework/Private/Elements/Interfaces/TypedElementDataStorageUiInterface.cpp
@@ -56,13 +56,13 @@ TConstArrayView<const UScriptStruct*> FTypedElementWidgetConstructor::GetAdditio
 }
 
 FString FTypedElementWidgetConstructor::CreateWidgetDisplayName(
-	IEditorDataStorageProvider* DataStorage, UE::Editor::DataStorage::RowHandle Row) const
+	UE::Editor::DataStorage::ICoreProvider* DataStorage, UE::Editor::DataStorage::RowHandle Row) const
 {
 	return CreateWidgetDisplayNameText(DataStorage, Row).ToString();
 }
 
 FText FTypedElementWidgetConstructor::CreateWidgetDisplayNameText(
-	IEditorDataStorageProvider* DataStorage, UE::Editor::DataStorage::RowHandle Row) const
+	UE::Editor::DataStorage::ICoreProvider* DataStorage, UE::Editor::DataStorage::RowHandle Row) const
 {
 	switch (MatchedColumnTypes.Num())
 	{
@@ -115,8 +115,8 @@ FText FTypedElementWidgetConstructor::CreateWidgetDisplayNameText(
 
 TSharedPtr<SWidget> FTypedElementWidgetConstructor::ConstructFinalWidget(
 	UE::Editor::DataStorage::RowHandle Row,
-	IEditorDataStorageProvider* DataStorage,
-	IEditorDataStorageUiProvider* DataStorageUi,
+	UE::Editor::DataStorage::ICoreProvider* DataStorage,
+	UE::Editor::DataStorage::IUiProvider* DataStorageUi,
 	const UE::Editor::DataStorage::FMetaDataView& Arguments)
 {
 	// Add the additional columns to the UI row
@@ -160,8 +160,8 @@ TSharedPtr<SWidget> FTypedElementWidgetConstructor::ConstructFinalWidget(
 
 TSharedPtr<SWidget> FTypedElementWidgetConstructor::Construct(
 	UE::Editor::DataStorage::RowHandle Row,
-	IEditorDataStorageProvider* DataStorage,
-	IEditorDataStorageUiProvider* DataStorageUi,
+	UE::Editor::DataStorage::ICoreProvider* DataStorage,
+	UE::Editor::DataStorage::IUiProvider* DataStorageUi,
 	const UE::Editor::DataStorage::FMetaDataView& Arguments)
 {
 	UE::Editor::DataStorage::RowHandle TargetRow = GetTargetRow(DataStorage, Row);
@@ -188,8 +188,8 @@ TSharedPtr<SWidget> FTypedElementWidgetConstructor::CreateWidget(const UE::Edito
 }
 
 TSharedPtr<SWidget> FTypedElementWidgetConstructor::CreateWidget(
-	IEditorDataStorageProvider* DataStorage,
-	IEditorDataStorageUiProvider* DataStorageUi,
+	UE::Editor::DataStorage::ICoreProvider* DataStorage,
+	UE::Editor::DataStorage::IUiProvider* DataStorageUi,
 	UE::Editor::DataStorage::RowHandle TargetRow,
 	UE::Editor::DataStorage::RowHandle UiRow,
 	const UE::Editor::DataStorage::FMetaDataView& Arguments)
@@ -197,7 +197,7 @@ TSharedPtr<SWidget> FTypedElementWidgetConstructor::CreateWidget(
 	return CreateWidget(Arguments);
 }
 
-bool FTypedElementWidgetConstructor::SetColumns(IEditorDataStorageProvider* DataStorage, UE::Editor::DataStorage::RowHandle Row)
+bool FTypedElementWidgetConstructor::SetColumns(UE::Editor::DataStorage::ICoreProvider* DataStorage, UE::Editor::DataStorage::RowHandle Row)
 {
 	return true;
 }
@@ -217,15 +217,15 @@ FText FTypedElementWidgetConstructor::DescribeColumnType(const UScriptStruct* Co
 }
 
 bool FTypedElementWidgetConstructor::FinalizeWidget(
-	IEditorDataStorageProvider* DataStorage,
-	IEditorDataStorageUiProvider* DataStorageUi,
+	UE::Editor::DataStorage::ICoreProvider* DataStorage,
+	UE::Editor::DataStorage::IUiProvider* DataStorageUi,
 	UE::Editor::DataStorage::RowHandle Row,
 	const TSharedPtr<SWidget>& Widget)
 {
 	return true;
 }
 
-void FTypedElementWidgetConstructor::AddDefaultWidgetColumns(UE::Editor::DataStorage::RowHandle Row, IEditorDataStorageProvider* DataStorage) const
+void FTypedElementWidgetConstructor::AddDefaultWidgetColumns(UE::Editor::DataStorage::RowHandle Row, UE::Editor::DataStorage::ICoreProvider* DataStorage) const
 {
 	const FString WidgetLabel(CreateWidgetDisplayNameText(DataStorage, Row).ToString());
 	DataStorage->AddColumn(Row, FTypedElementLabelColumn{.Label = WidgetLabel} );
@@ -245,7 +245,7 @@ void FTypedElementWidgetConstructor::AddDefaultWidgetColumns(UE::Editor::DataSto
 }
 
 UE::Editor::DataStorage::RowHandle FTypedElementWidgetConstructor::GetTargetRow(
-	IEditorDataStorageProvider* DataStorage, UE::Editor::DataStorage::RowHandle WidgetRow) const
+	UE::Editor::DataStorage::ICoreProvider* DataStorage, UE::Editor::DataStorage::RowHandle WidgetRow) const
 {
 	using namespace UE::Editor::DataStorage;
 	RowHandle TargetRow = InvalidRowHandle;
@@ -266,14 +266,16 @@ FSimpleWidgetConstructor::FSimpleWidgetConstructor(const UScriptStruct* InTypeIn
 {
 }
 
-TSharedPtr<SWidget> FSimpleWidgetConstructor::CreateWidget(IEditorDataStorageProvider* DataStorage,
-	IEditorDataStorageUiProvider* DataStorageUi, UE::Editor::DataStorage::RowHandle TargetRow, UE::Editor::DataStorage::RowHandle WidgetRow,
+TSharedPtr<SWidget> FSimpleWidgetConstructor::CreateWidget(UE::Editor::DataStorage::ICoreProvider* DataStorage,
+	UE::Editor::DataStorage::IUiProvider* DataStorageUi,
+	UE::Editor::DataStorage::RowHandle TargetRow,
+	UE::Editor::DataStorage::RowHandle WidgetRow,
 	const UE::Editor::DataStorage::FMetaDataView& Arguments)
 {
 	return nullptr;
 }
 
-bool FSimpleWidgetConstructor::SetColumns(IEditorDataStorageProvider* DataStorage, UE::Editor::DataStorage::RowHandle Row)
+bool FSimpleWidgetConstructor::SetColumns(UE::Editor::DataStorage::ICoreProvider* DataStorage, UE::Editor::DataStorage::RowHandle Row)
 {
 	return FTypedElementWidgetConstructor::SetColumns(DataStorage, Row);
 }
@@ -284,15 +286,16 @@ TSharedPtr<SWidget> FSimpleWidgetConstructor::CreateWidget(const UE::Editor::Dat
 	return nullptr;
 }
 
-bool FSimpleWidgetConstructor::FinalizeWidget(IEditorDataStorageProvider* DataStorage, IEditorDataStorageUiProvider* DataStorageUi,
+bool FSimpleWidgetConstructor::FinalizeWidget(UE::Editor::DataStorage::ICoreProvider* DataStorage,
+	UE::Editor::DataStorage::IUiProvider* DataStorageUi,
 	UE::Editor::DataStorage::RowHandle Row, const TSharedPtr<SWidget>& Widget)
 {
 	// This function is not needed anymore and only exists so derived classes cannot derive from it anymore
 	return true;
 }
 
-TSharedPtr<SWidget> FSimpleWidgetConstructor::Construct(UE::Editor::DataStorage::RowHandle WidgetRow, IEditorDataStorageProvider* DataStorage,
-	IEditorDataStorageUiProvider* DataStorageUi, const UE::Editor::DataStorage::FMetaDataView& Arguments)
+TSharedPtr<SWidget> FSimpleWidgetConstructor::Construct(UE::Editor::DataStorage::RowHandle WidgetRow, UE::Editor::DataStorage::ICoreProvider* DataStorage,
+	UE::Editor::DataStorage::IUiProvider* DataStorageUi, const UE::Editor::DataStorage::FMetaDataView& Arguments)
 {
 	const UE::Editor::DataStorage::RowHandle TargetRow = GetTargetRow(DataStorage, WidgetRow);
 
diff --git a/Engine/Source/Runtime/TypedElementFramework/Public/Elements/Common/TypedElementQueryConditions.h b/Engine/Source/Runtime/TypedElementFramework/Public/Elements/Common/TypedElementQueryConditions.h
--- a/Engine/Source/Runtime/TypedElementFramework/Public/Elements/Common/TypedElementQueryConditions.h
+++ b/Engine/Source/Runtime/TypedElementFramework/Public/Elements/Common/TypedElementQueryConditions.h
@@ -6,10 +6,13 @@
 #include "Templates/SubclassOf.h"
 #include "UObject/WeakObjectPtrTemplates.h"
 
-class IEditorDataStorageProvider;
 class UScriptStruct;
 
-namespace UE::Editor::DataStorage::Queries
+namespace UE::Editor::DataStorage
+{
+	class ICoreProvider;
+
+namespace Queries
 {
 	struct FColumnBase
 	{
@@ -47,16 +50,16 @@ namespace UE::Editor::DataStorage::Queries
 	};
 
 	/*
-	 * Specialized compile context that accepts an IEditorDataStorageProvider
+	 * Specialized compile context that accepts an ICoreProvider
 	 */
 	class FEditorStorageQueryConditionCompileContext : public IQueryConditionCompileContext 
 	{
 	public:
-		TYPEDELEMENTFRAMEWORK_API explicit FEditorStorageQueryConditionCompileContext(IEditorDataStorageProvider* InDataStorage);
+		TYPEDELEMENTFRAMEWORK_API explicit FEditorStorageQueryConditionCompileContext(ICoreProvider* InDataStorage);
 		TYPEDELEMENTFRAMEWORK_API virtual const UScriptStruct* GenerateDynamicColumn(const FDynamicColumnDescription& Description) const override;
 		
 	private:
-		IEditorDataStorageProvider* DataStorage;
+		ICoreProvider* DataStorage;
 	};
 
 
@@ -189,4 +192,6 @@ namespace UE::Editor::DataStorage::Queries
 	TYPEDELEMENTFRAMEWORK_API FConditions operator||(const FConditions& Lhs, const FConditions& Rhs);
 	TYPEDELEMENTFRAMEWORK_API FConditions operator||(FColumnBase Lhs, FColumnBase Rhs);
 	TYPEDELEMENTFRAMEWORK_API FConditions operator||(FColumnBase Lhs, const FConditions& Rhs);
-} // namespace UE::Editor::DataStorage::Queries
+} // namespace Queries
+} // namespace UE::Editor::DataStorage
+
diff --git a/Engine/Source/Runtime/TypedElementFramework/Public/Elements/Framework/TypedElementAttributeBinding.h b/Engine/Source/Runtime/TypedElementFramework/Public/Elements/Framework/TypedElementAttributeBinding.h
--- a/Engine/Source/Runtime/TypedElementFramework/Public/Elements/Framework/TypedElementAttributeBinding.h
+++ b/Engine/Source/Runtime/TypedElementFramework/Public/Elements/Framework/TypedElementAttributeBinding.h
@@ -8,10 +8,10 @@
 #include "Misc/Attribute.h"
 #include "TypedElementAttributeBindingText.h"
 
-class IEditorDataStorageProvider;
-
 namespace UE::Editor::DataStorage
 {
+	class ICoreProvider;
+
 	/**
 	 * Builder class that can be used as a shorthand to bind data inside a TEDS row, column pair to a TAttribute so the attribute updates if the data
 	 * in the column is changed.
@@ -29,7 +29,7 @@ namespace UE::Editor::DataStorage
 		TYPEDELEMENTFRAMEWORK_API FAttributeBinder(RowHandle InTargetRow);
 		
 		/* Create an attribute binder for a given row. */
-		TYPEDELEMENTFRAMEWORK_API FAttributeBinder(RowHandle InTargetRow, IEditorDataStorageProvider* InDataStorage);
+		TYPEDELEMENTFRAMEWORK_API FAttributeBinder(RowHandle InTargetRow, ICoreProvider* InDataStorage);
 
 		/**
 		 * Bind a specific data member inside a TEDS column to an attribute of the same type as the data
@@ -241,7 +241,7 @@ namespace UE::Editor::DataStorage
 		RowHandle TargetRow;
 
 		// A ptr to the data storage for quick access
-		IEditorDataStorageProvider* DataStorage;
+		ICoreProvider* DataStorage;
 	};
 } // namespace UE::Editor::DataStorage
 
diff --git a/Engine/Source/Runtime/TypedElementFramework/Public/Elements/Framework/TypedElementAttributeBindingProperty.h b/Engine/Source/Runtime/TypedElementFramework/Public/Elements/Framework/TypedElementAttributeBindingProperty.h
--- a/Engine/Source/Runtime/TypedElementFramework/Public/Elements/Framework/TypedElementAttributeBindingProperty.h
+++ b/Engine/Source/Runtime/TypedElementFramework/Public/Elements/Framework/TypedElementAttributeBindingProperty.h
@@ -25,13 +25,13 @@ namespace UE::Editor::DataStorage
 		concept AttributeBinderColumnDataInvocable = std::is_invocable_v<std::decay_t<FunctionType>, const TWeakObjectPtr<const UScriptStruct>&, const void*>;
 
 		template <TColumnType ColumnType>
-		ColumnType* GetColumn(IEditorDataStorageProvider* DataStorage, RowHandle Row)
+		ColumnType* GetColumn(ICoreProvider* DataStorage, RowHandle Row)
 		{
 			return DataStorage->GetColumn<ColumnType>(Row);
 		}
 		
 		template <TDynamicColumnTemplate ColumnType>
-		ColumnType* GetColumn(IEditorDataStorageProvider* DataStorage, RowHandle Row, const FName& InIdentifier)
+		ColumnType* GetColumn(ICoreProvider* DataStorage, RowHandle Row, const FName& InIdentifier)
 		{
 			if(ensureMsgf(InIdentifier != NAME_None, TEXT("None identifier passed to dynamic column version of GetColumn. Will always return nullptr")))
 			{
diff --git a/Engine/Source/Runtime/TypedElementFramework/Public/Elements/Framework/TypedElementAttributeBindingText.h b/Engine/Source/Runtime/TypedElementFramework/Public/Elements/Framework/TypedElementAttributeBindingText.h
--- a/Engine/Source/Runtime/TypedElementFramework/Public/Elements/Framework/TypedElementAttributeBindingText.h
+++ b/Engine/Source/Runtime/TypedElementFramework/Public/Elements/Framework/TypedElementAttributeBindingText.h
@@ -51,7 +51,7 @@ namespace UE::Editor::DataStorage
 		TYPEDELEMENTFRAMEWORK_API operator TAttribute<FText>();
 
 	private:
-		FTextAttributeFormatted(FTextFormat InFormat, RowHandle InTargetRow, IEditorDataStorageProvider* InDataStorage);
+		FTextAttributeFormatted(FTextFormat InFormat, RowHandle InTargetRow, ICoreProvider* InDataStorage);
 
 		struct FPropertyInfo
 		{
@@ -65,7 +65,7 @@ namespace UE::Editor::DataStorage
 		// The target row for this binder
 		RowHandle TargetRow;
 		// A ptr to the data storage for quick access
-		IEditorDataStorageProvider* DataStorage;
+		ICoreProvider* DataStorage;
 	};
 
 	template <TDataColumnType ColumnType>
diff --git a/Engine/Source/Runtime/TypedElementFramework/Public/Elements/Framework/TypedElementRegistry.h b/Engine/Source/Runtime/TypedElementFramework/Public/Elements/Framework/TypedElementRegistry.h
--- a/Engine/Source/Runtime/TypedElementFramework/Public/Elements/Framework/TypedElementRegistry.h
+++ b/Engine/Source/Runtime/TypedElementFramework/Public/Elements/Framework/TypedElementRegistry.h
@@ -78,19 +78,19 @@ public:
 	static TYPEDELEMENTFRAMEWORK_API UTypedElementRegistry* GetInstance();
 
 	UE_DEPRECATED(5.5, "Function has moved to modular features. Use UE::Editor::DataStorage::GetMutableDataStorageFeature instead.")
-	TYPEDELEMENTFRAMEWORK_API IEditorDataStorageProvider* GetMutableDataStorage();
+	TYPEDELEMENTFRAMEWORK_API UE::Editor::DataStorage::ICoreProvider* GetMutableDataStorage();
 	UE_DEPRECATED(5.5, "Function has moved to modular features. Use UE::Editor::DataStorage::GetDataStorageFeature instead.")
-	TYPEDELEMENTFRAMEWORK_API const IEditorDataStorageProvider* GetDataStorage() const;
+	TYPEDELEMENTFRAMEWORK_API const UE::Editor::DataStorage::ICoreProvider* GetDataStorage() const;
 
 	UE_DEPRECATED(5.5, "Function has moved to modular features. Use UE::Editor::DataStorage::GetMutableDataStorageFeature instead.")
-	TYPEDELEMENTFRAMEWORK_API IEditorDataStorageCompatibilityProvider* GetMutableDataStorageCompatibility();
+	TYPEDELEMENTFRAMEWORK_API UE::Editor::DataStorage::ICompatibilityProvider* GetMutableDataStorageCompatibility();
 	UE_DEPRECATED(5.5, "Function has moved to modular features. Use UE::Editor::DataStorage::GetDataStorageFeature instead.")
-	TYPEDELEMENTFRAMEWORK_API const IEditorDataStorageCompatibilityProvider* GetDataStorageCompatibility() const;
+	TYPEDELEMENTFRAMEWORK_API const UE::Editor::DataStorage::ICompatibilityProvider* GetDataStorageCompatibility() const;
 
 	UE_DEPRECATED(5.5, "Function has moved to modular features. Use UE::Editor::DataStorage::GetMutableDataStorageFeature instead.")
-	TYPEDELEMENTFRAMEWORK_API IEditorDataStorageUiProvider* GetMutableDataStorageUi();
+	TYPEDELEMENTFRAMEWORK_API UE::Editor::DataStorage::IUiProvider* GetMutableDataStorageUi();
 	UE_DEPRECATED(5.5, "Function has moved to modular features. Use UE::Editor::DataStorage::GetDataStorageFeature instead.")
-	TYPEDELEMENTFRAMEWORK_API const IEditorDataStorageUiProvider* GetDataStorageUi() const;
+	TYPEDELEMENTFRAMEWORK_API const UE::Editor::DataStorage::IUiProvider* GetDataStorageUi() const;
 
 	UE_DEPRECATED(5.5, "Use UE::Editor::DataStorage::AreDataStorageInterfacesSet instead.")
 	TYPEDELEMENTFRAMEWORK_API bool AreDataStorageInterfacesSet() const;
diff --git a/Engine/Source/Runtime/TypedElementFramework/Public/Elements/Interfaces/TypedElementDataStorageCompatibilityInterface.h b/Engine/Source/Runtime/TypedElementFramework/Public/Elements/Interfaces/TypedElementDataStorageCompatibilityInterface.h
--- a/Engine/Source/Runtime/TypedElementFramework/Public/Elements/Interfaces/TypedElementDataStorageCompatibilityInterface.h
+++ b/Engine/Source/Runtime/TypedElementFramework/Public/Elements/Interfaces/TypedElementDataStorageCompatibilityInterface.h
@@ -12,15 +12,17 @@
 #include "UObject/StrongObjectPtr.h"
 #include "UObject/WeakObjectPtrTemplates.h"
 
+namespace UE::Editor::DataStorage
+{
 /**
  * Interface to provide compatibility with existing systems that don't directly
  * support the data storage.
  */
-class IEditorDataStorageCompatibilityProvider : public IModularFeature
+class ICompatibilityProvider : public IModularFeature
 {
 public:
-	using ObjectRegistrationFilter = TFunction<bool(const IEditorDataStorageCompatibilityProvider&, const UObject*)>;
-	using ObjectToRowDealiaser = TFunction<UE::Editor::DataStorage::RowHandle(const IEditorDataStorageCompatibilityProvider&, const UObject*)>;
+	using ObjectRegistrationFilter = TFunction<bool(const ICompatibilityProvider&, const UObject*)>;
+	using ObjectToRowDealiaser = TFunction<RowHandle(const ICompatibilityProvider&, const UObject*)>;
 	
 	/**
 	 * @section Type-agnostic functions
@@ -35,14 +37,14 @@ public:
 	 * with a row and to setup the initial row data.
 	 */
 	template<typename ObjectType>
-	UE::Editor::DataStorage::RowHandle AddCompatibleObject(ObjectType&& Object);
+	RowHandle AddCompatibleObject(ObjectType&& Object);
 	
 	/** Removes a previously registered object from the data storage. */
 	template<typename ObjectType>
 	void RemoveCompatibleObject(ObjectType&& Object);
 
 	template<typename ObjectType>
-	UE::Editor::DataStorage::RowHandle FindRowWithCompatibleObject(ObjectType&& Object) const;
+	RowHandle FindRowWithCompatibleObject(ObjectType&& Object) const;
 
 	/**
 	 * @section Callback registration
@@ -67,7 +69,7 @@ public:
 	 * will be used to find the closest match in the registered types and use the associated table. E.g. actors derive from uobjects so
 	 * if the type information of an actor is registered the actor table will be used instead of the uobject table.
 	 */
-	virtual void RegisterTypeTableAssociation(TObjectPtr<UStruct> TypeInfo, UE::Editor::DataStorage::TableHandle Table) = 0;
+	virtual void RegisterTypeTableAssociation(TObjectPtr<UStruct> TypeInfo, TableHandle Table) = 0;
 
 	/**
 	 * @section Explicit functions
@@ -75,9 +77,9 @@ public:
 	 */
 
 	/** Adds a UObject to the data storage. */
-	virtual UE::Editor::DataStorage::RowHandle AddCompatibleObjectExplicit(UObject* Object) = 0;
+	virtual RowHandle AddCompatibleObjectExplicit(UObject* Object) = 0;
 	/** Adds an FStruct to the data storage. */
-	virtual UE::Editor::DataStorage::RowHandle AddCompatibleObjectExplicit(void* Object, TWeakObjectPtr<UScriptStruct> TypeInfo) = 0;
+	virtual RowHandle AddCompatibleObjectExplicit(void* Object, TWeakObjectPtr<UScriptStruct> TypeInfo) = 0;
 	
 	/** Removes a UObject from the data storage. */
 	virtual void RemoveCompatibleObjectExplicit(UObject* Object) = 0;
@@ -85,9 +87,9 @@ public:
 	virtual void RemoveCompatibleObjectExplicit(void* Object) = 0;
 
 	/** Finds a previously stored UObject. If not found an invalid row handle will be returned. */
-	virtual UE::Editor::DataStorage::RowHandle FindRowWithCompatibleObjectExplicit(const UObject* Object) const = 0;
+	virtual RowHandle FindRowWithCompatibleObjectExplicit(const UObject* Object) const = 0;
 	/** Finds a previously stored FStruct. If not found an invalid row handle will be returned. */
-	virtual UE::Editor::DataStorage::RowHandle FindRowWithCompatibleObjectExplicit(const void* Object) const = 0;
+	virtual RowHandle FindRowWithCompatibleObjectExplicit(const void* Object) const = 0;
 
 	/**
 	 * @section Miscellaneous functions
@@ -109,7 +111,7 @@ template<typename Type> Type* GetRawPointer(Type* Object)						{ return Object;
 template<typename Type> Type* GetRawPointer(Type& Object)						{ return &Object; }
 
 template<typename ObjectType>
-UE::Editor::DataStorage::RowHandle IEditorDataStorageCompatibilityProvider::AddCompatibleObject(ObjectType&& Object)
+RowHandle ICompatibilityProvider::AddCompatibleObject(ObjectType&& Object)
 {
 	auto RawPointer = GetRawPointer(Forward<ObjectType>(Object));
 	using BaseType = std::remove_cv_t<std::remove_pointer_t<decltype(RawPointer)>>;
@@ -125,13 +127,13 @@ UE::Editor::DataStorage::RowHandle IEditorDataStorageCompatibilityProvider::AddC
 }
 
 template<typename ObjectType>
-void IEditorDataStorageCompatibilityProvider::RemoveCompatibleObject(ObjectType&& Object)
+void ICompatibilityProvider::RemoveCompatibleObject(ObjectType&& Object)
 {
 	RemoveCompatibleObjectExplicit(GetRawPointer(Forward<ObjectType>(Object)));
 }
 
 template<typename ObjectType>
-UE::Editor::DataStorage::RowHandle IEditorDataStorageCompatibilityProvider::FindRowWithCompatibleObject(ObjectType&& Object) const
+RowHandle ICompatibilityProvider::FindRowWithCompatibleObject(ObjectType&& Object) const
 {
 	return FindRowWithCompatibleObjectExplicit(GetRawPointer(Forward<ObjectType>(Object)));
 }
@@ -190,4 +192,5 @@ struct TTypedElementSubsystemTraits final
 			return false;
 		}
 	}
-};
\ No newline at end of file
+};
+} // namespace UE::Editor::DataStorage
diff --git a/Engine/Source/Runtime/TypedElementFramework/Public/Elements/Interfaces/TypedElementDataStorageFactory.h b/Engine/Source/Runtime/TypedElementFramework/Public/Elements/Interfaces/TypedElementDataStorageFactory.h
--- a/Engine/Source/Runtime/TypedElementFramework/Public/Elements/Interfaces/TypedElementDataStorageFactory.h
+++ b/Engine/Source/Runtime/TypedElementFramework/Public/Elements/Interfaces/TypedElementDataStorageFactory.h
@@ -4,9 +4,12 @@
 
 #include "TypedElementDataStorageFactory.generated.h"
 
-class IEditorDataStorageCompatibilityProvider;
-class IEditorDataStorageProvider;
-class IEditorDataStorageUiProvider;
+namespace UE::Editor::DataStorage
+{
+	class ICoreProvider;
+	class IUiProvider;
+	class ICompatibilityProvider;
+}
 
 /**
  * Base class that can be used to register various elements, such as queries and widgets, with
@@ -29,22 +32,22 @@ public:
 	/**
 	 * All factories will have this called before any Register functions on any factories are called
 	 */
-	virtual void PreRegister(IEditorDataStorageProvider& DataStorage) {}
+	virtual void PreRegister(UE::Editor::DataStorage::ICoreProvider& DataStorage) {}
 
-	virtual void RegisterTables(IEditorDataStorageProvider& DataStorage) {}
-	virtual void RegisterTables(IEditorDataStorageProvider& DataStorage, IEditorDataStorageCompatibilityProvider& DataStorageCompatibility) {}
-	virtual void RegisterTickGroups(IEditorDataStorageProvider& DataStorage) const {}
-	virtual void RegisterQueries(IEditorDataStorageProvider& DataStorage) {}
+	virtual void RegisterTables(UE::Editor::DataStorage::ICoreProvider& DataStorage) {}
+	virtual void RegisterTables(UE::Editor::DataStorage::ICoreProvider& DataStorage, UE::Editor::DataStorage::ICompatibilityProvider& DataStorageCompatibility) {}
+	virtual void RegisterTickGroups(UE::Editor::DataStorage::ICoreProvider& DataStorage) const {}
+	virtual void RegisterQueries(UE::Editor::DataStorage::ICoreProvider& DataStorage) {}
 
-	virtual void RegisterRegistrationFilters(IEditorDataStorageCompatibilityProvider& DataStorageCompatibility) const {}
-	virtual void RegisterDealiaser(IEditorDataStorageCompatibilityProvider& DataStorageCompatibility) const {}
+	virtual void RegisterRegistrationFilters(UE::Editor::DataStorage::ICompatibilityProvider& DataStorageCompatibility) const {}
+	virtual void RegisterDealiaser(UE::Editor::DataStorage::ICompatibilityProvider& DataStorageCompatibility) const {}
 	
-	virtual void RegisterWidgetPurposes(IEditorDataStorageUiProvider& DataStorageUi) const {}
-	virtual void RegisterWidgetConstructors(IEditorDataStorageProvider& DataStorage,
-		IEditorDataStorageUiProvider& DataStorageUi) const {}
+	virtual void RegisterWidgetPurposes(UE::Editor::DataStorage::IUiProvider& DataStorageUi) const {}
+	virtual void RegisterWidgetConstructors(UE::Editor::DataStorage::ICoreProvider& DataStorage,
+		UE::Editor::DataStorage::IUiProvider& DataStorageUi) const {}
 
 	/**
 	 * Called in reverse order before the DataStorage object is shut down
 	 */
-	virtual void PreShutdown(IEditorDataStorageProvider& DataStorage) {}
+	virtual void PreShutdown(UE::Editor::DataStorage::ICoreProvider& DataStorage) {}
 };
diff --git a/Engine/Source/Runtime/TypedElementFramework/Public/Elements/Interfaces/TypedElementDataStorageInterface.h b/Engine/Source/Runtime/TypedElementFramework/Public/Elements/Interfaces/TypedElementDataStorageInterface.h
--- a/Engine/Source/Runtime/TypedElementFramework/Public/Elements/Interfaces/TypedElementDataStorageInterface.h
+++ b/Engine/Source/Runtime/TypedElementFramework/Public/Elements/Interfaces/TypedElementDataStorageInterface.h
@@ -27,6 +27,8 @@ class USubsystem;
 class UScriptStruct;
 class UEditorDataStorageFactory;
 
+namespace UE::Editor::DataStorage
+{
 using FTypedElementOnDataStorageCreation = FSimpleMulticastDelegate;
 using FTypedElementOnDataStorageDestruction = FSimpleMulticastDelegate;
 using FTypedElementOnDataStorageUpdate = FSimpleMulticastDelegate;
@@ -37,7 +39,7 @@ using FTypedElementOnDataStorageUpdate = FSimpleMulticastDelegate;
  * multiple column lists are used. Note that the returned array view is only available while
  * this object is constructed, so care must be taken with functions that return a const array view.
  */
-template<UE::Editor::DataStorage::TColumnType... Columns>
+template<TColumnType... Columns>
 struct TTypedElementColumnTypeList
 {
 	const UScriptStruct* ColumnTypes[sizeof...(Columns)] = { Columns::StaticStruct()... };
@@ -45,17 +47,8 @@ struct TTypedElementColumnTypeList
 	operator TConstArrayView<const UScriptStruct*>() const { return ColumnTypes; }
 };
 
-class IEditorDataStorageProvider : public IModularFeature
+class ICoreProvider : public IModularFeature
 {
-	using RowCreationCallbackRef = UE::Editor::DataStorage::RowCreationCallbackRef;
-	using ColumnCreationCallbackRef = UE::Editor::DataStorage::ColumnCreationCallbackRef;
-	using ColumnListCallbackRef = UE::Editor::DataStorage::ColumnListCallbackRef;
-	using ColumnListWithDataCallbackRef = UE::Editor::DataStorage::ColumnListWithDataCallbackRef;
-	using ColumnCopyOrMoveCallback = UE::Editor::DataStorage::ColumnCopyOrMoveCallback;
-	using RowHandle = UE::Editor::DataStorage::RowHandle;
-	using TableHandle = UE::Editor::DataStorage::TableHandle;
-	using QueryHandle = UE::Editor::DataStorage::QueryHandle;
-
 public:
 	/**
 	 * @section Factories
@@ -83,7 +76,7 @@ public:
 
 	/** Creates a new table for with the provided columns. Optionally a name can be given which is useful for retrieval later. */
 	virtual TableHandle RegisterTable(TConstArrayView<const UScriptStruct*> ColumnList, const FName& Name) = 0;
-	template<UE::Editor::DataStorage::TColumnType... Columns>
+	template<TColumnType... Columns>
 	TableHandle RegisterTable(const FName& Name);
 	/** 
 	 * Copies the column information from the provided table and creates a new table for with the provided columns. Optionally a 
@@ -91,7 +84,7 @@ public:
 	 */
 	virtual TableHandle RegisterTable(TableHandle SourceTable,
 		TConstArrayView<const UScriptStruct*> ColumnList, const FName& Name) = 0;
-	template<UE::Editor::DataStorage::TColumnType... Columns>
+	template<TColumnType... Columns>
 	TableHandle RegisterTable(TableHandle SourceTable, const FName& Name);
 
 	/** Returns a previously created table with the provided name or TypedElementInvalidTableHandle if not found. */
@@ -169,7 +162,7 @@ public:
 	 * This can be used to for instance remove all rows with a specific type tag, such as removing all rows with
 	 * an entity tag to remove all entities from the data storage.
 	 */
-	template<UE::Editor::DataStorage::TColumnType... Columns>
+	template<TColumnType... Columns>
 	void RemoveAllRowsWith();
 
 	/** Checks whether or not a row is in use. This is true even if the row has only been reserved. */
@@ -183,16 +176,16 @@ public:
 
 	/** Adds a column to a row or does nothing if already added. */
 	virtual void AddColumn(RowHandle Row, const UScriptStruct* ColumnType) = 0;
-	template<UE::Editor::DataStorage::TColumnType ColumnType>
+	template<TColumnType ColumnType>
 	void AddColumn(RowHandle Row);
 	/**
 	 * Adds a new data column and initializes it. The relocator will be used to copy or move the column out of
 	 * its temporary location into the final table if the addition needs to be deferred.
 	 */
 	virtual void AddColumnData(RowHandle Row, const UScriptStruct* ColumnType,
-		const ColumnCreationCallbackRef& Initializer,
+		const ColumnCreationCallbackRef Initializer,
 		ColumnCopyOrMoveCallback Relocator) = 0;
-	template<UE::Editor::DataStorage::TDataColumnType ColumnType>
+	template<TDataColumnType ColumnType>
 	void AddColumn(RowHandle Row, ColumnType&& Column);
 
 	/**
@@ -206,7 +199,7 @@ public:
 	 *       is added.  This will likely change in the future to transparently replace the tag to have consistent behaviour with other usages
 	 *       of AddColumn
 	 */
-	virtual void AddColumn(RowHandle Row, const UE::Editor::DataStorage::FValueTag& Tag, const FName& Value) = 0;
+	virtual void AddColumn(RowHandle Row, const FValueTag& Tag, const FName& Value) = 0;
 
 	template<typename T>
 	void AddColumn(RowHandle Row, const FName& Tag) = delete;
@@ -215,18 +208,18 @@ public:
 	void AddColumn(RowHandle Row, const FName& Tag, const FName& Value) = delete;
 	
 	template<>
-	void AddColumn<UE::Editor::DataStorage::FValueTag>(RowHandle Row, const FName& Tag, const FName& Value);
+	void AddColumn<FValueTag>(RowHandle Row, const FName& Tag, const FName& Value);
 
-	template<UE::Editor::DataStorage::TEnumType EnumT>
+	template<TEnumType EnumT>
 	void AddColumn(RowHandle Row, EnumT Value);
 	
-	template<auto Value, UE::Editor::DataStorage::TEnumType EnumT = decltype(Value)>
+	template<auto Value, TEnumType EnumT = decltype(Value)>
 	void AddColumn(RowHandle Row);
 
-	template<UE::Editor::DataStorage::TDynamicColumnTemplate DynamicColumnTemplate>
+	template<TDynamicColumnTemplate DynamicColumnTemplate>
 	void AddColumn(RowHandle Row, const FName& Identifier);
 	
-	template<UE::Editor::DataStorage::TDynamicColumnTemplate DynamicColumnTemplate>
+	template<TDynamicColumnTemplate DynamicColumnTemplate>
 	void AddColumn(RowHandle Row, const FName& Identifier, DynamicColumnTemplate&& TemplateInstance);
 
 	/**
@@ -234,38 +227,38 @@ public:
 	 * at a time.
 	 */
 	virtual void AddColumns(RowHandle Row, TConstArrayView<const UScriptStruct*> Columns) = 0;
-	template<UE::Editor::DataStorage::TColumnType... Columns>
+	template<TColumnType... Columns>
 	void AddColumns(RowHandle Row);
 
 	/** Removes a column from a row or does nothing if already removed. */
 	virtual void RemoveColumn(RowHandle Row, const UScriptStruct* ColumnType) = 0;
-	template<UE::Editor::DataStorage::TColumnType Column>
+	template<TColumnType Column>
 	void RemoveColumn(RowHandle Row);
 
-	template<UE::Editor::DataStorage::TEnumType EnumT>
+	template<TEnumType EnumT>
 	void RemoveColumn(RowHandle Row);
 
 	/**
 	 * Removes a value tag from the given row
 	 * If tag does not exist on row, operation will do nothing.
 	 */
-	virtual void RemoveColumn(RowHandle Row, const UE::Editor::DataStorage::FValueTag& Tag) = 0;
+	virtual void RemoveColumn(RowHandle Row, const FValueTag& Tag) = 0;
 
 	template<typename T>
 	void RemoveColumn(RowHandle Row, const FName& Tag) = delete;
 
-	template<UE::Editor::DataStorage::TDynamicColumnTemplate DynamicColumnTemplateType>
+	template<TDynamicColumnTemplate DynamicColumnTemplateType>
 	void RemoveColumn(RowHandle Row, const FName& Identifier);
 	
 	template<>
-	void RemoveColumn<UE::Editor::DataStorage::FValueTag>(RowHandle Row, const FName& Tag);
+	void RemoveColumn<FValueTag>(RowHandle Row, const FName& Tag);
 
 	/**
 	 * Removes multiple columns from a row. This is typically more efficient than adding columns one
 	 * at a time.
 	 */
 	virtual void RemoveColumns(RowHandle Row, TConstArrayView<const UScriptStruct*> Columns) = 0;
-	template<UE::Editor::DataStorage::TColumnType... Columns>
+	template<TColumnType... Columns>
 	void RemoveColumns(RowHandle Row);
 
 	/** 
@@ -286,14 +279,14 @@ public:
 	virtual void* GetColumnData(RowHandle Row, const UScriptStruct* ColumnType) = 0;
 	virtual const void* GetColumnData(RowHandle Row, const UScriptStruct* ColumnType) const = 0;
 	/** Returns a pointer to the column of the given row or a nullptr if the type couldn't be found or the row doesn't exist. */
-	template<UE::Editor::DataStorage::TDataColumnType ColumnType>
+	template<TDataColumnType ColumnType>
 	ColumnType* GetColumn(RowHandle Row);
-	template<UE::Editor::DataStorage::TDataColumnType ColumnType>
+	template<TDataColumnType ColumnType>
 	const ColumnType* GetColumn(RowHandle Row) const;
 	// Gets a dynamic column identified by the ColumnTypeTemplate and Identifier
-	template<UE::Editor::DataStorage::TDynamicColumnTemplate ColumnTypeTemplate>
+	template<TDynamicColumnTemplate ColumnTypeTemplate>
 	ColumnTypeTemplate* GetColumn(RowHandle Row, const FName& Identifer);
-	template<UE::Editor::DataStorage::TDynamicColumnTemplate ColumnTypeTemplate>
+	template<TDynamicColumnTemplate ColumnTypeTemplate>
 	const ColumnTypeTemplate* GetColumn(RowHandle Row, const FName& Identifer) const;
 	
 	/** 
@@ -303,7 +296,7 @@ public:
 	 */
 	virtual bool HasColumns(RowHandle Row, TConstArrayView<const UScriptStruct*> ColumnTypes) const = 0;
 	virtual bool HasColumns(RowHandle Row, TConstArrayView<TWeakObjectPtr<const UScriptStruct>> ColumnTypes) const = 0;
-	template<UE::Editor::DataStorage::TColumnType... ColumnTypes>
+	template<TColumnType... ColumnTypes>
 	bool HasColumns(RowHandle Row) const;
 
 	/** 
@@ -322,21 +315,21 @@ public:
 	virtual void ListColumns(RowHandle Row, ColumnListWithDataCallbackRef Callback) = 0;
 
 	/** Determines if the columns in the row match the query conditions. */
-	virtual bool MatchesColumns(RowHandle Row, const UE::Editor::DataStorage::Queries::FConditions& Conditions) const = 0;
+	virtual bool MatchesColumns(RowHandle Row, const Queries::FConditions& Conditions) const = 0;
 
 	/**
 	 * Finds the type information for a dynamic column.
 	 * If the dynamic column has not been generated, then return nullptr
-	 * The TemplateType may be a typed derived from either UE::Editor::DataStorage::FColumn or UE::Editor::DataStorage::FTag, anything else will return nullptr
+	 * The TemplateType may be a typed derived from either FColumn or FTag, anything else will return nullptr
 	 */
-	virtual const UScriptStruct* FindDynamicColumn(const UE::Editor::DataStorage::FDynamicColumnDescription& Description) const = 0;
+	virtual const UScriptStruct* FindDynamicColumn(const FDynamicColumnDescription& Description) const = 0;
 
 	/**
 	 * Generates a new dynamic column from a Template.  A dynamic column is uniquely identified using the given template and an Identifier
 	 * This function is idempotent - multiple calls with the same parameters will result in subsequent calls returning the same type
-	 * The TemplateType may be a typed derived from either UE::Editor::DataStorage::FColumn or UE::Editor::DataStorage::FTag
+	 * The TemplateType may be a typed derived from either FColumn or FTag
 	 */
-	virtual const UScriptStruct* GenerateDynamicColumn(const UE::Editor::DataStorage::FDynamicColumnDescription& Description) = 0;
+	virtual const UScriptStruct* GenerateDynamicColumn(const FDynamicColumnDescription& Description) = 0;
 
 	/**
 	 * Executes the given callback for each known dynamic column that derives from the base template provided
@@ -346,7 +339,7 @@ public:
 	/**
 	 * Executes the given callback for each known dynamic column that derives from the base template provided
 	 */
-	template<UE::Editor::DataStorage::TDynamicColumnTemplate DynamicColumnTemplateType>
+	template<TDynamicColumnTemplate DynamicColumnTemplateType>
 	void ForEachDynamicColumn(TFunctionRef<void(const UScriptStruct& Type)> Callback) const;
 
 	/**
@@ -363,25 +356,6 @@ public:
 	 * therefore recommended to try to simplify the query first before relying on extended query filtering in a processor.
 	 */
 
-	using EQueryTickPhase = UE::Editor::DataStorage::EQueryTickPhase;
-	using EQueryTickGroups = UE::Editor::DataStorage::EQueryTickGroups;
-	using EQueryCallbackType = UE::Editor::DataStorage::EQueryCallbackType;
-	using EQueryAccessType = UE::Editor::DataStorage::EQueryAccessType;
-	using EQueryDependencyFlags = UE::Editor::DataStorage::EQueryDependencyFlags;
-	using FQueryResult = UE::Editor::DataStorage::FQueryResult;
-
-	using IQueryContext = UE::Editor::DataStorage::IQueryContext;
-	using IDirectQueryContext = UE::Editor::DataStorage::IDirectQueryContext;
-	using ISubqueryContext = UE::Editor::DataStorage::ISubqueryContext;
-
-	using FQueryDescription = UE::Editor::DataStorage::FQueryDescription;
-	using QueryCallback = UE::Editor::DataStorage::QueryCallback;
-	using QueryCallbackRef = UE::Editor::DataStorage::QueryCallbackRef;
-	using DirectQueryCallback = UE::Editor::DataStorage::DirectQueryCallback;
-	using DirectQueryCallbackRef = UE::Editor::DataStorage::DirectQueryCallbackRef;
-	using SubqueryCallback = UE::Editor::DataStorage::SubqueryCallback;
-	using SubqueryCallbackRef = UE::Editor::DataStorage::SubqueryCallbackRef;
-
 	/** 
 	 * Registers a query with the data storage. The description is processed into an internal format and may be changed. If no valid
 	 * could be created an invalid query handle will be returned. It's recommended to use the Query Builder for a more convenient
@@ -409,7 +383,7 @@ public:
 	 * Directly runs a query. The callback will be called for batches of matching rows. During a single call to RunQuery the callback
 	 * may be called multiple times. If the query handle is invalid or has been deleted nothing happens and the callback won't be called.
 	 */
-	virtual FQueryResult RunQuery(QueryHandle Query, UE::Editor::DataStorage::EDirectQueryExecutionFlags Flags,
+	virtual FQueryResult RunQuery(QueryHandle Query, EDirectQueryExecutionFlags Flags,
 		DirectQueryCallbackRef Callback) = 0;
 	/**
 	 * Triggers all queries registered under the activation name to run for one update cycle. The activatable queries will be activated at
@@ -427,20 +401,20 @@ public:
 	 */
 
 	/** Retrieves the row for a mapped object. Returns an invalid row handle if the no row with the provided key was found. */
-	virtual RowHandle LookupMappedRow(const UE::Editor::DataStorage::FMapKeyView& Key) const = 0;
+	virtual RowHandle LookupMappedRow(const FMapKeyView& Key) const = 0;
 	/** 
 	 * Registers a row under a key. The same row can be registered multiple, but an key can only be associated with a single row.
 	 */
-	virtual void MapRow(UE::Editor::DataStorage::FMapKey Key, RowHandle Row) = 0;
+	virtual void MapRow(FMapKey Key, RowHandle Row) = 0;
 	/**
 	 * Register multiple rows under their key. The same row can be registered multiple times, but key can only be associated with a single
 	 * row. During processing the keys will be moved out into the final location.
 	 */
-	virtual void BatchMapRows(TArrayView<TPair<UE::Editor::DataStorage::FMapKey, RowHandle>> MapRowPairs) = 0;
+	virtual void BatchMapRows(TArrayView<TPair<FMapKey, RowHandle>> MapRowPairs) = 0;
 	/** Updates the key of a row to a new value. Effectively this is the same as removing a key and adding a new one. */
-	virtual void RemapRow(const UE::Editor::DataStorage::FMapKeyView& OriginalKey, UE::Editor::DataStorage::FMapKey NewKey) = 0;
+	virtual void RemapRow(const FMapKeyView& OriginalKey, FMapKey NewKey) = 0;
 	/** Removes a previously registered key from the mapping table or does nothing if the key no longer exists. */
-	virtual void RemoveRowMapping(const UE::Editor::DataStorage::FMapKeyView& Key) = 0;
+	virtual void RemoveRowMapping(const FMapKeyView& Key) = 0;
 	
 	/**
 	 * @section Miscellaneous
@@ -479,76 +453,76 @@ public:
 	 * @section Deprecated
 	 */
 	UE_DEPRECATED(5.6, "Use 'LookUpMappedRow' instead and use the new FMapKey(View) instead of the explicit 'GenerateIndexHash'-functions.")
-	inline RowHandle FindIndexedRow(UE::Editor::DataStorage::IndexHash Index) const;
+	inline RowHandle FindIndexedRow(IndexHash Index) const;
 	UE_DEPRECATED(5.6, "Use 'MapRow' instead and use the new FMapKey(View) instead of the explicit 'GenerateIndexHash'-functions.")
-	inline void IndexRow(UE::Editor::DataStorage::IndexHash Index, RowHandle Row);
+	inline void IndexRow(IndexHash Index, RowHandle Row);
 	UE_DEPRECATED(5.6, "Use 'BatchMapRows' instead and use the new FMapKey(View) instead of the explicit 'GenerateIndexHash'-functions.")
-	inline void BatchIndexRows(TConstArrayView<TPair<UE::Editor::DataStorage::IndexHash, RowHandle>> IndexRowPairs);
+	inline void BatchIndexRows(TConstArrayView<TPair<IndexHash, RowHandle>> IndexRowPairs);
 	UE_DEPRECATED(5.6, "Use 'RemapRow' instead and use the new FMapKey(View) instead of the explicit 'GenerateIndexHash'-functions.")
-	inline void Reindex(UE::Editor::DataStorage::IndexHash OriginalIndex, UE::Editor::DataStorage::IndexHash NewIndex);
+	inline void Reindex(IndexHash OriginalIndex, IndexHash NewIndex);
 	UE_DEPRECATED(5.6, "Use 'RemovedRowMapping' instead and use the new FMapKey(View) instead of the explicit 'GenerateIndexHash'-functions.")
-	inline void RemoveIndex(UE::Editor::DataStorage::IndexHash Index);
+	inline void RemoveIndex(IndexHash Index);
 
 	UE_DEPRECATED(5.6, "Use 'RemapRow' instead and use the new FMapKey(View) instead of the explicit 'GenerateIndexHash'-functions. The Row argument is no longer used.")
-	inline void ReindexRow(UE::Editor::DataStorage::IndexHash OriginalIndex, UE::Editor::DataStorage::IndexHash NewIndex, RowHandle Row);
+	inline void ReindexRow(IndexHash OriginalIndex, IndexHash NewIndex, RowHandle Row);
 };
 
 // Implementations
 
 template <typename FactoryT>
-const FactoryT* IEditorDataStorageProvider::FindFactory() const
+const FactoryT* ICoreProvider::FindFactory() const
 {
 	return static_cast<const FactoryT*>(FindFactory(FactoryT::StaticClass()));
 }
 
-template<UE::Editor::DataStorage::TColumnType... Columns>
-UE::Editor::DataStorage::TableHandle IEditorDataStorageProvider::RegisterTable(const FName& Name)
+template<TColumnType... Columns>
+TableHandle ICoreProvider::RegisterTable(const FName& Name)
 {
 	return RegisterTable({ Columns::StaticStruct()... }, Name);
 }
 
-template<UE::Editor::DataStorage::TColumnType... Columns>
-UE::Editor::DataStorage::TableHandle IEditorDataStorageProvider::RegisterTable(
+template<TColumnType... Columns>
+TableHandle ICoreProvider::RegisterTable(
 	TableHandle SourceTable, const FName& Name)
 {
 	return RegisterTable(SourceTable, { Columns::StaticStruct()... }, Name);
 }
 
-template<UE::Editor::DataStorage::TColumnType Column>
-void IEditorDataStorageProvider::AddColumn(RowHandle Row)
+template<TColumnType Column>
+void ICoreProvider::AddColumn(RowHandle Row)
 {
 	AddColumn(Row, Column::StaticStruct());
 }
 
-template<UE::Editor::DataStorage::TColumnType Column>
-void IEditorDataStorageProvider::RemoveColumn(RowHandle Row)
+template<TColumnType Column>
+void ICoreProvider::RemoveColumn(RowHandle Row)
 {
 	RemoveColumn(Row, Column::StaticStruct());
 }
 
-template<UE::Editor::DataStorage::TColumnType... Columns>
-void IEditorDataStorageProvider::AddColumns(RowHandle Row)
+template<TColumnType... Columns>
+void ICoreProvider::AddColumns(RowHandle Row)
 {
 	AddColumns(Row, { Columns::StaticStruct()...});
 }
 
 template <>
-inline void IEditorDataStorageProvider::AddColumn<UE::Editor::DataStorage::FValueTag>(RowHandle Row, const FName& Tag, const FName& Value)
+inline void ICoreProvider::AddColumn<FValueTag>(RowHandle Row, const FName& Tag, const FName& Value)
 {
-	AddColumn(Row, UE::Editor::DataStorage::FValueTag(Tag), Value);
+	AddColumn(Row, FValueTag(Tag), Value);
 }
 
 template <>
-inline void IEditorDataStorageProvider::RemoveColumn<UE::Editor::DataStorage::FValueTag>(RowHandle Row, const FName& Tag)
+inline void ICoreProvider::RemoveColumn<FValueTag>(RowHandle Row, const FName& Tag)
 {
 	using namespace UE::Editor::DataStorage;
 	RemoveColumn(Row, FValueTag(Tag));
 }
 
-template<UE::Editor::DataStorage::TDynamicColumnTemplate DynamicColumnTemplate>
-void IEditorDataStorageProvider::RemoveColumn(RowHandle Row, const FName& Identifier)
+template<TDynamicColumnTemplate DynamicColumnTemplate>
+void ICoreProvider::RemoveColumn(RowHandle Row, const FName& Identifier)
 {
-	const UE::Editor::DataStorage::FDynamicColumnDescription Description
+	const FDynamicColumnDescription Description
 	{
 		.TemplateType = DynamicColumnTemplate::StaticStruct(),
 		.Identifier = Identifier
@@ -557,33 +531,33 @@ void IEditorDataStorageProvider::RemoveColumn(RowHandle Row, const FName& Identi
 	RemoveColumn(Row, StructInfo);
 }
 
-template<UE::Editor::DataStorage::TEnumType EnumT>
-void IEditorDataStorageProvider::AddColumn(RowHandle Row, EnumT Value)
+template<TEnumType EnumT>
+void ICoreProvider::AddColumn(RowHandle Row, EnumT Value)
 {
 	const UEnum* Enum = StaticEnum<EnumT>();
 	const FName ValueAsFName = *Enum->GetNameStringByValue(static_cast<int64>(Value));
 	if (ValueAsFName != NAME_None)
 	{
-		AddColumn(Row, UE::Editor::DataStorage::FValueTag(Enum->GetFName()), ValueAsFName);
+		AddColumn(Row, FValueTag(Enum->GetFName()), ValueAsFName);
 	}
 }
 
-template<UE::Editor::DataStorage::TColumnType... Columns>
-void IEditorDataStorageProvider::RemoveAllRowsWith()
+template<TColumnType... Columns>
+void ICoreProvider::RemoveAllRowsWith()
 {
 	RemoveAllRowsWithColumns({ Columns::StaticStruct()... });
 }
 
-template<auto Value, UE::Editor::DataStorage::TEnumType EnumT>
-void IEditorDataStorageProvider::AddColumn(RowHandle Row)
+template<auto Value, TEnumType EnumT>
+void ICoreProvider::AddColumn(RowHandle Row)
 {
 	AddColumn<EnumT>(Row, Value);
 }
 
-template <UE::Editor::DataStorage::TDynamicColumnTemplate DynamicColumnTemplate>
-void IEditorDataStorageProvider::AddColumn(RowHandle Row, const FName& Identifier)
+template <TDynamicColumnTemplate DynamicColumnTemplate>
+void ICoreProvider::AddColumn(RowHandle Row, const FName& Identifier)
 {
-	const UE::Editor::DataStorage::FDynamicColumnDescription Description
+	const FDynamicColumnDescription Description
 	{
 		.TemplateType = DynamicColumnTemplate::StaticStruct(),
 		.Identifier = Identifier
@@ -592,10 +566,10 @@ void IEditorDataStorageProvider::AddColumn(RowHandle Row, const FName& Identifie
 	AddColumn(Row, StructInfo);
 }
 
-template <UE::Editor::DataStorage::TDynamicColumnTemplate DynamicColumnTemplate>
-void IEditorDataStorageProvider::AddColumn(RowHandle Row, const FName& Identifier, DynamicColumnTemplate&& TemplateInstance)
+template <TDynamicColumnTemplate DynamicColumnTemplate>
+void ICoreProvider::AddColumn(RowHandle Row, const FName& Identifier, DynamicColumnTemplate&& TemplateInstance)
 {
-	const UE::Editor::DataStorage::FDynamicColumnDescription Description
+	const FDynamicColumnDescription Description
 	{
 		.TemplateType = DynamicColumnTemplate::StaticStruct(),
 		.Identifier = Identifier
@@ -627,21 +601,21 @@ void IEditorDataStorageProvider::AddColumn(RowHandle Row, const FName& Identifie
 	
 }
 
-template<UE::Editor::DataStorage::TEnumType EnumT>
-void IEditorDataStorageProvider::RemoveColumn(RowHandle Row)
+template<TEnumType EnumT>
+void ICoreProvider::RemoveColumn(RowHandle Row)
 {
 	const UEnum* Enum = StaticEnum<EnumT>();
-	RemoveColumn(Row, UE::Editor::DataStorage::FValueTag(Enum->GetFName()));
+	RemoveColumn(Row, FValueTag(Enum->GetFName()));
 }
 
-template<UE::Editor::DataStorage::TColumnType... Columns>
-void IEditorDataStorageProvider::RemoveColumns(RowHandle Row)
+template<TColumnType... Columns>
+void ICoreProvider::RemoveColumns(RowHandle Row)
 {
 	RemoveColumns(Row, { Columns::StaticStruct()...});
 }
 
-template<UE::Editor::DataStorage::TDataColumnType ColumnType>
-void IEditorDataStorageProvider::AddColumn(RowHandle Row, ColumnType&& Column)
+template<TDataColumnType ColumnType>
+void ICoreProvider::AddColumn(RowHandle Row, ColumnType&& Column)
 {
 	AddColumnData(Row, ColumnType::StaticStruct(),
 		[&Column](void* ColumnData, const UScriptStruct&)
@@ -668,22 +642,22 @@ void IEditorDataStorageProvider::AddColumn(RowHandle Row, ColumnType&& Column)
 		});
 }
 
-template<UE::Editor::DataStorage::TDataColumnType ColumnType>
-ColumnType* IEditorDataStorageProvider::GetColumn(RowHandle Row)
+template<TDataColumnType ColumnType>
+ColumnType* ICoreProvider::GetColumn(RowHandle Row)
 {
 	return reinterpret_cast<ColumnType*>(GetColumnData(Row, ColumnType::StaticStruct()));
 }
 
-template<UE::Editor::DataStorage::TDataColumnType ColumnType>
-const ColumnType* IEditorDataStorageProvider::GetColumn(RowHandle Row) const
+template<TDataColumnType ColumnType>
+const ColumnType* ICoreProvider::GetColumn(RowHandle Row) const
 {
 	return reinterpret_cast<const ColumnType*>(GetColumnData(Row, ColumnType::StaticStruct()));
 }
 
-template <UE::Editor::DataStorage::TDynamicColumnTemplate DynamicColumnTemplate>
-DynamicColumnTemplate* IEditorDataStorageProvider::GetColumn(RowHandle Row, const FName& Identifier)
+template <TDynamicColumnTemplate DynamicColumnTemplate>
+DynamicColumnTemplate* ICoreProvider::GetColumn(RowHandle Row, const FName& Identifier)
 {
-	const UE::Editor::DataStorage::FDynamicColumnDescription Description
+	const FDynamicColumnDescription Description
 	{
 		.TemplateType = DynamicColumnTemplate::StaticStruct(),
 		.Identifier = Identifier
@@ -696,10 +670,10 @@ DynamicColumnTemplate* IEditorDataStorageProvider::GetColumn(RowHandle Row, cons
 	return nullptr;
 }
 
-template <UE::Editor::DataStorage::TDynamicColumnTemplate DynamicColumnTemplate>
-const DynamicColumnTemplate* IEditorDataStorageProvider::GetColumn(RowHandle Row, const FName& Identifier) const
+template <TDynamicColumnTemplate DynamicColumnTemplate>
+const DynamicColumnTemplate* ICoreProvider::GetColumn(RowHandle Row, const FName& Identifier) const
 {
-	const UE::Editor::DataStorage::FDynamicColumnDescription Description
+	const FDynamicColumnDescription Description
 	{
 		.TemplateType = DynamicColumnTemplate::StaticStruct(),
 		.Identifier = Identifier
@@ -712,54 +686,56 @@ const DynamicColumnTemplate* IEditorDataStorageProvider::GetColumn(RowHandle Row
 	return nullptr;
 }
 
-template<UE::Editor::DataStorage::TColumnType... ColumnType>
-bool IEditorDataStorageProvider::HasColumns(RowHandle Row) const
+template<TColumnType... ColumnType>
+bool ICoreProvider::HasColumns(RowHandle Row) const
 {
 	return HasColumns(Row, TConstArrayView<const UScriptStruct*>({ ColumnType::StaticStruct()... }));
 }
 
 template<typename SystemType>
-SystemType* IEditorDataStorageProvider::GetExternalSystem()
+SystemType* ICoreProvider::GetExternalSystem()
 {
 	return reinterpret_cast<SystemType*>(GetExternalSystemAddress(SystemType::StaticClass()));
 }
 
-void IEditorDataStorageProvider::ReindexRow(
-	UE::Editor::DataStorage::IndexHash OriginalIndex, UE::Editor::DataStorage::IndexHash NewIndex, RowHandle Row)
+void ICoreProvider::ReindexRow(
+	IndexHash OriginalIndex, IndexHash NewIndex, RowHandle Row)
 {
-	RemapRow(UE::Editor::DataStorage::FMapKeyView(OriginalIndex), UE::Editor::DataStorage::FMapKey(NewIndex));
+	RemapRow(FMapKeyView(OriginalIndex), FMapKey(NewIndex));
 }
 
-UE::Editor::DataStorage::RowHandle IEditorDataStorageProvider::FindIndexedRow(UE::Editor::DataStorage::IndexHash Index) const
+RowHandle ICoreProvider::FindIndexedRow(IndexHash Index) const
 {
-	return LookupMappedRow(UE::Editor::DataStorage::FMapKeyView(Index));
+	return LookupMappedRow(FMapKeyView(Index));
 }
 
-void IEditorDataStorageProvider::IndexRow(UE::Editor::DataStorage::IndexHash Index, RowHandle Row)
+void ICoreProvider::IndexRow(IndexHash Index, RowHandle Row)
 {
-	MapRow(UE::Editor::DataStorage::FMapKey(Index), Row);
+	MapRow(FMapKey(Index), Row);
 }
 
-void IEditorDataStorageProvider::BatchIndexRows(TConstArrayView<TPair<UE::Editor::DataStorage::IndexHash, RowHandle>> IndexRowPairs)
+void ICoreProvider::BatchIndexRows(TConstArrayView<TPair<IndexHash, RowHandle>> IndexRowPairs)
 {
-	for (const TPair<UE::Editor::DataStorage::IndexHash, RowHandle>& Pair : IndexRowPairs)
+	for (const TPair<IndexHash, RowHandle>& Pair : IndexRowPairs)
 	{
-		MapRow(UE::Editor::DataStorage::FMapKey(Pair.Key), Pair.Value);
+		MapRow(FMapKey(Pair.Key), Pair.Value);
 	}
 }
 
-void IEditorDataStorageProvider::Reindex(UE::Editor::DataStorage::IndexHash OriginalIndex, UE::Editor::DataStorage::IndexHash NewIndex)
+void ICoreProvider::Reindex(IndexHash OriginalIndex, IndexHash NewIndex)
 {
-	RemapRow(UE::Editor::DataStorage::FMapKeyView(OriginalIndex), UE::Editor::DataStorage::FMapKey(NewIndex));
+	RemapRow(FMapKeyView(OriginalIndex), FMapKey(NewIndex));
 }
 
-void IEditorDataStorageProvider::RemoveIndex(UE::Editor::DataStorage::IndexHash Index)
+void ICoreProvider::RemoveIndex(IndexHash Index)
 {
-	RemoveRowMapping(UE::Editor::DataStorage::FMapKeyView(Index));
+	RemoveRowMapping(FMapKeyView(Index));
 }
 
-template<UE::Editor::DataStorage::TDynamicColumnTemplate DynamicColumnTemplateType>
-void IEditorDataStorageProvider::ForEachDynamicColumn(TFunctionRef<void(const UScriptStruct& Type)> Callback) const
+template<TDynamicColumnTemplate DynamicColumnTemplateType>
+void ICoreProvider::ForEachDynamicColumn(TFunctionRef<void(const UScriptStruct& Type)> Callback) const
 {
 	ForEachDynamicColumn(DynamicColumnTemplateType::StaticStruct(), Callback);
-}
\ No newline at end of file
+}
+
+} // namespace UE::Editor::DataStorage
diff --git a/Engine/Source/Runtime/TypedElementFramework/Public/Elements/Interfaces/TypedElementDataStorageUiInterface.h b/Engine/Source/Runtime/TypedElementFramework/Public/Elements/Interfaces/TypedElementDataStorageUiInterface.h
--- a/Engine/Source/Runtime/TypedElementFramework/Public/Elements/Interfaces/TypedElementDataStorageUiInterface.h
+++ b/Engine/Source/Runtime/TypedElementFramework/Public/Elements/Interfaces/TypedElementDataStorageUiInterface.h
@@ -17,9 +17,9 @@
 namespace UE::Editor::DataStorage
 {
 	class ITedsWidget;
+	class IUiProvider;
 }
 
-class IEditorDataStorageUiProvider;
 class SWidget;
 
 /**
@@ -60,7 +60,7 @@ public:
 	 */
 	UE_DEPRECATED(5.6, "CreateWidgetDisplayName is deprecated, please use CreateWidgetDisplayNameText instead, which will return FText rather than an FString.")
 	TYPEDELEMENTFRAMEWORK_API virtual FString CreateWidgetDisplayName(
-		IEditorDataStorageProvider* DataStorage, UE::Editor::DataStorage::RowHandle Row) const final;
+		UE::Editor::DataStorage::ICoreProvider* DataStorage, UE::Editor::DataStorage::RowHandle Row) const final;
 
 	/** 
 	 * Returns an FText of the friendly name for the data the created widget represents.
@@ -72,7 +72,7 @@ public:
 	 * name for the whole widget and not a specific instance.
 	 */
 	TYPEDELEMENTFRAMEWORK_API virtual FText CreateWidgetDisplayNameText(
-		IEditorDataStorageProvider* DataStorage, UE::Editor::DataStorage::RowHandle Row = UE::Editor::DataStorage::InvalidRowHandle) const;
+		UE::Editor::DataStorage::ICoreProvider* DataStorage, UE::Editor::DataStorage::RowHandle Row = UE::Editor::DataStorage::InvalidRowHandle) const;
 	
 	/**
 	 *	Calls Construct() to create the internal widget, and then stores it in a container before returning.
@@ -84,8 +84,8 @@ public:
 	 */
 	TYPEDELEMENTFRAMEWORK_API virtual TSharedPtr<SWidget> ConstructFinalWidget(
 		UE::Editor::DataStorage::RowHandle Row, /** The row the widget will be stored in. */
-		IEditorDataStorageProvider* DataStorage,
-		IEditorDataStorageUiProvider* DataStorageUi,
+		UE::Editor::DataStorage::ICoreProvider* DataStorage,
+		UE::Editor::DataStorage::IUiProvider* DataStorageUi,
 		const UE::Editor::DataStorage::FMetaDataView& Arguments);
 
 	/**
@@ -97,8 +97,8 @@ public:
 	 */
 	TYPEDELEMENTFRAMEWORK_API virtual TSharedPtr<SWidget> Construct(
 		UE::Editor::DataStorage::RowHandle Row, /** The row the widget will be stored in. */
-		IEditorDataStorageProvider* DataStorage,
-		IEditorDataStorageUiProvider* DataStorageUi,
+		UE::Editor::DataStorage::ICoreProvider* DataStorage,
+		UE::Editor::DataStorage::IUiProvider* DataStorageUi,
 		const UE::Editor::DataStorage::FMetaDataView& Arguments);
 
 protected:
@@ -106,13 +106,13 @@ protected:
 	TYPEDELEMENTFRAMEWORK_API virtual TSharedPtr<SWidget> CreateWidget(const UE::Editor::DataStorage::FMetaDataView& Arguments);
 	/** Create a new instance of the target widget. This is a required function. */
 	TYPEDELEMENTFRAMEWORK_API virtual TSharedPtr<SWidget> CreateWidget(
-		IEditorDataStorageProvider* DataStorage,
-		IEditorDataStorageUiProvider* DataStorageUi,
+		UE::Editor::DataStorage::ICoreProvider* DataStorage,
+		UE::Editor::DataStorage::IUiProvider* DataStorageUi,
 		UE::Editor::DataStorage::RowHandle TargetRow,
 		UE::Editor::DataStorage::RowHandle WidgetRow, 
 		const UE::Editor::DataStorage::FMetaDataView& Arguments);
 	/** Set any values in columns if needed. The columns provided through GetAdditionalColumnsList() will have already been created. */
-	TYPEDELEMENTFRAMEWORK_API virtual bool SetColumns(IEditorDataStorageProvider* DataStorage, UE::Editor::DataStorage::RowHandle Row);
+	TYPEDELEMENTFRAMEWORK_API virtual bool SetColumns(UE::Editor::DataStorage::ICoreProvider* DataStorage, UE::Editor::DataStorage::RowHandle Row);
 	
 	/** Creates a (friendly) name for the provided column type. */
 	TYPEDELEMENTFRAMEWORK_API virtual FText DescribeColumnType(const UScriptStruct* ColumnType) const;
@@ -122,20 +122,20 @@ protected:
 	 * in columns.
 	 */
 	TYPEDELEMENTFRAMEWORK_API virtual bool FinalizeWidget(
-		IEditorDataStorageProvider* DataStorage,
-		IEditorDataStorageUiProvider* DataStorageUi,
+		UE::Editor::DataStorage::ICoreProvider* DataStorage,
+		UE::Editor::DataStorage::IUiProvider* DataStorageUi,
 		UE::Editor::DataStorage::RowHandle Row,
 		const TSharedPtr<SWidget>& Widget);
 
 	/** Add the default misc columns we want a widget row to have. */
-	TYPEDELEMENTFRAMEWORK_API void AddDefaultWidgetColumns(UE::Editor::DataStorage::RowHandle Row, IEditorDataStorageProvider* DataStorage) const;
+	TYPEDELEMENTFRAMEWORK_API void AddDefaultWidgetColumns(UE::Editor::DataStorage::RowHandle Row, UE::Editor::DataStorage::ICoreProvider* DataStorage) const;
 
 	/**
 	 * Helper function to get the actual target row with the data the widget is operating on (if applicable). Returns InvalidRowHandle if there is no
 	 * target row
 	 */
 	TYPEDELEMENTFRAMEWORK_API UE::Editor::DataStorage::RowHandle GetTargetRow(
-		IEditorDataStorageProvider* DataStorage, UE::Editor::DataStorage::RowHandle WidgetRow) const;
+		UE::Editor::DataStorage::ICoreProvider* DataStorage, UE::Editor::DataStorage::RowHandle WidgetRow) const;
 
 protected:
 
@@ -172,8 +172,8 @@ struct FSimpleWidgetConstructor : public FTypedElementWidgetConstructor
 	 * @return The actual widget instance
 	 */
 	TYPEDELEMENTFRAMEWORK_API virtual TSharedPtr<SWidget> CreateWidget(
-		IEditorDataStorageProvider* DataStorage,
-		IEditorDataStorageUiProvider* DataStorageUi,
+		UE::Editor::DataStorage::ICoreProvider* DataStorage,
+		UE::Editor::DataStorage::IUiProvider* DataStorageUi,
 		UE::Editor::DataStorage::RowHandle TargetRow,
 		UE::Editor::DataStorage::RowHandle WidgetRow,
 		const UE::Editor::DataStorage::FMetaDataView& Arguments) override;
@@ -185,7 +185,7 @@ struct FSimpleWidgetConstructor : public FTypedElementWidgetConstructor
 	 * @param WidgetRow The row that contains information about the widget itself
 	 * @return Whether any columns were added
 	 */
-	TYPEDELEMENTFRAMEWORK_API virtual bool SetColumns(IEditorDataStorageProvider* DataStorage, UE::Editor::DataStorage::RowHandle WidgetRow) override;
+	TYPEDELEMENTFRAMEWORK_API virtual bool SetColumns(UE::Editor::DataStorage::ICoreProvider* DataStorage, UE::Editor::DataStorage::RowHandle WidgetRow) override;
 
 protected:
 
@@ -196,8 +196,8 @@ protected:
 	 * Old function in the widget creation pipeline that isn't used anymore. All your logic should go in CreateWidget() itself
 	 */
 	TYPEDELEMENTFRAMEWORK_API virtual bool FinalizeWidget(
-		IEditorDataStorageProvider* DataStorage,
-		IEditorDataStorageUiProvider* DataStorageUi,
+		UE::Editor::DataStorage::ICoreProvider* DataStorage,
+		UE::Editor::DataStorage::IUiProvider* DataStorageUi,
 		UE::Editor::DataStorage::RowHandle Row,
 		const TSharedPtr<SWidget>& Widget) override final;
 	
@@ -206,8 +206,8 @@ protected:
 	 */
 	TYPEDELEMENTFRAMEWORK_API virtual TSharedPtr<SWidget> Construct(
 		UE::Editor::DataStorage::RowHandle WidgetRow,
-		IEditorDataStorageProvider* DataStorage,
-		IEditorDataStorageUiProvider* DataStorageUi,
+		UE::Editor::DataStorage::ICoreProvider* DataStorage,
+		UE::Editor::DataStorage::IUiProvider* DataStorageUi,
 		const UE::Editor::DataStorage::FMetaDataView& Arguments) override final;
 };
 
@@ -221,7 +221,9 @@ struct TStructOpsTypeTraits<FTypedElementWidgetConstructor> : public TStructOpsT
 	};
 };
 
-class IEditorDataStorageUiProvider : public IModularFeature
+namespace UE::Editor::DataStorage
+{
+class IUiProvider : public IModularFeature
 {
 public:
 	enum class EPurposeType : uint8
@@ -253,7 +255,7 @@ public:
 		SingleMatch
 	};
 
-	using WidgetCreatedCallback = TFunctionRef<void(const TSharedRef<SWidget>& NewWidget, UE::Editor::DataStorage::RowHandle Row)>;
+	using WidgetCreatedCallback = TFunctionRef<void(const TSharedRef<SWidget>& NewWidget, RowHandle Row)>;
 	using WidgetConstructorCallback = TFunctionRef<bool(TUniquePtr<FTypedElementWidgetConstructor>, TConstArrayView<TWeakObjectPtr<const UScriptStruct>>)>;
 	using WidgetPurposeCallback = TFunctionRef<void(FName, EPurposeType, const FText&)>;
 
@@ -293,7 +295,7 @@ public:
 	 * If registration is successful true will be returned otherwise false.
 	 */
 	virtual bool RegisterWidgetFactory(FName Purpose, const UScriptStruct* Constructor,
-		UE::Editor::DataStorage::Queries::FConditions Columns) = 0;
+		Queries::FConditions Columns) = 0;
 	/**
 	 * Registers a widget factory that will be called when the purpose it's registered under is requested.
 	 * This version registers a generic type. Construction using these are typically cheaper as they can avoid
@@ -303,7 +305,7 @@ public:
 	 * If registration is successful true will be returned otherwise false.
 	 */
 	template<typename ConstructorType>
-	bool RegisterWidgetFactory(FName Purpose, UE::Editor::DataStorage::Queries::FConditions Columns);
+	bool RegisterWidgetFactory(FName Purpose, Queries::FConditions Columns);
 	/**
 	 * Registers a widget factory that will be called when the purpose it's registered under is requested.
 	 * This version uses a previously created instance of the Constructor. The benefit of this is that it store
@@ -321,7 +323,7 @@ public:
 	 * If registration is successful true will be returned otherwise false.
 	 */
 	virtual bool RegisterWidgetFactory(FName Purpose, TUniquePtr<FTypedElementWidgetConstructor>&& Constructor, 
-		UE::Editor::DataStorage::Queries::FConditions Columns) = 0;
+		Queries::FConditions Columns) = 0;
 	
 	/** 
 	 * Creates widget constructors for the requested purpose.
@@ -329,7 +331,7 @@ public:
 	 * widgets created from the constructor, if applicable.
 	 */
 	virtual void CreateWidgetConstructors(FName Purpose, 
-		const UE::Editor::DataStorage::FMetaDataView& Arguments, const WidgetConstructorCallback& Callback) = 0;
+		const FMetaDataView& Arguments, const WidgetConstructorCallback& Callback) = 0;
 	/** 
 	 * Finds matching widget constructors for provided columns, preferring longer matches over shorter matches.
 	 * The provided list of columns will be updated to contain all columns that couldn't be matched.
@@ -337,14 +339,14 @@ public:
 	 * widgets created from the constructor, if applicable.
 	 */
 	virtual void CreateWidgetConstructors(FName Purpose, EMatchApproach MatchApproach, TArray<TWeakObjectPtr<const UScriptStruct>>& Columns,
-		const UE::Editor::DataStorage::FMetaDataView& Arguments, const WidgetConstructorCallback& Callback) = 0;
+		const FMetaDataView& Arguments, const WidgetConstructorCallback& Callback) = 0;
 
 	/**
 	 * Creates all the widgets registered under the provided name. This may be a large number of widgets for a wide name
 	 * or exactly one when the exact name of the widget is registered. Arguments can be provided, but widgets are free
 	 * to ignore them.
 	 */
-	virtual void ConstructWidgets(FName Purpose, const UE::Editor::DataStorage::FMetaDataView& Arguments,
+	virtual void ConstructWidgets(FName Purpose, const FMetaDataView& Arguments,
 		const WidgetCreatedCallback& ConstructionCallback) = 0;
 
 	/** 
@@ -352,8 +354,8 @@ public:
 	 * The provided row will be used to store the widget information on. If columns have already been added to the row, the 
 	 * constructor is free to use that to configure the widget. Arguments are used by the constructor to configure the widget.
 	 */
-	virtual TSharedPtr<SWidget> ConstructWidget(UE::Editor::DataStorage::RowHandle Row, FTypedElementWidgetConstructor& Constructor,
-		const UE::Editor::DataStorage::FMetaDataView& Arguments) = 0;
+	virtual TSharedPtr<SWidget> ConstructWidget(RowHandle Row, FTypedElementWidgetConstructor& Constructor,
+		const FMetaDataView& Arguments) = 0;
 
 	/** Calls the provided callback for all known registered widget purposes. */
 	virtual void ListWidgetPurposes(const WidgetPurposeCallback& Callback) const = 0;
@@ -364,10 +366,10 @@ public:
 	virtual void ListExtensions(TFunctionRef<void(FName)> Callback) const = 0;
 
 	/** Create the container widget that every TEDS UI widget is stored in */
-	virtual TSharedPtr<UE::Editor::DataStorage::ITedsWidget> CreateContainerTedsWidget(UE::Editor::DataStorage::RowHandle UiRowHandle) const = 0;
+	virtual TSharedPtr<ITedsWidget> CreateContainerTedsWidget(RowHandle UiRowHandle) const = 0;
 
 	/** Get the table where TEDS UI widgets are stored */
-	virtual UE::Editor::DataStorage::TableHandle GetWidgetTable() const = 0;
+	virtual TableHandle GetWidgetTable() const = 0;
 };
 
 
@@ -376,13 +378,14 @@ public:
 //
 
 template<typename ConstructorType>
-bool IEditorDataStorageUiProvider::RegisterWidgetFactory(FName Purpose)
+bool IUiProvider::RegisterWidgetFactory(FName Purpose)
 {
 	return this->RegisterWidgetFactory(Purpose, ConstructorType::StaticStruct());
 }
 
 template<typename ConstructorType>
-bool IEditorDataStorageUiProvider::RegisterWidgetFactory(FName Purpose, UE::Editor::DataStorage::Queries::FConditions Columns)
+bool IUiProvider::RegisterWidgetFactory(FName Purpose, UE::Editor::DataStorage::Queries::FConditions Columns)
 {
 	return this->RegisterWidgetFactory(Purpose, ConstructorType::StaticStruct(), Columns);
 }
+} // namespace UE::Editor::DataStorage
diff --git a/Engine/Source/Runtime/TypedElementFramework/Tests/Elements/Framework/TypedElementDataStoragePerformanceTestCommands.cpp b/Engine/Source/Runtime/TypedElementFramework/Tests/Elements/Framework/TypedElementDataStoragePerformanceTestCommands.cpp
--- a/Engine/Source/Runtime/TypedElementFramework/Tests/Elements/Framework/TypedElementDataStoragePerformanceTestCommands.cpp
+++ b/Engine/Source/Runtime/TypedElementFramework/Tests/Elements/Framework/TypedElementDataStoragePerformanceTestCommands.cpp
@@ -24,7 +24,7 @@ namespace UE::Editor::DataStorage
 				int32 EntitiesToAdd;
 				LexFromString(EntitiesToAdd, *Args[0]);
 
-				IEditorDataStorageProvider* DataStorageInterface = GetMutableDataStorageFeature<IEditorDataStorageProvider>(StorageFeatureName);
+				ICoreProvider* DataStorageInterface = GetMutableDataStorageFeature<ICoreProvider>(StorageFeatureName);
 			
 				DataStorageInterface->BatchAddRow(Private::PerformanceTestCommandTable, EntitiesToAdd, [DataStorageInterface](RowHandle Row)
 					{
@@ -42,14 +42,14 @@ namespace UE::Editor::DataStorage
 		{
 			using namespace UE::Editor::DataStorage::Queries;
 
-			IEditorDataStorageProvider* DataStorageInterface = GetMutableDataStorageFeature<IEditorDataStorageProvider>(StorageFeatureName);
+			ICoreProvider* DataStorageInterface = GetMutableDataStorageFeature<ICoreProvider>(StorageFeatureName);
 		
 			QueryHandle Query = DataStorageInterface->RegisterQuery(
 				Select().Where().All<FTest_PingPongPrePhys>().Compile());
 
 			TArray<RowHandle> RowsToDelete;
 
-			DataStorageInterface->RunQuery(Query, CreateDirectQueryCallbackBinding([&RowsToDelete](const IEditorDataStorageProvider::IDirectQueryContext& Context, const RowHandle* RowArray)
+			DataStorageInterface->RunQuery(Query, CreateDirectQueryCallbackBinding([&RowsToDelete](const IDirectQueryContext& Context, const RowHandle* RowArray)
 			{
 				RowsToDelete.Insert(RowArray, Context.GetRowCount(), RowsToDelete.Num());
 			}));
@@ -64,7 +64,7 @@ namespace UE::Editor::DataStorage
 	));
 } // namespace UE::Editor::DataStorage
 
-void UTest_PingPongBetweenPhaseFactory::RegisterTables(IEditorDataStorageProvider& DataStorage)
+void UTest_PingPongBetweenPhaseFactory::RegisterTables(UE::Editor::DataStorage::ICoreProvider& DataStorage)
 {
 	UE::Editor::DataStorage::Private::PerformanceTestCommandTable = DataStorage.RegisterTable({FTest_PingPongPrePhys::StaticStruct()}, TEXT("Test_PingPongPrePhys"));
 }
@@ -74,7 +74,7 @@ void UTest_PingPongBetweenPhaseFactory::RegisterTables(IEditorDataStorageProvide
 // A is sensitive to tag PingPongPrePhys.  It adds a PingPongDurPhys tag and removes the PingPonPrePhys tag.  This causes the processed row
 // to then be processed by processor B.  Processor B does a similar thing to ensure it is processed by C.  C then ensures it is processed by A
 // the next time A is run.
-void UTest_PingPongBetweenPhaseFactory::RegisterQueries(IEditorDataStorageProvider& DataStorage)
+void UTest_PingPongBetweenPhaseFactory::RegisterQueries(UE::Editor::DataStorage::ICoreProvider& DataStorage)
 {
 	using namespace UE::Editor::DataStorage::Queries;
 
diff --git a/Engine/Source/Runtime/TypedElementFramework/Tests/Elements/Framework/TypedElementDataStoragePerformanceTestCommands.h b/Engine/Source/Runtime/TypedElementFramework/Tests/Elements/Framework/TypedElementDataStoragePerformanceTestCommands.h
--- a/Engine/Source/Runtime/TypedElementFramework/Tests/Elements/Framework/TypedElementDataStoragePerformanceTestCommands.h
+++ b/Engine/Source/Runtime/TypedElementFramework/Tests/Elements/Framework/TypedElementDataStoragePerformanceTestCommands.h
@@ -53,6 +53,6 @@ class UTest_PingPongBetweenPhaseFactory : public UEditorDataStorageFactory
 public:
 	~UTest_PingPongBetweenPhaseFactory() override = default;
 
-	void RegisterTables(IEditorDataStorageProvider& DataStorage) override;
-	void RegisterQueries(IEditorDataStorageProvider& DataStorage) override;
+	void RegisterTables(UE::Editor::DataStorage::ICoreProvider& DataStorage) override;
+	void RegisterQueries(UE::Editor::DataStorage::ICoreProvider& DataStorage) override;
 };
\ No newline at end of file
diff --git a/Engine/Source/Runtime/TypedElementFramework/Tests/Elements/Framework/TypedElementDataStorageProcessorTests.cpp b/Engine/Source/Runtime/TypedElementFramework/Tests/Elements/Framework/TypedElementDataStorageProcessorTests.cpp
--- a/Engine/Source/Runtime/TypedElementFramework/Tests/Elements/Framework/TypedElementDataStorageProcessorTests.cpp
+++ b/Engine/Source/Runtime/TypedElementFramework/Tests/Elements/Framework/TypedElementDataStorageProcessorTests.cpp
@@ -30,7 +30,7 @@ namespace UE::Editor::DataStorage::Debug::ProcessorTests
 	{
 		using namespace UE::Editor::DataStorage::Queries;
 
-		IEditorDataStorageProvider* DataStorage = GetMutableDataStorageFeature<IEditorDataStorageProvider>(StorageFeatureName);
+		ICoreProvider* DataStorage = GetMutableDataStorageFeature<ICoreProvider>(StorageFeatureName);
 
 		if (PrimaryTable == InvalidTableHandle)
 		{
@@ -124,7 +124,7 @@ namespace UE::Editor::DataStorage::Debug::ProcessorTests
 
 	void UnregisterProcessors()
 	{
-		IEditorDataStorageProvider* DataStorage = GetMutableDataStorageFeature<IEditorDataStorageProvider>(StorageFeatureName);
+		ICoreProvider* DataStorage = GetMutableDataStorageFeature<ICoreProvider>(StorageFeatureName);
 		for (const UE::Editor::DataStorage::QueryHandle& Handle : RegisteredQueries)
 		{
 			DataStorage->UnregisterQuery(Handle);
@@ -147,7 +147,7 @@ namespace UE::Editor::DataStorage::Debug::ProcessorTests
 				return;
 			}
 
-			IEditorDataStorageProvider* DataStorage = GetMutableDataStorageFeature<IEditorDataStorageProvider>(StorageFeatureName);
+			ICoreProvider* DataStorage = GetMutableDataStorageFeature<ICoreProvider>(StorageFeatureName);
 			DataStorage->BatchAddRow(PrimaryTable, RowsToCreate, [](UE::Editor::DataStorage::RowHandle Row)
 			{
 			});
