TEDS UI: Fix incorrect TEXT() macro inside UE_DEPRECATED

diff --git a/Engine/Source/Runtime/TypedElementFramework/Public/Elements/Interfaces/TypedElementDataStorageUiInterface.h b/Engine/Source/Runtime/TypedElementFramework/Public/Elements/Interfaces/TypedElementDataStorageUiInterface.h
--- a/Engine/Source/Runtime/TypedElementFramework/Public/Elements/Interfaces/TypedElementDataStorageUiInterface.h
+++ b/Engine/Source/Runtime/TypedElementFramework/Public/Elements/Interfaces/TypedElementDataStorageUiInterface.h
@@ -93,12 +93,12 @@ public:
 		const UE::Editor::DataStorage::FMetaDataView& Arguments);
 	
 	/** Initializes a new constructor based on the provided arguments.. */
-	UE_DEPRECATED(5.6, TEXT("Use Initialize that takes in the factory row handle instead"))
+	UE_DEPRECATED(5.6, "Use Initialize that takes in the factory row handle instead")
 	TYPEDELEMENTFRAMEWORK_API virtual bool Initialize(const UE::Editor::DataStorage::FMetaDataView& InArguments,
 		TArray<TWeakObjectPtr<const UScriptStruct>> InMatchedColumnTypes, const UE::Editor::DataStorage::Queries::FConditions& InQueryConditions);
 
 	/** Retrieves the query conditions that need to match for this widget constructor to produce a widget. */
-	UE_DEPRECATED(5.6, TEXT("Use GetQueryConditions that takes in ICoreProvider* instead"))
+	UE_DEPRECATED(5.6, "Use GetQueryConditions that takes in ICoreProvider* instead")
 	TYPEDELEMENTFRAMEWORK_API virtual const UE::Editor::DataStorage::Queries::FConditions* GetQueryConditions() const final;
 	
 	/**
