[TEDS] Extended the widget constructor class with a way to provide sorting information.

This is part of a larger effort to introduce generic sorting through TEDS to the editor.

diff --git a/Engine/Source/Runtime/TypedElementFramework/Private/Elements/Interfaces/TypedElementDataStorageUiInterface.cpp b/Engine/Source/Runtime/TypedElementFramework/Private/Elements/Interfaces/TypedElementDataStorageUiInterface.cpp
--- a/Engine/Source/Runtime/TypedElementFramework/Private/Elements/Interfaces/TypedElementDataStorageUiInterface.cpp
+++ b/Engine/Source/Runtime/TypedElementFramework/Private/Elements/Interfaces/TypedElementDataStorageUiInterface.cpp
@@ -6,13 +6,18 @@
 #include "Elements/Columns/TypedElementMiscColumns.h"
 #include "Elements/Columns/TypedElementSlateWidgetColumns.h"
 #include "Elements/Columns/TypedElementUIColumns.h"
-#include "Elements/Framework/TypedElementDataStorageWidget.h"
 #include "Elements/Columns/WidgetPurposeColumns.h"
+#include "Elements/Framework/TypedElementDataStorageWidget.h"
+#include "Math/NumericLimits.h"
 #include "Elements/Common/EditorDataStorageFeatures.h"
 #include "Widgets/SNullWidget.h"
 
 #define LOCTEXT_NAMESPACE "FTypedElementWidgetConstructor"
 
+//
+// FTypedElementWidgetConstructor
+//
+
 FTypedElementWidgetConstructor::FTypedElementWidgetConstructor(const UScriptStruct* InTypeInfo)
 	: TypeInfo(InTypeInfo)
 {
@@ -212,6 +217,13 @@ TSharedPtr<SWidget> FTypedElementWidgetConstructor::Construct(
 	return nullptr;
 }
 
+TArray<TSharedPtr<const UE::Editor::DataStorage::FColumnSorterInterface>> FTypedElementWidgetConstructor::ConstructColumnSorters(
+	UE::Editor::DataStorage::ICoreProvider* DataStorage,
+	const UE::Editor::DataStorage::FMetaDataView& Arguments)
+{
+	return TArray<TSharedPtr<const UE::Editor::DataStorage::FColumnSorterInterface>>();
+}
+
 TSharedPtr<SWidget> FTypedElementWidgetConstructor::CreateWidget(const UE::Editor::DataStorage::FMetaDataView& Arguments)
 {
 	return nullptr;
@@ -289,6 +301,7 @@ UE::Editor::DataStorage::RowHandle FTypedElementWidgetConstructor::GetTargetRow(
 }
 
 
+
 // FSimpleWidgetConstructor
 
 FSimpleWidgetConstructor::FSimpleWidgetConstructor(const UScriptStruct* InTypeInfo)
diff --git a/Engine/Source/Runtime/TypedElementFramework/Public/Elements/Framework/TypedElementSorter.h b/Engine/Source/Runtime/TypedElementFramework/Public/Elements/Framework/TypedElementSorter.h
new file mode 100644
--- /dev/null
+++ b/Engine/Source/Runtime/TypedElementFramework/Public/Elements/Framework/TypedElementSorter.h
@@ -0,0 +1,301 @@
+// Copyright Epic Games, Inc. All Rights Reserved.
+
+#pragma once
+
+#include <type_traits>
+#include "Containers/StringFwd.h"
+#include "Elements/Common/TypedElementHandles.h"
+#include "Internationalization/Text.h"
+#include "Math/NumericLimits.h"
+
+namespace UE::Editor::DataStorage
+{
+	/**
+	 * @section Sort provider
+	 *
+	 * @description
+	 * Interface and supporting classes to provide sorting of rows by column.
+	 */
+	class ICoreProvider;
+
+	struct FSortPrefixResult
+	{
+		uint64 Prefix = TNumericLimits<uint64>::Max();
+		bool bHasRemainingBytes = false;
+	};
+
+	/** Interface to provide sorting of rows by column. */
+	class FColumnSorterInterface
+	{
+	public:
+		virtual ~FColumnSorterInterface() = default;
+
+		enum class ESortType : uint8
+		{
+			/**
+			 * Supports sort algorithms that sort by only using a single (packed) value of 64 bits.
+			 * Algorithms supporting this option will only use `SortPrefix` and will only call this once per row.
+			 */
+			FixedSize64,
+			/**
+			 * Supports sort algorithms that sort by only using one or more fixed sized values like integers or floats.
+			 * Algorithms supporting this option will only use `SortPrefix` and assume only fixed sized values are used. Including
+			 * values with variable size such as strings will lead to suboptimal performance.
+			 */
+			FixedSizeOnly,
+			/** 
+			 * Supports sort algorithms that sort using a comparative function.
+			 * Algorithms supporting this option will only use the `Compare` function to sort. Comparative sorting requires more reads
+			 * from columns than fixed sized sorting or hybrid sorting an is also more susceptible to frame spikes when distributed over
+			 * multiple frames and/or threads. It does however provide greater flexibility when requiring complex comparisons such as ones
+			 * requiring multiple strings.
+			 */
+			ComparativeSort,
+			/**
+			 * Supports sort algorithms that sort using both a prefix and a comparative function.
+			 * Algorithms supporting this option will use both the `SortPrefix` and `Compare` functions to sort. Typically rows are grouped
+			 * into buckets based on their prefix and when a bucket is small enough a compare may be used to sort the bucket.
+			 * A hybrid approach is typically faster for larger numbers of rows as is reduces the amount of times columns have to be read.
+			 * It also reduces or removes the need for a compare which can be beneficial if the compare is expensive. Lastly it typically
+			 * also allows for more even work distribution over frames/threads, reducing the risk of spikes. However it does come at an
+			 * increased peak memory usage and limits the number of variable sized variables to one, which also have to be the last 
+			 * variable in the chain.
+			 */
+			HybridSort
+		};
+
+		/** Returns the type of sorting required for this column. */
+		virtual ESortType GetSortType() const = 0;
+
+		/** If set, this can be used. */
+		virtual FText GetShortName() const = 0;
+
+		/**
+		 * Compare the content of the left row to the right row and return a negative number if left is smaller than right, zero if left
+		 * and right are equal and a positive number if right is larger than left.
+		 */
+		virtual int32 Compare(const ICoreProvider& Storage, RowHandle Left, RowHandle Right) const = 0;
+
+		/**
+		 * The numeric prefix for the column. For numeric values this typically the number itself. For strings it's typically the next 8
+		 * characters or 4 wide characters starting at the provided byte index. The utility function can `CreateSortPrefix` be used to 
+		 * help create a prefix from one or more variables.
+		 * Note that using a prefix is limited to only one value that has a variable length (e.g. strings) and has to be the last value
+		 * to be sorted.
+		 */
+		virtual FSortPrefixResult SortPrefix(const ICoreProvider& Storage, RowHandle Row, uint32 ByteIndex) const = 0;
+	};
+
+	namespace Private
+	{
+		template<typename T>
+		concept SortNumericValue = std::is_integral_v<T> || std::is_floating_point_v<T>;
+
+		template <typename T>
+		struct TSortStringViewType : std::false_type {};
+
+		template <typename T>
+		struct TSortStringViewType<TStringView<T>> : std::true_type {};
+
+		template<typename T>
+		concept SortStringViewValue = TSortStringViewType<T>::value;
+
+		template<typename T>
+		concept SortStringViewCopyableType =
+			SortStringViewValue<T> ||
+			std::is_same_v<T, FString> ||
+			std::is_same_v<T, FAnsiString> ||
+			std::is_same_v<T, FWideString> ||
+			std::is_same_v<T, FUtf8String>;
+
+		template<typename T>
+		concept SortStringVariantType = SortStringViewCopyableType<T> || std::is_same_v<T, FText>;
+
+		template<typename T>
+		concept NameValue = std::is_same_v<T, FName>;
+	}
+
+	/**
+	 * @section Sort support structures
+	 *
+	 * @description
+	 * Includes a variety of structures that wrap types to add additional information needed for sorting. An example is a view
+	 * to wrap strings to indicate whether the sort is case sensitive or not.
+	 */
+
+	struct FSortCaseSensitive {};
+	struct FSortCaseInsensitive {};
+	template<typename T>
+	concept SortCase = std::is_same_v<T, FSortCaseSensitive> || std::is_same_v<T, FSortCaseInsensitive>;
+	
+	template<SortCase Casing, Private::SortStringViewValue StringView>
+	struct TSortStringView
+	{
+		using StringViewType = StringView;
+		static constexpr bool bIsCaseSensitive = std::is_same_v<Casing, FSortCaseSensitive>;
+		static constexpr ESearchCase::Type SearchCase = bIsCaseSensitive ? ESearchCase::CaseSensitive : ESearchCase::IgnoreCase;
+		
+		TSortStringView() = default;
+		
+		template<Private::SortStringViewCopyableType ViewType>
+		TSortStringView(const ViewType& InView) : View(InView) {}
+		template<Private::SortStringViewCopyableType ViewType>
+		TSortStringView& operator=(const ViewType& InView) { this->View = InView; return *this; }
+
+		TSortStringView(const FText& Text) : View(Text.ToString()) {}
+		TSortStringView& operator=(const FText& Text) { this->View = Text.ToString(); return *this; }
+		
+		template<SortCase TargetCasing, Private::SortStringViewCopyableType ViewType>
+		TSortStringView(TargetCasing, const ViewType& String) : View(String) { static_assert(std::is_same_v<Casing, TargetCasing>); }
+		template<SortCase TargetCasing>
+		TSortStringView(TargetCasing, const FText& Text) : View(Text.ToString()) { static_assert(std::is_same_v<Casing, TargetCasing>); }
+
+		bool operator==(const TSortStringView& Rhs) const { return View.Equals(Rhs.View, SearchCase) == true; }
+		bool operator!=(const TSortStringView& Rhs) const { return View.Equals(Rhs.View, SearchCase) == false; }
+		bool operator< (const TSortStringView& Rhs) const { return View.Compare(Rhs.View, SearchCase) < 0; }
+		bool operator<=(const TSortStringView& Rhs) const { return View.Compare(Rhs.View, SearchCase) <= 0; }
+		bool operator> (const TSortStringView& Rhs) const { return View.Compare(Rhs.View, SearchCase) > 0; }
+		bool operator>=(const TSortStringView& Rhs) const { return View.Compare(Rhs.View, SearchCase) >= 0; }
+
+		StringView View;
+	};
+
+	template<SortCase Casing, Private::SortStringViewCopyableType View>
+	TSortStringView(Casing, const View&) -> TSortStringView<Casing, TStringView<typename View::ElementType>>;
+	template<SortCase Casing>
+	TSortStringView(Casing, const FText&) -> TSortStringView<Casing, FStringView>;
+
+	enum class ENameSortBy : bool { Id, String };
+
+	struct FSortById {};
+	struct FSortByName {};
+	template<typename T>
+	concept SortBy = std::is_same_v<T, FSortById> || std::is_same_v<T, FSortByName>;
+
+	template<SortBy By>
+	struct TSortNameView
+	{
+		constexpr static bool bIsById = std::is_same_v<By, FSortById>;
+		using CompareType = std::conditional_t<bIsById, int32, FString>;
+		constexpr static bool bIsFixedSize = bIsById;
+
+		TSortNameView() = default;
+		TSortNameView(const FName& Name) : View(&Name) {}
+		TSortNameView& operator=(const FName& Name);
+
+		template<SortBy TargetBy>
+		TSortNameView(TargetBy, const FName& Name) : View(&Name) { static_assert(std::is_same_v<By, TargetBy>); }
+
+		uint32 GetByteSize() const;
+		constexpr static uint32 GetElementSize();
+		FSortPrefixResult CalculatePrefix(int32 CurrentIndex, int32 ByteIndex) const;
+
+		int32 Compare(const TSortNameView& Rhs) const;
+		bool operator==(const TSortNameView& Rhs) const;
+		bool operator!=(const TSortNameView& Rhs) const;
+		bool operator< (const TSortNameView& Rhs) const;
+		bool operator<=(const TSortNameView& Rhs) const;
+		bool operator> (const TSortNameView& Rhs) const;
+		bool operator>=(const TSortNameView& Rhs) const;
+
+		const FName* View = nullptr;
+		mutable CompareType Cache = {};
+		mutable bool bIsCached = false;
+
+	private:
+		void CacheCompareType() const;
+	};
+
+	template<SortBy By>
+	TSortNameView(By, const FName&) -> TSortNameView<By>;
+
+	template<typename... ValueTypes>
+	FSortPrefixResult CreateSortPrefix(uint32 ByteIndex, ValueTypes&&... Values);
+
+	/**
+	 * @section Sort type information
+	 *
+	 * @description
+	 * Utility structs and function that describe the way types can be sorted and packed into an index.
+	 */
+
+	namespace Private
+	{
+		template<typename T>
+		constexpr uint64 MoveToLocation(int32 ByteIndex, T Value);
+
+		template<typename Numeric>
+		constexpr auto Rebase(Numeric Value);
+	}
+
+	template<typename ValueType>
+	struct TSortTypeInfo
+	{
+		// Used to short-circuit several calls to minimize noise in the error output of the compiler.
+		static constexpr bool bIsSupportedType = false;
+
+		// Check is just a random way to say "false" so the compiler doesn't evaluate the assert until an attempt is made to instance 
+		// this version.
+		static_assert(sizeof(ValueType) == 0, "Unsupported sort type.");
+	};
+
+	template<Private::SortNumericValue NumericType>
+	struct TSortTypeInfo<NumericType>
+	{
+		static constexpr bool bIsSupportedType = true;
+		static constexpr bool bIsFixedSize = true;
+		constexpr static uint32 GetByteSize(NumericType Value) { return sizeof(NumericType); }
+		constexpr static uint32 GetElementSize() { return sizeof(NumericType); }
+		constexpr static FSortPrefixResult CalculatePrefix(int32 CurrentIndex, int32 ByteIndex, NumericType Value)
+		{
+			// Due to an issue with MSVC resolving forward declared partially specialized templates that use concepts, this 
+			// can't be moved to the .inl file.
+			return FSortPrefixResult
+			{
+				.Prefix = Private::MoveToLocation(CurrentIndex - ByteIndex, Private::Rebase(Value)),
+				.bHasRemainingBytes = false
+			};
+		}
+	};
+
+	template<SortCase Casing, typename T>
+	struct TSortTypeInfo<TSortStringView<Casing, T>>
+	{
+		static constexpr bool bIsSupportedType = true;
+		static constexpr bool bIsFixedSize = false;
+		constexpr static uint32 GetByteSize(TSortStringView<Casing, T> Value);
+		constexpr static uint32 GetElementSize();
+		constexpr static FSortPrefixResult CalculatePrefix(
+			int32 CurrentIndex, int32 ByteIndex, TSortStringView<Casing, T> Value);
+	};
+
+	template<typename StringType> requires Private::SortStringVariantType<StringType>
+	struct TSortTypeInfo<StringType>
+	{
+		static constexpr bool bIsSupportedType = false;
+		static_assert(sizeof(StringType) == 0, 
+			"Strings and string views are not directly supported. Use `TSortStringView` to indicate if sorting is case sensitive or not.");
+	};
+
+	template<SortBy By>
+	struct TSortTypeInfo<TSortNameView<By>>
+	{
+		static constexpr bool bIsSupportedType = true;
+		static constexpr bool bIsFixedSize = TSortNameView<By>::bIsFixedSize;
+		static uint32 GetByteSize(const TSortNameView<By>& Value) { return Value.GetByteSize(); };
+		constexpr static uint32 GetElementSize() { return TSortNameView<By>::GetElementSize(); };
+		static FSortPrefixResult CalculatePrefix(
+			int32 CurrentIndex, int32 ByteIndex, TSortNameView<By> Value) { return Value.CalculatePrefix(CurrentIndex, ByteIndex); };
+	};
+
+	// Not a direct specialization, but done through requirements to prevent the compiler from evaluating the static_assert when it's not instanced.
+	template<Private::NameValue NameType>
+	struct TSortTypeInfo<NameType>
+	{
+		static constexpr bool bIsSupportedType = false;
+		static_assert(sizeof(NameType) == 0,
+			"FNames are not directly supported. Use `TSortNameView` to indicate if sorting is based on a string or the unique FName number.");
+	};
+} // namespace UE::Editor::DataStorage
+
+#include "Elements/Framework/TypedElementSorter.inl"
diff --git a/Engine/Source/Runtime/TypedElementFramework/Public/Elements/Framework/TypedElementSorter.inl b/Engine/Source/Runtime/TypedElementFramework/Public/Elements/Framework/TypedElementSorter.inl
new file mode 100644
--- /dev/null
+++ b/Engine/Source/Runtime/TypedElementFramework/Public/Elements/Framework/TypedElementSorter.inl
@@ -0,0 +1,347 @@
+// Copyright Epic Games, Inc. All Rights Reserved.
+
+#include "Misc/Char.h"
+
+namespace UE::Editor::DataStorage
+{
+	namespace Private
+	{
+		template<typename T>
+		constexpr uint64 MoveToLocation(int32 ByteIndex, T Value)
+		{
+			constexpr int32 ByteSize = sizeof(T);
+			const uint64 Result = static_cast<uint64>(Value);
+			const int32 ByteShift = 8 - ByteIndex - ByteSize;
+			const int32 BitShift = ByteShift * 8;
+			return (BitShift >= 0) ? (Result << BitShift) : (Result >> -BitShift);
+		}
+
+		template<typename Numeric>
+		constexpr auto Rebase(Numeric Value)
+		{
+			constexpr uint64 BitShift = (sizeof(Numeric) * 8) - 1;
+
+			if constexpr (std::is_floating_point_v<Numeric>)
+			{
+				using SignedInt = std::conditional_t<sizeof(Numeric) == sizeof(double), int64, int32>;
+				using UnsignedInt = std::make_unsigned_t<SignedInt>;
+
+				union { Numeric F; UnsignedInt U; } Converter;
+				Converter.F = Value;
+
+				SignedInt mask = -static_cast<SignedInt>(Converter.U >> BitShift) | (UnsignedInt(1) << BitShift);
+				return Converter.U ^ mask;
+			}
+			else if constexpr (std::is_signed_v<Numeric>)
+			{
+				using Unsigned = std::make_unsigned_t<Numeric>;
+				
+				Numeric mask = static_cast<Numeric>(-(Value >> BitShift) | (Unsigned(1) << BitShift));
+				return static_cast<Unsigned>(Value ^ mask);
+			}
+			else
+			{
+				static_assert(std::is_unsigned_v<Numeric>, "Input value to rebase for sort indexing must be a numeric value.");
+				return Value;
+			}
+		}
+
+		template<bool bCaseSensitive, typename CharType>
+		CharType ToUpper(CharType Input)
+		{
+			if constexpr (bCaseSensitive)
+			{
+				return Input;
+			}
+			else
+			{
+				return TChar<CharType>::ToUpper(Input);
+			}
+		}
+
+		template<typename ValueType>
+		void CalculatePrefix(FSortPrefixResult& Result, int32 CurrentIndex, int32 ByteIndex, const ValueType& Value)
+		{
+			if constexpr (TSortTypeInfo<ValueType>::bIsSupportedType)
+			{
+				constexpr int32 ResultSize = sizeof(FSortPrefixResult::Prefix);
+				int32 Size = static_cast<int32>(TSortTypeInfo<ValueType>::GetByteSize(Value));
+
+				if constexpr (TSortTypeInfo<ValueType>::bIsFixedSize)
+				{
+					int32 Position = CurrentIndex - ByteIndex;
+					if (Position < ResultSize)
+					{
+						Result.Prefix |= TSortTypeInfo<ValueType>::CalculatePrefix(CurrentIndex, ByteIndex, Value).Prefix;
+						Result.bHasRemainingBytes = (Position + Size) > sizeof(FSortPrefixResult::Prefix);
+					}
+				}
+				else
+				{
+					constexpr uint32 ElementSize = TSortTypeInfo<ValueType>::GetElementSize();
+
+					int32 Index = CurrentIndex - ByteIndex;
+					if (ResultSize - (Index & (ResultSize - 1)) >= static_cast<int32>(ElementSize))
+					{
+						FSortPrefixResult Intermediate =
+							TSortTypeInfo<ValueType>::CalculatePrefix(CurrentIndex, ByteIndex, Value);
+						Result.Prefix |= Intermediate.Prefix;
+						Result.bHasRemainingBytes = Intermediate.bHasRemainingBytes;
+					}
+				}
+			}
+		}
+
+		template<typename ValueType, typename... ValueTypes>
+		void CalculatePrefix(FSortPrefixResult& Result, int32 CurrentIndex, int32 ByteIndex,
+			const ValueType& Value, ValueTypes&&... Values)
+		{
+			if constexpr (TSortTypeInfo<ValueType>::bIsSupportedType)
+			{
+				constexpr int32 ResultSize = sizeof(FSortPrefixResult::Prefix);
+				int32 Size = static_cast<int32>(TSortTypeInfo<ValueType>::GetByteSize(Value));
+				// Skip over bytes that fall out of the requested range.
+				if (CurrentIndex + Size <= ByteIndex)
+				{
+					CalculatePrefix(Result, CurrentIndex + Size, ByteIndex, Values...);
+				}
+				else
+				{
+					static_assert(TSortTypeInfo<ValueType>::bIsFixedSize, "Only the last value type can be of variable size.");
+
+					// Try to pack as many bytes into the final result as are available. This might require pulling portions of values.
+					int32 Position = CurrentIndex - ByteIndex;
+					if (Position < ResultSize)
+					{
+						// Continue while there are still bytes available.
+						Result.Prefix |= TSortTypeInfo<ValueType>::CalculatePrefix(CurrentIndex, ByteIndex, Value).Prefix;
+						CalculatePrefix(Result, CurrentIndex + Size, ByteIndex, Values...);
+					}
+				}
+			}
+		}
+	}
+
+	
+	
+	//
+	// TSortNameView
+	//
+
+	template<SortBy By>
+	TSortNameView<By>& TSortNameView<By>::operator=(const FName& Name)
+	{ 
+		View = &Name; 
+		bIsCached = false; 
+		return *this;
+	}
+
+	template<SortBy By>
+	uint32 TSortNameView<By>::GetByteSize() const
+	{
+		if constexpr (bIsById)
+		{
+			return sizeof(Cache);
+		}
+		else
+		{
+			CacheCompareType();
+			return static_cast<uint32>(Cache.NumBytesWithoutNull());
+		}
+	}
+
+	template<SortBy By>
+	constexpr uint32 TSortNameView<By>::GetElementSize()
+	{
+		if constexpr (bIsById)
+		{
+			return sizeof(Cache);
+		}
+		else
+		{
+			return sizeof(typename CompareType::ElementType);
+		}
+	}
+
+	template<SortBy By>
+	FSortPrefixResult TSortNameView<By>::CalculatePrefix(int32 CurrentIndex, int32 ByteIndex) const
+	{
+		CacheCompareType();
+		if constexpr (bIsById)
+		{
+			return TSortTypeInfo<CompareType>::CalculatePrefix(CurrentIndex, ByteIndex, Cache);
+		}
+		else
+		{
+			using StringView = TStringView<typename CompareType::ElementType>;
+			using SortStringView = TSortStringView<FSortCaseInsensitive, StringView>;
+			return TSortTypeInfo<SortStringView>::CalculatePrefix(CurrentIndex, ByteIndex, StringView(Cache));
+		}
+	}
+
+	template<SortBy By>
+	int32 TSortNameView<By>::Compare(const TSortNameView& Rhs) const
+	{
+		if (View && Rhs.View)
+		{
+			if constexpr (bIsById)
+			{
+				return View->CompareIndexes(*Rhs.View);
+			}
+			else
+			{
+				return View->Compare(*Rhs.View);
+			}
+		}
+		else
+		{
+			int32 ViewIsNull = View == nullptr ? 1 : -1;
+			return View == Rhs.View ? 0 : ViewIsNull;
+		}
+	}
+
+	template<SortBy By>
+	bool TSortNameView<By>::operator==(const TSortNameView& Rhs) const
+	{ 
+		return Compare(Rhs) == 0;
+	}
+
+	template<SortBy By>
+	bool TSortNameView<By>::operator!=(const TSortNameView& Rhs) const
+	{ 
+		return Compare(Rhs) != 0;
+	}
+
+	template<SortBy By>
+	bool TSortNameView<By>::operator< (const TSortNameView& Rhs) const
+	{ 
+		return Compare(Rhs) < 0;
+	}
+	
+	template<SortBy By>
+	bool TSortNameView<By>::operator<=(const TSortNameView& Rhs) const
+	{ 
+		return Compare(Rhs) <= 0;
+	}
+
+	template<SortBy By>
+	bool TSortNameView<By>::operator> (const TSortNameView& Rhs) const
+	{ 
+		return Compare(Rhs) > 0;
+	}
+
+	template<SortBy By>
+	bool TSortNameView<By>::operator>=(const TSortNameView& Rhs) const
+	{ 
+		return Compare(Rhs) >= 0;
+	}
+
+	template<SortBy By>
+	void TSortNameView<By>::CacheCompareType() const
+	{
+		if (View && !bIsCached)
+		{
+			if constexpr (bIsById)
+			{
+				Cache = View->GetNumber();
+			}
+			else
+			{
+				Cache = View->ToString();
+			}
+		}
+	}
+
+
+
+	// TSortTypeInfo<TSortStringView<...>>
+
+	template<SortCase Casing, typename T>
+	constexpr uint32 TSortTypeInfo<TSortStringView<Casing, T>>::GetByteSize(TSortStringView<Casing, T> Value)
+	{
+		return static_cast<uint32>(Value.View.NumBytes());
+	}
+
+	template<SortCase Casing, typename T>
+	constexpr uint32 TSortTypeInfo<TSortStringView<Casing, T>>::GetElementSize()
+	{
+		return sizeof(typename T::ElementType);
+	}
+
+	template<SortCase Casing, typename T>
+	constexpr FSortPrefixResult TSortTypeInfo<TSortStringView<Casing, T>>::CalculatePrefix(
+		int32 CurrentIndex, int32 ByteIndex, TSortStringView<Casing, T> Value)
+	{
+		constexpr bool bIsCaseSensitive = TSortStringView<Casing, T>::bIsCaseSensitive;
+
+		uint64 Result = 0;
+		bool bHasRemainingCharacters = false;
+
+		uint32 BytesIntoResult = ByteIndex < CurrentIndex ? (CurrentIndex & 7) : 0;
+		uint32 RemainingBytes = 8 - BytesIntoResult;
+		int32 StringIndex = FMath::Clamp<int32>(ByteIndex - CurrentIndex, 0, static_cast<int32>(Value.View.NumBytes()));
+		
+		if constexpr (sizeof(typename T::ElementType) == 1)
+		{
+			int32 Min = FMath::Min<int32>(Value.View.Len() - StringIndex, RemainingBytes);
+
+			switch (Min)
+			{
+			case 8:
+				Result |= (static_cast<uint64>(Private::ToUpper<bIsCaseSensitive>(Value.View[StringIndex + Min - 8])) << 56);
+			case 7:
+				Result |= (static_cast<uint64>(Private::ToUpper<bIsCaseSensitive>(Value.View[StringIndex + Min - 7])) << 48);
+			case 6:
+				Result |= (static_cast<uint64>(Private::ToUpper<bIsCaseSensitive>(Value.View[StringIndex + Min - 6])) << 40);
+			case 5:
+				Result |= (static_cast<uint64>(Private::ToUpper<bIsCaseSensitive>(Value.View[StringIndex + Min - 5])) << 32);
+			case 4:
+				Result |= (static_cast<uint64>(Private::ToUpper<bIsCaseSensitive>(Value.View[StringIndex + Min - 4])) << 24);
+			case 3:
+				Result |= (static_cast<uint64>(Private::ToUpper<bIsCaseSensitive>(Value.View[StringIndex + Min - 3])) << 16);
+			case 2:
+				Result |= (static_cast<uint64>(Private::ToUpper<bIsCaseSensitive>(Value.View[StringIndex + Min - 2])) << 8);
+			case 1:
+				Result |= (static_cast<uint64>(Private::ToUpper<bIsCaseSensitive>(Value.View[StringIndex + Min - 1])));
+			}
+			Result <<= 64 - (BytesIntoResult + Min) * 8;
+			bHasRemainingCharacters = StringIndex + Min < Value.View.Len();
+		}
+		else
+		{
+			StringIndex >>= 1;
+			int32 Min = FMath::Min<int32>(Value.View.Len() - StringIndex, RemainingBytes >> 1);
+			switch (Min)
+			{
+			case 4:
+				Result |= (static_cast<uint64>(Private::ToUpper<bIsCaseSensitive>(Value.View[StringIndex + Min - 4])) << 48);
+			case 3:
+				Result |= (static_cast<uint64>(Private::ToUpper<bIsCaseSensitive>(Value.View[StringIndex + Min - 3])) << 32);
+			case 2:
+				Result |= (static_cast<uint64>(Private::ToUpper<bIsCaseSensitive>(Value.View[StringIndex + Min - 2])) << 16);
+			case 1:
+				Result |= (static_cast<uint64>(Private::ToUpper<bIsCaseSensitive>(Value.View[StringIndex + Min - 1])));
+			}
+			Result <<= 64 - (BytesIntoResult + (Min *2)) * 8;
+			bHasRemainingCharacters = StringIndex + Min < Value.View.Len();
+		}
+
+		return FSortPrefixResult
+		{
+			.Prefix = Result,
+			.bHasRemainingBytes = bHasRemainingCharacters
+		};
+	}
+
+	template<typename... ValueTypes>
+	FSortPrefixResult CreateSortPrefix(uint32 ByteIndex, ValueTypes&&... Values)
+	{
+		FSortPrefixResult Result
+		{
+			.Prefix = 0,
+			.bHasRemainingBytes = true
+		};
+		Private::CalculatePrefix(Result, 0, static_cast<int32>(ByteIndex), Values...);
+		return Result;
+	}
+} // namespace UE::Editor::DataStorage
diff --git a/Engine/Source/Runtime/TypedElementFramework/Public/Elements/Interfaces/TypedElementDataStorageUiInterface.h b/Engine/Source/Runtime/TypedElementFramework/Public/Elements/Interfaces/TypedElementDataStorageUiInterface.h
--- a/Engine/Source/Runtime/TypedElementFramework/Public/Elements/Interfaces/TypedElementDataStorageUiInterface.h
+++ b/Engine/Source/Runtime/TypedElementFramework/Public/Elements/Interfaces/TypedElementDataStorageUiInterface.h
@@ -3,10 +3,12 @@
 #pragma once
 
 #include "Containers/ContainersFwd.h"
+#include "Containers/StringFwd.h"
 #include "Elements/Common/TypedElementQueryConditions.h"
 #include "Elements/Interfaces/TypedElementDataStorageInterface.h"
 #include "Elements/Framework/TypedElementMetaData.h"
 #include "Elements/Framework/TypedElementQueryBuilder.h"
+#include "Elements/Framework/TypedElementSorter.h"
 #include "Features/IModularFeature.h"
 #include "Templates/UnrealTypeTraits.h"
 #include "Templates/SharedPointer.h"
@@ -92,6 +94,16 @@ public:
 		UE::Editor::DataStorage::IUiProvider* DataStorageUi,
 		const UE::Editor::DataStorage::FMetaDataView& Arguments);
 	
+	/**
+	 * Constructs a list of classes that are used to sort the rows with the column(s) associated with this widget constructor. 
+	 * If this widget constructor is used to create multiple widgets, for instance when used in a column of a grid view, these sorters can
+	 * be used to order the rows in a logical order described by the sorter. A widget constructor can provide multiple constructors which will
+	 * allow the same list of rows to be sorted in different orders.
+	 */
+	TYPEDELEMENTFRAMEWORK_API virtual TArray<TSharedPtr<const UE::Editor::DataStorage::FColumnSorterInterface>> ConstructColumnSorters(
+		UE::Editor::DataStorage::ICoreProvider* DataStorage,
+		const UE::Editor::DataStorage::FMetaDataView& Arguments);
+
 	/** Initializes a new constructor based on the provided arguments.. */
 	UE_DEPRECATED(5.6, "Use Initialize that takes in the factory row handle instead")
 	TYPEDELEMENTFRAMEWORK_API virtual bool Initialize(const UE::Editor::DataStorage::FMetaDataView& InArguments,
@@ -279,7 +291,7 @@ namespace UE::Editor::DataStorage
 			SingleMatch
 		};
 
-		using FPurposeID = UE::Editor::DataStorage::FMapKey;
+		using FPurposeID = FMapKey;
 
 		/*
 		 * Struct describing the default init params for a widget purpose. These can also be added to the purpose row post init using the columns
@@ -330,7 +342,7 @@ namespace UE::Editor::DataStorage
 		 *
 		 * @param InPurposeInfo Struct with init params for the purpose
 		 */
-		virtual UE::Editor::DataStorage::RowHandle RegisterWidgetPurpose(const FPurposeInfo& InPurposeInfo) = 0;
+		virtual RowHandle RegisterWidgetPurpose(const FPurposeInfo& InPurposeInfo) = 0;
 	
 		/**
 		 * Overload for RegisterWidgetPurpose that allows you to use a custom ID to map to the TEDS Mapping table for lookup
@@ -338,7 +350,7 @@ namespace UE::Editor::DataStorage
 		 * @param PurposeID A unique ID for the purpose that can be used to look up the row handle later
 		 * @param InPurposeInfo Struct with init params for the purpose
 		 */
-		virtual UE::Editor::DataStorage::RowHandle RegisterWidgetPurpose(const FPurposeID& PurposeID, const FPurposeInfo& InPurposeInfo) = 0;
+		virtual RowHandle RegisterWidgetPurpose(const FPurposeID& PurposeID, const FPurposeInfo& InPurposeInfo) = 0;
 	
 	
 		/** 
@@ -348,7 +360,7 @@ namespace UE::Editor::DataStorage
 		 * options. If the purpose has not been registered the factory will not be recorded and a warning will be printed.
 		 * If registration is successful true will be returned otherwise false.
 		 */
-		virtual bool RegisterWidgetFactory(UE::Editor::DataStorage::RowHandle PurposeRow, const UScriptStruct* Constructor) = 0;
+		virtual bool RegisterWidgetFactory(RowHandle PurposeRow, const UScriptStruct* Constructor) = 0;
 
 		/**
 		 * Registers a widget factory that will be called when the purpose it's registered under is requested.
@@ -358,7 +370,7 @@ namespace UE::Editor::DataStorage
 		 * If registration is successful true will be returned otherwise false.
 		 */
 		template<typename ConstructorType>
-		bool RegisterWidgetFactory(UE::Editor::DataStorage::RowHandle PurposeRow);
+		bool RegisterWidgetFactory(RowHandle PurposeRow);
 
 		/**
 		 * Registers a widget factory that will be called when the purpose it's registered under is requested.
@@ -368,8 +380,8 @@ namespace UE::Editor::DataStorage
 		 * The provided columns will be used when matching the factory during widget construction.
 		 * If registration is successful true will be returned otherwise false.
 		 */
-		virtual bool RegisterWidgetFactory(UE::Editor::DataStorage::RowHandle PurposeRow, const UScriptStruct* Constructor,
-			UE::Editor::DataStorage::Queries::FConditions Columns) = 0;
+		virtual bool RegisterWidgetFactory(RowHandle PurposeRow, const UScriptStruct* Constructor,
+			Queries::FConditions Columns) = 0;
 	
 		/**
 		 * Registers a widget factory that will be called when the purpose it's registered under is requested.
@@ -380,7 +392,7 @@ namespace UE::Editor::DataStorage
 		 * If registration is successful true will be returned otherwise false.
 		 */
 		template<typename ConstructorType>
-		bool RegisterWidgetFactory(UE::Editor::DataStorage::RowHandle PurposeRow, UE::Editor::DataStorage::Queries::FConditions Columns);
+		bool RegisterWidgetFactory(RowHandle PurposeRow, Queries::FConditions Columns);
 
 		/**
 		 * Registers a widget factory that will be called when the purpose it's registered under is requested.
@@ -389,7 +401,7 @@ namespace UE::Editor::DataStorage
 		 * used. If the purpose has not been registered the factory will not be recorded and a warning will be printed.
 		 * If registration is successful true will be returned otherwise false.
 		 */
-		virtual bool RegisterWidgetFactory(UE::Editor::DataStorage::RowHandle PurposeRow, TUniquePtr<FTypedElementWidgetConstructor>&& Constructor) = 0;
+		virtual bool RegisterWidgetFactory(RowHandle PurposeRow, TUniquePtr<FTypedElementWidgetConstructor>&& Constructor) = 0;
 	
 		/**
 		 * Registers a widget factory that will be called when the purpose it's registered under is requested.
@@ -399,8 +411,8 @@ namespace UE::Editor::DataStorage
 		 * The provided columns will be used when matching the factory during widget construction.
 		 * If registration is successful true will be returned otherwise false.
 		 */
-		virtual bool RegisterWidgetFactory(UE::Editor::DataStorage::RowHandle PurposeRow, TUniquePtr<FTypedElementWidgetConstructor>&& Constructor, 
-			UE::Editor::DataStorage::Queries::FConditions Columns) = 0;
+		virtual bool RegisterWidgetFactory(RowHandle PurposeRow, TUniquePtr<FTypedElementWidgetConstructor>&& Constructor, 
+			Queries::FConditions Columns) = 0;
 	
 		/** 
 		 * Creates widget constructors for the requested purpose.
@@ -409,8 +421,8 @@ namespace UE::Editor::DataStorage
 		 * 
 		 * If no factories were found for the requested purpose, the purpose's parent chain will be traversed to try and find an ancestor purpose with factories
 		 */
-		virtual void CreateWidgetConstructors(UE::Editor::DataStorage::RowHandle PurposeRow, 
-			const UE::Editor::DataStorage::FMetaDataView& Arguments, const WidgetConstructorCallback& Callback) = 0;
+		virtual void CreateWidgetConstructors(RowHandle PurposeRow, 
+			const FMetaDataView& Arguments, const WidgetConstructorCallback& Callback) = 0;
 	
 		/** 
 		 * Finds matching widget constructors for provided columns, preferring longer matches over shorter matches.
@@ -421,8 +433,8 @@ namespace UE::Editor::DataStorage
 		 * If there are still columns remaining after matching against the requested purpose, the purpose's parent chain will be traversed to try and match
 		 * the remaining columns
 		 */
-		virtual void CreateWidgetConstructors(UE::Editor::DataStorage::RowHandle PurposeRow, EMatchApproach MatchApproach, TArray<TWeakObjectPtr<const UScriptStruct>>& Columns,
-			const UE::Editor::DataStorage::FMetaDataView& Arguments, const WidgetConstructorCallback& Callback) = 0;
+		virtual void CreateWidgetConstructors(RowHandle PurposeRow, EMatchApproach MatchApproach, TArray<TWeakObjectPtr<const UScriptStruct>>& Columns,
+			const FMetaDataView& Arguments, const WidgetConstructorCallback& Callback) = 0;
 	
 		/**
 		 * Creates all the widgets registered under the provided purpose. This may be a large number of widgets for a wide name
@@ -431,7 +443,7 @@ namespace UE::Editor::DataStorage
 		 * 
 		 * If no factories were found for the requested purpose, the purpose's parent chain will be traversed to try and find an ancestor purpose with factories
 		 */
-		virtual void ConstructWidgets(UE::Editor::DataStorage::RowHandle PurposeRow, const UE::Editor::DataStorage::FMetaDataView& Arguments,
+		virtual void ConstructWidgets(RowHandle PurposeRow, const FMetaDataView& Arguments,
 			const WidgetCreatedCallback& ConstructionCallback) = 0;
 	
 		/** 
@@ -439,8 +451,8 @@ namespace UE::Editor::DataStorage
 		 * The provided row will be used to store the widget information on. If columns have already been added to the row, the 
 		 * constructor is free to use that to configure the widget. Arguments are used by the constructor to configure the widget.
 		 */
-		virtual TSharedPtr<SWidget> ConstructWidget(UE::Editor::DataStorage::RowHandle Row, FTypedElementWidgetConstructor& Constructor,
-			const UE::Editor::DataStorage::FMetaDataView& Arguments) = 0;
+		virtual TSharedPtr<SWidget> ConstructWidget(RowHandle Row, FTypedElementWidgetConstructor& Constructor,
+			const FMetaDataView& Arguments) = 0;
 
 		/** Calls the provided callback for all known registered widget purposes. */
 		virtual void ListWidgetPurposes(const WidgetPurposeCallback& Callback) const = 0;
@@ -451,13 +463,13 @@ namespace UE::Editor::DataStorage
 		virtual void ListExtensions(TFunctionRef<void(FName)> Callback) const = 0;
 
 		/** Create the container widget that every TEDS UI widget is stored in */
-		virtual TSharedPtr<UE::Editor::DataStorage::ITedsWidget> CreateContainerTedsWidget(UE::Editor::DataStorage::RowHandle UiRowHandle) const = 0;
+		virtual TSharedPtr<ITedsWidget> CreateContainerTedsWidget(RowHandle UiRowHandle) const = 0;
 
 		/** Get the table where TEDS UI widgets are stored */
-		virtual UE::Editor::DataStorage::TableHandle GetWidgetTable() const = 0;
+		virtual TableHandle GetWidgetTable() const = 0;
 
 		/** Find the row handle for a purpose by looking it up using the purpose ID */
-		virtual UE::Editor::DataStorage::RowHandle FindPurpose(const FPurposeID& PurposeID) const = 0;
+		virtual RowHandle FindPurpose(const FPurposeID& PurposeID) const = 0;
 	
 		/** Get the name of the default TEDS UI widget purpose used to register default widgets for different types of data (e.g FText -> STextBlock) */
 		virtual FPurposeID GetDefaultWidgetPurposeID() const = 0;
@@ -476,34 +488,33 @@ namespace UE::Editor::DataStorage
 		bool RegisterWidgetFactory(FName Purpose);
 
 		UE_DEPRECATED(5.6, "Use the version that takes in PurposeRowHandle instead. FindPurpose() to lookup row handle by PurposeID")
-		virtual bool RegisterWidgetFactory(FName Purpose, const UScriptStruct* Constructor,
-			UE::Editor::DataStorage::Queries::FConditions Columns) = 0;
+		virtual bool RegisterWidgetFactory(FName Purpose, const UScriptStruct* Constructor, Queries::FConditions Columns) = 0;
 
 		template<typename ConstructorType>
 		UE_DEPRECATED(5.6, "Use the version that takes in PurposeRowHandle instead. FindPurpose() to lookup row handle by PurposeID")
-		bool RegisterWidgetFactory(FName Purpose, UE::Editor::DataStorage::Queries::FConditions Columns);
+		bool RegisterWidgetFactory(FName Purpose, Queries::FConditions Columns);
 
 		UE_DEPRECATED(5.6, "Use the version that takes in PurposeRowHandle instead. FindPurpose() to lookup row handle by PurposeID")
 		virtual bool RegisterWidgetFactory(FName Purpose, TUniquePtr<FTypedElementWidgetConstructor>&& Constructor) = 0;
 
 		UE_DEPRECATED(5.6, "Use the version that takes in PurposeRowHandle instead. FindPurpose() to lookup row handle by PurposeID")
 		virtual bool RegisterWidgetFactory(FName Purpose, TUniquePtr<FTypedElementWidgetConstructor>&& Constructor, 
-			UE::Editor::DataStorage::Queries::FConditions Columns) = 0;
+			Queries::FConditions Columns) = 0;
 
 		UE_DEPRECATED(5.6, "Use the version that takes in PurposeRowHandle instead. FindPurpose() to lookup row handle by PurposeID")
-		virtual void CreateWidgetConstructors(FName Purpose, 
-			const UE::Editor::DataStorage::FMetaDataView& Arguments, const WidgetConstructorCallback& Callback) = 0;
+		virtual void CreateWidgetConstructors(FName Purpose, const FMetaDataView& Arguments, const WidgetConstructorCallback& Callback) = 0;
 
 		UE_DEPRECATED(5.6, "Use the version that takes in PurposeRowHandle instead. FindPurpose() to lookup row handle by PurposeID")
 		virtual void CreateWidgetConstructors(FName Purpose, EMatchApproach MatchApproach, TArray<TWeakObjectPtr<const UScriptStruct>>& Columns,
-			const UE::Editor::DataStorage::FMetaDataView& Arguments, const WidgetConstructorCallback& Callback) = 0;
+			const FMetaDataView& Arguments, const WidgetConstructorCallback& Callback) = 0;
 
 		UE_DEPRECATED(5.6, "Use the version that takes in PurposeRowHandle instead. FindPurpose() to lookup row handle by PurposeID")
-		virtual void ConstructWidgets(FName Purpose, const UE::Editor::DataStorage::FMetaDataView& Arguments,
-			const WidgetCreatedCallback& ConstructionCallback) = 0;
+		virtual void ConstructWidgets(FName Purpose, const FMetaDataView& Arguments, const WidgetCreatedCallback& ConstructionCallback) = 0;
 
 	};
 
+	using FWidgetConstructor = FTypedElementWidgetConstructor;
+	using FSimpleWidgetConstructor = FSimpleWidgetConstructor;
 
 	//
 	// Implementations
@@ -518,13 +529,13 @@ namespace UE::Editor::DataStorage
 	}
 
 	template<typename ConstructorType>
-	bool IUiProvider::RegisterWidgetFactory(UE::Editor::DataStorage::RowHandle PurposeRow)
+	bool IUiProvider::RegisterWidgetFactory(RowHandle PurposeRow)
 	{
 		return this->RegisterWidgetFactory(PurposeRow, ConstructorType::StaticStruct());
 	}
 
 	template<typename ConstructorType>
-	bool IUiProvider::RegisterWidgetFactory(FName Purpose, UE::Editor::DataStorage::Queries::FConditions Columns)
+	bool IUiProvider::RegisterWidgetFactory(FName Purpose, Queries::FConditions Columns)
 	{
 		PRAGMA_DISABLE_DEPRECATION_WARNINGS
 		return this->RegisterWidgetFactory(Purpose, ConstructorType::StaticStruct(), Columns);
@@ -532,7 +543,7 @@ namespace UE::Editor::DataStorage
 	}
 
 	template<typename ConstructorType>
-	bool IUiProvider::RegisterWidgetFactory(UE::Editor::DataStorage::RowHandle PurposeRow, UE::Editor::DataStorage::Queries::FConditions Columns)
+	bool IUiProvider::RegisterWidgetFactory(RowHandle PurposeRow, Queries::FConditions Columns)
 	{
 		return this->RegisterWidgetFactory(PurposeRow, ConstructorType::StaticStruct(), Columns);
 	}
diff --git a/Engine/Source/Runtime/TypedElementFramework/Tests/Elements/Framework/TypedElementDataStorageUiSorterTests.cpp b/Engine/Source/Runtime/TypedElementFramework/Tests/Elements/Framework/TypedElementDataStorageUiSorterTests.cpp
new file mode 100644
--- /dev/null
+++ b/Engine/Source/Runtime/TypedElementFramework/Tests/Elements/Framework/TypedElementDataStorageUiSorterTests.cpp
@@ -0,0 +1,906 @@
+// Copyright Epic Games, Inc. All Rights Reserved.
+
+#if WITH_TESTS
+
+#include "Algo/Sort.h"
+#include "Containers/AnsiString.h"
+#include "Elements/Common/EditorDataStorageFeatures.h"
+#include "Elements/Interfaces/TypedElementDataStorageUiInterface.h"
+#include "Math/NumericLimits.h"
+#include "Templates/Tuple.h"
+#include "Tests/TestHarnessAdapter.h"
+
+namespace UE::Editor::DataStorage::Tests
+{
+	template<typename T>
+	class FSingleValueTestSorter : public FColumnSorterInterface
+	{
+	public:
+		virtual ~FSingleValueTestSorter() override = default;
+
+		virtual int32 Compare(const ICoreProvider& Storage, RowHandle Left, RowHandle Right) const override
+		{
+			if (Values[static_cast<int32>(Left)] == Values[static_cast<int32>(Right)])
+			{
+				return 0;
+			}
+			else
+			{
+				return Values[static_cast<int32>(Left)] < Values[static_cast<int32>(Right)] ? -1 : 1;
+			}
+		}
+
+		virtual FSortPrefixResult SortPrefix(const ICoreProvider& Storage, RowHandle Row, uint32 ByteIndex) const override
+		{
+			return CreateSortPrefix(ByteIndex, Values[static_cast<int32>(Row)]);
+		}
+
+		virtual ESortType GetSortType() const override
+		{
+			return ESortType::HybridSort;
+		}
+
+		virtual FText GetShortName() const override
+		{
+			return FText::GetEmpty();
+		}
+
+		void SortByPrefix(const ICoreProvider& Storage, uint32 ByteIndex)
+		{
+			struct FData
+			{
+				uint64 Prefix;
+				T Value;
+			};
+			TArray<FData> Data;
+			Data.Reserve(Values.Num());
+
+			for (int32 Index = 0; Index < Values.Num(); ++Index)
+			{
+				Data.Add(FData
+					{
+						.Prefix = SortPrefix(Storage, Index, ByteIndex).Prefix,
+						.Value = MoveTemp(Values[Index])
+					});
+			}
+
+			RadixSort64(Data.GetData(), Data.Num(), [](const FData& Entry) { return Entry.Prefix; });
+
+			Values.Reset();
+			for (FData& Entry : Data)
+			{
+				Values.Add(MoveTemp(Entry.Value));
+			}
+		}
+
+		bool Verify(const ICoreProvider& Storage)
+		{
+			for (int32 Index = 0; Index < Values.Num() - 1; ++Index)
+			{
+				if (Compare(Storage, Index, Index + 1) > 0)
+				{
+					return false;
+				}
+			}
+			return true;
+		}
+
+		TArray<T> Values;
+	};
+
+	template<typename... Types>
+	class FMultiValueTestSorter : public FColumnSorterInterface
+	{
+	public:
+		virtual ~FMultiValueTestSorter() override = default;
+
+		virtual int32 Compare(const ICoreProvider& Storage, RowHandle Left, RowHandle Right) const override
+		{
+			if (Values[static_cast<int32>(Left)] == Values[static_cast<int32>(Right)])
+			{
+				return 0;
+			}
+			else
+			{
+				return Values[static_cast<int32>(Left)] < Values[static_cast<int32>(Right)] ? -1 : 1;
+			}
+		}
+
+		virtual FSortPrefixResult SortPrefix(const ICoreProvider& Storage, RowHandle Row, uint32 ByteIndex) const override
+		{
+			return Values[static_cast<int32>(Row)].ApplyAfter([this, ByteIndex](auto&&... TupleValues)
+				{
+					return CreateSortPrefix(ByteIndex, TupleValues...);
+				});
+		}
+
+		virtual ESortType GetSortType() const override
+		{
+			return ESortType::HybridSort;
+		}
+
+		virtual FText GetShortName() const override
+		{
+			return FText::GetEmpty();
+		}
+
+		void SortByPrefix(const ICoreProvider& Storage, uint32 ByteIndex, int32 RangeStart, int32 RangeEnd)
+		{
+			struct FData
+			{
+				uint64 Prefix;
+				TTuple<Types...> ValueTuple;
+			};
+			TArray<FData> Data;
+			Data.Reserve(RangeEnd - RangeStart + 1);
+
+			for (int32 Index = RangeStart; Index <= RangeEnd; ++Index)
+			{
+				uint64 Prefix = Values[Index].ApplyAfter([this, ByteIndex](auto&&... TupleValues)
+					{
+						return CreateSortPrefix(ByteIndex, TupleValues...);
+					}).Prefix;
+				Data.Add(FData
+					{
+						.Prefix = Prefix,
+						.ValueTuple = MoveTemp(Values[Index])
+					});
+			}
+
+			RadixSort64(Data.GetData(), Data.Num(), [](const FData& Entry) { return Entry.Prefix; });
+
+			for (int32 Index = 0; Index <= RangeEnd - RangeStart; ++Index)
+			{
+				Values[RangeStart + Index] = MoveTemp(Data[Index].ValueTuple);
+			}
+		}
+
+		bool Verify(const ICoreProvider& Storage)
+		{
+			for (int32 Index = 0; Index < Values.Num() - 1; ++Index)
+			{
+				if (Compare(Storage, Index, Index + 1) > 0)
+				{
+					return false;
+				}
+			}
+			return true;
+		}
+
+		TArray<TTuple<Types...>> Values;
+	};
+
+	TEST_CASE_NAMED(FEditorDataStorageUiSorter_Prefix_UInt64, 
+		"Editor::DataStorage::FColumnSorterInterface_Prefix_UInt64", "[ApplicationContextMask][EngineFilter]")
+	{
+		if (ICoreProvider* DataStorage = GetMutableDataStorageFeature<ICoreProvider>(StorageFeatureName))
+		{
+			FSingleValueTestSorter<uint64> Sorter;
+			Sorter.Values.Append({ 1, 42, 7, 33, 1024 });
+			Sorter.SortByPrefix(*DataStorage, 0);
+			CHECK(Sorter.Verify(*DataStorage));
+		}
+	}
+
+	TEST_CASE_NAMED(FEditorDataStorageUiSorter_Prefix_UInt32, 
+		"Editor::DataStorage::FColumnSorterInterface_Prefix_UInt32", "[ApplicationContextMask][EngineFilter]")
+	{
+		if (ICoreProvider* DataStorage = GetMutableDataStorageFeature<ICoreProvider>(StorageFeatureName))
+		{
+			FSingleValueTestSorter<uint32> Sorter;
+			Sorter.Values.Append({ 1, 42, 7, 33, 1024 });
+			Sorter.SortByPrefix(*DataStorage, 0);
+			CHECK(Sorter.Verify(*DataStorage));
+		}
+	}
+
+	TEST_CASE_NAMED(FEditorDataStorageUiSorter_Prefix_UInt16, 
+		"Editor::DataStorage::FColumnSorterInterface_Prefix_UInt16", "[ApplicationContextMask][EngineFilter]")
+	{
+		if (ICoreProvider* DataStorage = GetMutableDataStorageFeature<ICoreProvider>(StorageFeatureName))
+		{
+			FSingleValueTestSorter<uint16> Sorter;
+			Sorter.Values.Append({ 1, 42, 7, 33, 1024 });
+			Sorter.SortByPrefix(*DataStorage, 0);
+			CHECK(Sorter.Verify(*DataStorage));
+		}
+	}
+
+	TEST_CASE_NAMED(FEditorDataStorageUiSorter_Prefix_UInt8, 
+		"Editor::DataStorage::FColumnSorterInterface_Prefix_UInt8", "[ApplicationContextMask][EngineFilter]")
+	{
+		if (ICoreProvider* DataStorage = GetMutableDataStorageFeature<ICoreProvider>(StorageFeatureName))
+		{
+			FSingleValueTestSorter<uint8> Sorter;
+			Sorter.Values.Append({ 1, 42, 7, 33, 255 });
+			Sorter.SortByPrefix(*DataStorage, 0);
+			CHECK(Sorter.Verify(*DataStorage));
+		}
+	}
+
+	TEST_CASE_NAMED(FEditorDataStorageUiSorter_Prefix_Int64, 
+		"Editor::DataStorage::FColumnSorterInterface_Prefix_Int64", "[ApplicationContextMask][EngineFilter]")
+	{
+		if (ICoreProvider* DataStorage = GetMutableDataStorageFeature<ICoreProvider>(StorageFeatureName))
+		{
+			FSingleValueTestSorter<int64> Sorter;
+			Sorter.Values.Append({ 1, -1, -8, 42, TNumericLimits<int64>::Min(), TNumericLimits<int64>::Max(), 7, -53221, 33, 1024 });
+			Sorter.SortByPrefix(*DataStorage, 0);
+			CHECK(Sorter.Verify(*DataStorage));
+		}
+	}
+
+	TEST_CASE_NAMED(FEditorDataStorageUiSorter_Prefix_Int32, 
+		"Editor::DataStorage::FColumnSorterInterface_Prefix_Int32", "[ApplicationContextMask][EngineFilter]")
+	{
+		if (ICoreProvider* DataStorage = GetMutableDataStorageFeature<ICoreProvider>(StorageFeatureName))
+		{
+			FSingleValueTestSorter<int32> Sorter;
+			Sorter.Values.Append({ 1, -1, -8, 42, TNumericLimits<int32>::Min(), TNumericLimits<int32>::Max(), 7, -53221, 33, 1024 });
+			Sorter.SortByPrefix(*DataStorage, 0);
+			CHECK(Sorter.Verify(*DataStorage));
+		}
+	}
+
+	TEST_CASE_NAMED(FEditorDataStorageUiSorter_Prefix_Int16, 
+		"Editor::DataStorage::FColumnSorterInterface_Prefix_Int16", "[ApplicationContextMask][EngineFilter]")
+	{
+		if (ICoreProvider* DataStorage = GetMutableDataStorageFeature<ICoreProvider>(StorageFeatureName))
+		{
+			FSingleValueTestSorter<int16> Sorter;
+			Sorter.Values.Append({ 1, -1, -8, 42, TNumericLimits<int16>::Min(), TNumericLimits<int16>::Max(), 7, -3221, 33, 1024 });
+			Sorter.SortByPrefix(*DataStorage, 0);
+			CHECK(Sorter.Verify(*DataStorage));
+		}
+	}
+
+	TEST_CASE_NAMED(FEditorDataStorageUiSorter_Prefix_Int8, 
+		"Editor::DataStorage::FColumnSorterInterface_Prefix_Int8", "[ApplicationContextMask][EngineFilter]")
+	{
+		if (ICoreProvider* DataStorage = GetMutableDataStorageFeature<ICoreProvider>(StorageFeatureName))
+		{
+			FSingleValueTestSorter<int8> Sorter;
+			Sorter.Values.Append({ 1, -1, -8, 42, TNumericLimits<int8>::Min(), TNumericLimits<int8>::Max(), 7, -128, 33, 127 });
+			Sorter.SortByPrefix(*DataStorage, 0);
+			CHECK(Sorter.Verify(*DataStorage));
+		}
+	}
+
+	TEST_CASE_NAMED(FEditorDataStorageUiSorter_Prefix_Float, 
+		"Editor::DataStorage::FColumnSorterInterface_Prefix_Float", "[ApplicationContextMask][EngineFilter]")
+	{
+		if (ICoreProvider* DataStorage = GetMutableDataStorageFeature<ICoreProvider>(StorageFeatureName))
+		{
+			FSingleValueTestSorter<float> Sorter;
+			Sorter.Values.Append({ 1.44f, -1.32f, -8.8675463f, 42.145165f, 7.0f, -53221542.331f, 33.8763f, 1024.45290625f });
+			Sorter.SortByPrefix(*DataStorage, 0);
+			CHECK(Sorter.Verify(*DataStorage));
+		}
+	}
+
+	TEST_CASE_NAMED(FEditorDataStorageUiSorter_Prefix_Double, 
+		"Editor::DataStorage::FColumnSorterInterface_Prefix_Double", "[ApplicationContextMask][EngineFilter]")
+	{
+		if (ICoreProvider* DataStorage = GetMutableDataStorageFeature<ICoreProvider>(StorageFeatureName))
+		{
+			FSingleValueTestSorter<double> Sorter;
+			Sorter.Values.Append({ 1.44, -1.32, -8.8675463, 42.145165, 7.0, -53221542.331, 33.8763, 1024.45290625 });
+			Sorter.SortByPrefix(*DataStorage, 0);
+			CHECK(Sorter.Verify(*DataStorage));
+		}
+	}
+
+	TEST_CASE_NAMED(FEditorDataStorageUiSorter_Prefix_Boolean, 
+		"Editor::DataStorage::FColumnSorterInterface_Prefix_Boolean", "[ApplicationContextMask][EngineFilter]")
+	{
+		if (ICoreProvider* DataStorage = GetMutableDataStorageFeature<ICoreProvider>(StorageFeatureName))
+		{
+			FSingleValueTestSorter<bool> Sorter;
+			Sorter.Values.Append({ true, false, false, true, false, false, true, true });
+			Sorter.SortByPrefix(*DataStorage, 0);
+			CHECK(Sorter.Verify(*DataStorage));
+		}
+	}
+
+	TEST_CASE_NAMED(FEditorDataStorageUiSorter_Prefix_MultiVariableSizedValues,
+		"Editor::DataStorage::FColumnSorterInterface_Prefix_MultiVariableSizedValues", "[ApplicationContextMask][EngineFilter]")
+	{
+		if (ICoreProvider* DataStorage = GetMutableDataStorageFeature<ICoreProvider>(StorageFeatureName))
+		{
+			{
+				FMultiValueTestSorter<uint32> Sorter;
+				Sorter.Values.Add(
+					{
+						static_cast<uint32>(0x11223344)
+					});
+
+				FSortPrefixResult Result;
+				Result = Sorter.SortPrefix(*DataStorage, 0, 0);
+				CHECK(Result.Prefix == 0x1122334400000000);
+				CHECK(Result.bHasRemainingBytes == false);
+			}
+
+			{
+				FMultiValueTestSorter<uint32, uint8> Sorter;
+				Sorter.Values.Add(
+					{
+						static_cast<uint32>(0x11223344),
+						static_cast<uint8>(0x55)
+					});
+
+				FSortPrefixResult Result;
+				Result = Sorter.SortPrefix(*DataStorage, 0, 0);
+				CHECK(Result.Prefix == 0x1122334455000000);
+				CHECK(Result.bHasRemainingBytes == false);
+			}
+
+			{
+				FMultiValueTestSorter<uint32, uint8, uint16> Sorter;
+				Sorter.Values.Add(
+					{
+						static_cast<uint32>(0x11223344),
+						static_cast<uint8>(0x55),
+						static_cast<uint16>(0x6677)
+					});
+
+				FSortPrefixResult Result;
+				Result = Sorter.SortPrefix(*DataStorage, 0, 0);
+				CHECK(Result.Prefix == 0x1122334455667700);
+				CHECK(Result.bHasRemainingBytes == false);
+			}
+
+			{
+				FMultiValueTestSorter<uint32, uint8, uint16, uint32> Sorter;
+				Sorter.Values.Add(
+					{
+						static_cast<uint32>(0x11223344),
+						static_cast<uint8>(0x55),
+						static_cast<uint16>(0x6677),
+						static_cast<uint32>(0x8899aabb)
+					});
+
+				FSortPrefixResult Result;
+				Result = Sorter.SortPrefix(*DataStorage, 0, 0);
+				CHECK(Result.Prefix == 0x1122334455667788);
+				CHECK(Result.bHasRemainingBytes == true);
+
+				Result = Sorter.SortPrefix(*DataStorage, 0, 8);
+				CHECK(Result.Prefix == 0x99aabb0000000000);
+				CHECK(Result.bHasRemainingBytes == false);
+			}
+
+			{
+				FMultiValueTestSorter<uint32, uint8, uint16, uint32, uint64> Sorter;
+				Sorter.Values.Add(
+					{
+						static_cast<uint32>(0x11223344),
+						static_cast<uint8>(0x55),
+						static_cast<uint16>(0x6677),
+						static_cast<uint32>(0x8899aabb),
+						static_cast<uint64>(0xccddeeff11223344)
+					});
+
+				FSortPrefixResult Result;
+				Result = Sorter.SortPrefix(*DataStorage, 0, 0);
+				CHECK(Result.Prefix == 0x1122334455667788);
+				CHECK(Result.bHasRemainingBytes == true);
+
+				Result = Sorter.SortPrefix(*DataStorage, 0, 8);
+				CHECK(Result.Prefix == 0x99aabbccddeeff11);
+				CHECK(Result.bHasRemainingBytes == true);
+
+				Result = Sorter.SortPrefix(*DataStorage, 0, 16);
+				CHECK(Result.Prefix == 0x2233440000000000);
+				CHECK(Result.bHasRemainingBytes == false);
+			}
+
+			{
+				FMultiValueTestSorter<uint32, uint8, uint16, uint32, uint64, uint16> Sorter;
+				Sorter.Values.Add(
+					{
+						static_cast<uint32>(0x11223344),
+						static_cast<uint8>(0x55),
+						static_cast<uint16>(0x6677),
+						static_cast<uint32>(0x8899aabb),
+						static_cast<uint64>(0xccddeeff11223344),
+						static_cast<uint16>(0x5566)
+					});
+
+				FSortPrefixResult Result;
+				Result = Sorter.SortPrefix(*DataStorage, 0, 0);
+				CHECK(Result.Prefix == 0x1122334455667788);
+				CHECK(Result.bHasRemainingBytes == true);
+
+				Result = Sorter.SortPrefix(*DataStorage, 0, 8);
+				CHECK(Result.Prefix == 0x99aabbccddeeff11);
+				CHECK(Result.bHasRemainingBytes == true);
+
+				Result = Sorter.SortPrefix(*DataStorage, 0, 16);
+				CHECK(Result.Prefix == 0x2233445566000000);
+				CHECK(Result.bHasRemainingBytes == false);
+			}
+
+			{
+				FMultiValueTestSorter<uint32, uint8, uint16, uint32, uint64, uint16, uint8> Sorter;
+				Sorter.Values.Add(
+					{
+						static_cast<uint32>(0x11223344),
+						static_cast<uint8>(0x55),
+						static_cast<uint16>(0x6677),
+						static_cast<uint32>(0x8899aabb),
+						static_cast<uint64>(0xccddeeff11223344),
+						static_cast<uint16>(0x5566),
+						static_cast<uint8>(0x77)
+					});
+
+				FSortPrefixResult Result;
+				Result = Sorter.SortPrefix(*DataStorage, 0, 0);
+				CHECK(Result.Prefix == 0x1122334455667788);
+				CHECK(Result.bHasRemainingBytes == true);
+
+				Result = Sorter.SortPrefix(*DataStorage, 0, 8);
+				CHECK(Result.Prefix == 0x99aabbccddeeff11);
+				CHECK(Result.bHasRemainingBytes == true);
+
+				Result = Sorter.SortPrefix(*DataStorage, 0, 16);
+				CHECK(Result.Prefix == 0x2233445566770000);
+				CHECK(Result.bHasRemainingBytes == false);
+			}
+
+			{
+				FMultiValueTestSorter<uint32, uint8, uint16, uint32, uint64, uint16, uint8, uint16> Sorter;
+				Sorter.Values.Add(
+					{
+						static_cast<uint32>(0x11223344),
+						static_cast<uint8>(0x55),
+						static_cast<uint16>(0x6677),
+						static_cast<uint32>(0x8899aabb),
+						static_cast<uint64>(0xccddeeff11223344),
+						static_cast<uint16>(0x5566),
+						static_cast<uint8>(0x77),
+						static_cast<uint16>(0x8899)
+					});
+
+				FSortPrefixResult Result;
+				Result = Sorter.SortPrefix(*DataStorage, 0, 0);
+				CHECK(Result.Prefix == 0x1122334455667788);
+				CHECK(Result.bHasRemainingBytes == true);
+
+				Result = Sorter.SortPrefix(*DataStorage, 0, 8);
+				CHECK(Result.Prefix == 0x99aabbccddeeff11);
+				CHECK(Result.bHasRemainingBytes == true);
+
+				Result = Sorter.SortPrefix(*DataStorage, 0, 16);
+				CHECK(Result.Prefix == 0x2233445566778899);
+				CHECK(Result.bHasRemainingBytes == false);
+			}
+		}
+	}
+
+	TEST_CASE_NAMED(FEditorDataStorageUiSorter_Prefix_StringIndexCheck,
+		"Editor::DataStorage::FColumnSorterInterface_Prefix_StringIndexCheck", "[ApplicationContextMask][EngineFilter]")
+	{
+		if (ICoreProvider* DataStorage = GetMutableDataStorageFeature<ICoreProvider>(StorageFeatureName))
+		{
+			using SortStringView = TSortStringView<FSortCaseSensitive, FWideStringView>;
+			FSingleValueTestSorter<SortStringView> Sorter;
+			FSortPrefixResult Result;
+			{
+				FWideString String = TEXT("a");
+				Sorter.Values.Add(SortStringView(String));
+				Result = Sorter.SortPrefix(*DataStorage, 0, 0);
+				CHECK(Result.Prefix == 0x0061000000000000);
+				CHECK(Result.bHasRemainingBytes == false);
+				Sorter.Values.Empty();
+			}
+			{
+				FWideString String = TEXT("aa");
+				Sorter.Values.Add(SortStringView(String));
+				Result = Sorter.SortPrefix(*DataStorage, 0, 0);
+				CHECK(Result.Prefix == 0x0061006100000000);
+				CHECK(Result.bHasRemainingBytes == false);
+				Sorter.Values.Empty();
+			}
+			{
+				FWideString String = TEXT("aaa");
+				Sorter.Values.Add(SortStringView(String));
+				Result = Sorter.SortPrefix(*DataStorage, 0, 0);
+				CHECK(Result.Prefix == 0x0061006100610000);
+				CHECK(Result.bHasRemainingBytes == false);
+				Sorter.Values.Empty();
+			}
+			{
+				FWideString String = TEXT("aaaa");
+				Sorter.Values.Add(SortStringView(String));
+				Result = Sorter.SortPrefix(*DataStorage, 0, 0);
+				CHECK(Result.Prefix == 0x0061006100610061);
+				CHECK(Result.bHasRemainingBytes == false);
+				Sorter.Values.Empty();
+			}
+			{
+				FWideString String = TEXT("aaaab");
+				Sorter.Values.Add(SortStringView(String));
+				Result = Sorter.SortPrefix(*DataStorage, 0, 0);
+				CHECK(Result.Prefix == 0x0061006100610061);
+				CHECK(Result.bHasRemainingBytes == true);
+				Result = Sorter.SortPrefix(*DataStorage, 0, 8);
+				CHECK(Result.Prefix == 0x0062000000000000);
+				CHECK(Result.bHasRemainingBytes == false);
+				Sorter.Values.Empty();
+			}
+			{
+				FWideString String = TEXT("aaaabb");
+				Sorter.Values.Add(SortStringView(String));
+				Result = Sorter.SortPrefix(*DataStorage, 0, 0);
+				CHECK(Result.Prefix == 0x0061006100610061);
+				CHECK(Result.bHasRemainingBytes == true);
+				Result = Sorter.SortPrefix(*DataStorage, 0, 8);
+				CHECK(Result.Prefix == 0x0062006200000000);
+				CHECK(Result.bHasRemainingBytes == false);
+				Sorter.Values.Empty();
+			}
+		}
+	}
+
+	TEST_CASE_NAMED(FEditorDataStorageUiSorter_Prefix_StringIndexCheckWithOffset,
+		"Editor::DataStorage::FColumnSorterInterface_Prefix_StringIndexCheckWithOffset", "[ApplicationContextMask][EngineFilter]")
+	{
+		if (ICoreProvider* DataStorage = GetMutableDataStorageFeature<ICoreProvider>(StorageFeatureName))
+		{
+			using SortStringView = TSortStringView<FSortCaseSensitive, FWideStringView>;
+			FMultiValueTestSorter<uint16, uint32, SortStringView> Sorter;
+			FSortPrefixResult Result;
+			{
+				FWideString String = TEXT("a");
+				Sorter.Values.Add({ static_cast<uint16>(0xaaaa), 0xbbbbcccc, SortStringView(String) });
+				Result = Sorter.SortPrefix(*DataStorage, 0, 0);
+				CHECK(Result.Prefix == 0xaaaabbbbcccc0061);
+				CHECK(Result.bHasRemainingBytes == false);
+				Sorter.Values.Empty();
+			}
+			{
+				FWideString String = TEXT("ab");
+				Sorter.Values.Add({ static_cast<uint16>(0xaaaa), 0xbbbbcccc, SortStringView(String) });
+				Result = Sorter.SortPrefix(*DataStorage, 0, 0);
+				CHECK(Result.Prefix == 0xaaaabbbbcccc0061);
+				CHECK(Result.bHasRemainingBytes == true);
+				Result = Sorter.SortPrefix(*DataStorage, 0, 8);
+				CHECK(Result.Prefix == 0x0062000000000000);
+				CHECK(Result.bHasRemainingBytes == false);
+				Sorter.Values.Empty();
+			}
+			{
+				FWideString String = TEXT("abc");
+				Sorter.Values.Add({ static_cast<uint16>(0xaaaa), 0xbbbbcccc, SortStringView(String) });
+				Result = Sorter.SortPrefix(*DataStorage, 0, 0);
+				CHECK(Result.Prefix == 0xaaaabbbbcccc0061);
+				CHECK(Result.bHasRemainingBytes == true);
+				Result = Sorter.SortPrefix(*DataStorage, 0, 8);
+				CHECK(Result.Prefix == 0x0062006300000000);
+				CHECK(Result.bHasRemainingBytes == false);
+				Sorter.Values.Empty();
+			}
+			{
+				FWideString String = TEXT("abcd");
+				Sorter.Values.Add({ static_cast<uint16>(0xaaaa), 0xbbbbcccc, SortStringView(String) });
+				Result = Sorter.SortPrefix(*DataStorage, 0, 0);
+				CHECK(Result.Prefix == 0xaaaabbbbcccc0061);
+				CHECK(Result.bHasRemainingBytes == true);
+				Result = Sorter.SortPrefix(*DataStorage, 0, 8);
+				CHECK(Result.Prefix == 0x0062006300640000);
+				CHECK(Result.bHasRemainingBytes == false);
+				Sorter.Values.Empty();
+			}
+			{
+				FWideString String = TEXT("abcde");
+				Sorter.Values.Add({ static_cast<uint16>(0xaaaa), 0xbbbbcccc, SortStringView(String) });
+				Result = Sorter.SortPrefix(*DataStorage, 0, 0);
+				CHECK(Result.Prefix == 0xaaaabbbbcccc0061);
+				CHECK(Result.bHasRemainingBytes == true);
+				Result = Sorter.SortPrefix(*DataStorage, 0, 8);
+				CHECK(Result.Prefix == 0x0062006300640065);
+				CHECK(Result.bHasRemainingBytes == false);
+				Sorter.Values.Empty();
+			}
+			{
+				FWideString String = TEXT("abcdef");
+				Sorter.Values.Add({ static_cast<uint16>(0xaaaa), 0xbbbbcccc, SortStringView(String) });
+				Result = Sorter.SortPrefix(*DataStorage, 0, 0);
+				CHECK(Result.Prefix == 0xaaaabbbbcccc0061);
+				CHECK(Result.bHasRemainingBytes == true);
+				Result = Sorter.SortPrefix(*DataStorage, 0, 8);
+				CHECK(Result.Prefix == 0x0062006300640065);
+				CHECK(Result.bHasRemainingBytes == true);
+				Result = Sorter.SortPrefix(*DataStorage, 0, 16);
+				CHECK(Result.Prefix == 0x0066000000000000);
+				CHECK(Result.bHasRemainingBytes == false);
+				Sorter.Values.Empty();
+			}
+		}
+	}
+
+	TEST_CASE_NAMED(FEditorDataStorageUiSorter_Prefix_StringView_CaseSensitive,
+		"Editor::DataStorage::FColumnSorterInterface_Prefix_StringView_CaseSensitive", "[ApplicationContextMask][EngineFilter]")
+	{
+		if (ICoreProvider* DataStorage = GetMutableDataStorageFeature<ICoreProvider>(StorageFeatureName))
+		{
+			TArray<FString> Strings{ TEXT("a"), TEXT("bba"), TEXT("Aba"), TEXT("cba"), TEXT("ABac"), TEXT("longstring") };
+			FSingleValueTestSorter<TSortStringView<FSortCaseSensitive, FStringView>> Sorter;
+			Sorter.Values.Append(Strings);
+			Sorter.SortByPrefix(*DataStorage, 0);
+			CHECK(Sorter.Verify(*DataStorage));
+		}
+	}
+
+	TEST_CASE_NAMED(FEditorDataStorageUiSorter_Prefix_StringView_CaseInsensitive,
+		"Editor::DataStorage::FColumnSorterInterface_Prefix_StringView_CaseInsensitive", "[ApplicationContextMask][EngineFilter]")
+	{
+		if (ICoreProvider* DataStorage = GetMutableDataStorageFeature<ICoreProvider>(StorageFeatureName))
+		{
+			TArray<FString> Strings{ TEXT("a"), TEXT("bba"), TEXT("Aba"), TEXT("cba"), TEXT("ABac"), TEXT("longstring") };
+			FSingleValueTestSorter<TSortStringView<FSortCaseInsensitive, FStringView>> Sorter;
+			Sorter.Values.Append(Strings);
+			Sorter.SortByPrefix(*DataStorage, 0);
+			CHECK(Sorter.Verify(*DataStorage));
+		}
+	}
+
+	TEST_CASE_NAMED(FEditorDataStorageUiSorter_Prefix_WideStringView_CaseSensitive,
+		"Editor::DataStorage::FColumnSorterInterface_Prefix_WideStringView_CaseSensitive", "[ApplicationContextMask][EngineFilter]")
+	{
+		if (ICoreProvider* DataStorage = GetMutableDataStorageFeature<ICoreProvider>(StorageFeatureName))
+		{
+			TArray<FWideString> Strings
+			{ 
+				TEXT("a"), 
+				TEXT("bba"), 
+				TEXT("Aba"), 
+				TEXT("cba"), 
+				TEXT("ABac"), 
+				TEXT("longstring")
+			};
+			FSingleValueTestSorter<TSortStringView<FSortCaseSensitive, FWideStringView>> Sorter;
+			Sorter.Values.Append(Strings);
+			Sorter.SortByPrefix(*DataStorage, 0);
+			CHECK(Sorter.Verify(*DataStorage));
+		}
+	}
+
+	TEST_CASE_NAMED(FEditorDataStorageUiSorter_Prefix_WideStringView_CaseInsensitive,
+		"Editor::DataStorage::FColumnSorterInterface_Prefix_WideStringView_CaseInsensitive", "[ApplicationContextMask][EngineFilter]")
+	{
+		if (ICoreProvider* DataStorage = GetMutableDataStorageFeature<ICoreProvider>(StorageFeatureName))
+		{
+			TArray<FWideString> Strings
+			{ 
+				TEXT("a"), 
+				TEXT("bba"), 
+				TEXT("Aba"), 
+				TEXT("cba"), 
+				TEXT("ABac"), 
+				TEXT("longstring")
+			};
+			FSingleValueTestSorter<TSortStringView<FSortCaseInsensitive, FWideStringView>> Sorter;
+			Sorter.Values.Append(Strings);
+			Sorter.SortByPrefix(*DataStorage, 0);
+			CHECK(Sorter.Verify(*DataStorage));
+		}
+	}
+
+	TEST_CASE_NAMED(FEditorDataStorageUiSorter_Prefix_AnsiStringView_CaseSensitive,
+		"Editor::DataStorage::FColumnSorterInterface_Prefix_AnsiStringView_CaseSensitive", "[ApplicationContextMask][EngineFilter]")
+	{
+		if (ICoreProvider* DataStorage = GetMutableDataStorageFeature<ICoreProvider>(StorageFeatureName))
+		{
+			TArray<FAnsiString> Strings{ "a", "bba", "Aba", "cba", "ABac", "longstring" };
+			FSingleValueTestSorter<TSortStringView<FSortCaseSensitive, FAnsiStringView>> Sorter;
+			Sorter.Values.Append(Strings);
+			Sorter.SortByPrefix(*DataStorage, 0);
+			CHECK(Sorter.Verify(*DataStorage));
+		}
+	}
+
+	TEST_CASE_NAMED(FEditorDataStorageUiSorter_Prefix_AnsiStringView_CaseInsensitive,
+		"Editor::DataStorage::FColumnSorterInterface_Prefix_AnsiStringView_CaseInsensitive", "[ApplicationContextMask][EngineFilter]")
+	{
+		if (ICoreProvider* DataStorage = GetMutableDataStorageFeature<ICoreProvider>(StorageFeatureName))
+		{
+			TArray<FAnsiString> Strings{ "a", "bba", "Aba", "cba", "ABac", "longstring" };
+			FSingleValueTestSorter<TSortStringView<FSortCaseInsensitive, FAnsiStringView>> Sorter;
+			Sorter.Values.Append(Strings);
+			Sorter.SortByPrefix(*DataStorage, 0);
+			CHECK(Sorter.Verify(*DataStorage));
+		}
+	}
+
+	TEST_CASE_NAMED(FEditorDataStorageUiSorter_Prefix_TextView_CaseInsensitive,
+		"Editor::DataStorage::FColumnSorterInterface_Prefix_TextView_CaseInsensitive", "[ApplicationContextMask][EngineFilter]")
+	{
+		if (ICoreProvider* DataStorage = GetMutableDataStorageFeature<ICoreProvider>(StorageFeatureName))
+		{
+			TArray<FText> Strings
+			{ 
+				FText::FromString("a"), 
+				FText::FromString("bba"), 
+				FText::FromString("Aba"), 
+				FText::FromString("cba"), 
+				FText::FromString("ABac"), 
+				FText::FromString("longstring")
+			};
+			FSingleValueTestSorter<TSortStringView<FSortCaseInsensitive, FStringView>> Sorter;
+			Sorter.Values.Append(Strings);
+			Sorter.SortByPrefix(*DataStorage, 0);
+			CHECK(Sorter.Verify(*DataStorage));
+		}
+	}
+
+	TEST_CASE_NAMED(FEditorDataStorageUiSorter_Prefix_TextView_CaseSensitive,
+		"Editor::DataStorage::FColumnSorterInterface_Prefix_TextView_CaseSensitive", "[ApplicationContextMask][EngineFilter]")
+	{
+		if (ICoreProvider* DataStorage = GetMutableDataStorageFeature<ICoreProvider>(StorageFeatureName))
+		{
+			TArray<FText> Strings
+			{
+				FText::FromString("a"),
+				FText::FromString("bba"),
+				FText::FromString("Aba"),
+				FText::FromString("cba"),
+				FText::FromString("ABac"),
+				FText::FromString("longstring")
+			};
+			FSingleValueTestSorter<TSortStringView<FSortCaseSensitive, FStringView>> Sorter;
+			Sorter.Values.Append(Strings);
+			Sorter.SortByPrefix(*DataStorage, 0);
+			CHECK(Sorter.Verify(*DataStorage));
+		}
+	}
+
+	TEST_CASE_NAMED(FEditorDataStorageUiSorter_Prefix_FNameByString,
+		"Editor::DataStorage::FColumnSorterInterface_Prefix_FNameByString", "[ApplicationContextMask][EngineFilter]")
+	{
+		if (ICoreProvider* DataStorage = GetMutableDataStorageFeature<ICoreProvider>(StorageFeatureName))
+		{
+			TArray<FName> Names{ "a", "bba", "Aba", "cba", "ABac", "LongName" };
+			FSingleValueTestSorter<TSortNameView<FSortByName>> Sorter;
+			Sorter.Values.Append(Names);
+			Sorter.SortByPrefix(*DataStorage, 0);
+			CHECK(Sorter.Verify(*DataStorage));
+		}
+	}
+
+	TEST_CASE_NAMED(FEditorDataStorageUiSorter_Prefix_FNameById,
+		"Editor::DataStorage::FColumnSorterInterface_Prefix_FNameById", "[ApplicationContextMask][EngineFilter]")
+	{
+		if (ICoreProvider* DataStorage = GetMutableDataStorageFeature<ICoreProvider>(StorageFeatureName))
+		{
+			TArray<FName> Names{ "a", "bba", "Aba", "cba", "ABac", "LongName" };
+			FSingleValueTestSorter<TSortNameView<FSortById>> Sorter;
+			Sorter.Values.Append(Names);
+			Sorter.SortByPrefix(*DataStorage, 0);
+			CHECK(Sorter.Verify(*DataStorage));
+		}
+	}
+
+	TEST_CASE_NAMED(FEditorDataStorageUiSorter_Prefix_8ByteMultiValue, 
+		"Editor::DataStorage::FColumnSorterInterface_Prefix_8ByteMultiValue", "[ApplicationContextMask][EngineFilter]")
+	{
+		if (ICoreProvider* DataStorage = GetMutableDataStorageFeature<ICoreProvider>(StorageFeatureName))
+		{
+			FMultiValueTestSorter<uint64, uint64, uint64> Sorter;
+			Sorter.Values.Append(
+				{
+					{ 42, 88, 31 },
+					{ 3, 102, 88 },
+					{ 88, 42, 400 },
+					{ 88, 16, 1 },
+					{ 88, 32, 24 },
+					{ 88, 32, 12 },
+					{ 4, 502998, 87 },
+					{ 502998, 17, 99}
+				});
+			// Mimic the steps of a bucket sort.
+			Sorter.SortByPrefix(*DataStorage, 0, 0, 7);
+			Sorter.SortByPrefix(*DataStorage, 8, 3, 6);
+			Sorter.SortByPrefix(*DataStorage, 16, 4, 5);
+			CHECK(Sorter.Verify(*DataStorage));
+		}
+	}
+
+	TEST_CASE_NAMED(FEditorDataStorageUiSorter_Prefix_4ByteMultiValue, 
+		"Editor::DataStorage::FColumnSorterInterface_Prefix_4ByteMultiValue", "[ApplicationContextMask][EngineFilter]")
+	{
+		if (ICoreProvider* DataStorage = GetMutableDataStorageFeature<ICoreProvider>(StorageFeatureName))
+		{
+			FMultiValueTestSorter<uint32, uint32, uint32> Sorter;
+			Sorter.Values.Append(
+				{
+					{ 42, 88, 31 },
+					{ 3, 102, 88 },
+					{ 88, 42, 400 },
+					{ 88, 16, 1 },
+					{ 88, 32, 24 },
+					{ 88, 32, 12 },
+					{ 4, 502998, 87 },
+					{ 502998, 17, 99}
+				});
+			// Mimic the steps of a bucket sort.
+			Sorter.SortByPrefix(*DataStorage, 0, 0, 7);
+			Sorter.SortByPrefix(*DataStorage, 8, 4, 5);
+			CHECK(Sorter.Verify(*DataStorage));
+		}
+	}
+
+	TEST_CASE_NAMED(FEditorDataStorageUiSorter_Prefix_MultiValueMisMatchedSizes, 
+		"Editor::DataStorage::FColumnSorterInterface_Prefix_MultiValueMisMatchedSizes", "[ApplicationContextMask][EngineFilter]")
+	{
+		if (ICoreProvider* DataStorage = GetMutableDataStorageFeature<ICoreProvider>(StorageFeatureName))
+		{
+			FMultiValueTestSorter<int8, uint16, uint64> Sorter;
+			Sorter.Values.Append(
+				{
+					{ static_cast<int8>(42), static_cast<uint16>(88), 31ull },
+					{ static_cast<int8>(-3), static_cast<uint16>(102), 88ull },
+					{ static_cast<int8>(88), static_cast<uint16>(42), 400ull },
+					{ static_cast<int8>(88), static_cast<uint16>(16), 1ull },
+					{ static_cast<int8>(88), static_cast<uint16>(32), 0xaabbccdd11223333ull },
+					{ static_cast<int8>(88), static_cast<uint16>(32), 0xaabbccdd11223344ull },
+					{ static_cast<int8>(88), static_cast<uint16>(32), 0xaabbccdd11223322ull },
+					{ static_cast<int8>( 4), static_cast<uint16>(502998), 87ull },
+					{ static_cast<int8>(-1), static_cast<uint16>(17), 99ull}
+				});
+			// Mimic the steps of a bucket sort.
+			Sorter.SortByPrefix(*DataStorage, 0, 0, 8);
+			Sorter.SortByPrefix(*DataStorage, 8, 5, 7);
+			CHECK(Sorter.Verify(*DataStorage));
+		}
+	}
+
+	TEST_CASE_NAMED(FEditorDataStorageUiSorter_Prefix_MultiValueStrings,
+		"Editor::DataStorage::FColumnSorterInterface_Prefix_MultiValueStrings", "[ApplicationContextMask][EngineFilter]")
+	{
+		if (ICoreProvider* DataStorage = GetMutableDataStorageFeature<ICoreProvider>(StorageFeatureName))
+		{
+			TArray<FString> Strings{ TEXT("a"), TEXT("bba"), TEXT("cb"), TEXT("aba"), TEXT("abac"), TEXT("longstring") };
+			
+			FMultiValueTestSorter<int16, TSortStringView<FSortCaseInsensitive, FStringView>> Sorter;
+			Sorter.Values.Append(
+				{
+					{ static_cast<int16>(42), TSortStringView(FSortCaseInsensitive{}, Strings[0]) },
+					{ static_cast<int16>(-3), TSortStringView(FSortCaseInsensitive{}, Strings[1]) },
+					{ static_cast<int16>(88), TSortStringView(FSortCaseInsensitive{}, Strings[2]) },
+					{ static_cast<int16>(88), TSortStringView(FSortCaseInsensitive{}, Strings[3]) },
+					{ static_cast<int16>( 4), TSortStringView(FSortCaseInsensitive{}, Strings[4]) },
+					{ static_cast<int16>(-1), TSortStringView(FSortCaseInsensitive{}, Strings[5]) }
+				});
+			// Mimic the steps of a bucket sort.
+			Sorter.SortByPrefix(*DataStorage, 0, 0, 5);
+			CHECK(Sorter.Verify(*DataStorage));
+		}
+	}
+
+	TEST_CASE_NAMED(FEditorDataStorageUiSorter_Prefix_MultiValueStringsMatchingNumbers,
+		"Editor::DataStorage::FColumnSorterInterface_Prefix_MultiValueStringsMatchingNumbers", "[ApplicationContextMask][EngineFilter]")
+	{
+		if (ICoreProvider* DataStorage = GetMutableDataStorageFeature<ICoreProvider>(StorageFeatureName))
+		{
+			TArray<FString> Strings{ TEXT("a"), TEXT("bb_a"), TEXT("cb"), TEXT("ab_ca"), TEXT("ab_ac"), TEXT("longstring") };
+
+			FMultiValueTestSorter<uint32, TSortStringView<FSortCaseInsensitive, FStringView>> Sorter;
+			Sorter.Values.Append(
+				{
+					{ 42u, TSortStringView(FSortCaseInsensitive{}, Strings[0]) },
+					{ 42u, TSortStringView(FSortCaseInsensitive{}, Strings[1]) },
+					{ 42u, TSortStringView(FSortCaseInsensitive{}, Strings[2]) },
+					{ 42u, TSortStringView(FSortCaseInsensitive{}, Strings[3]) },
+					{ 42u, TSortStringView(FSortCaseInsensitive{}, Strings[4]) },
+					{ 42u, TSortStringView(FSortCaseInsensitive{}, Strings[5]) }
+				});
+			// Mimic the steps of a bucket sort.
+			Sorter.SortByPrefix(*DataStorage, 0, 0, 5);
+			Sorter.SortByPrefix(*DataStorage, 8, 1, 2);
+			CHECK(Sorter.Verify(*DataStorage));
+		}
+	}
+} // namespace UE::Editor::DataStorage::Tests
+
+#endif // WITH_TESTS
\ No newline at end of file
