Added TNotNull overload to TEDS Compatibility GetRawPointer function.

diff --git a/Engine/Source/Runtime/TypedElementFramework/Public/Elements/Interfaces/TypedElementDataStorageCompatibilityInterface.h b/Engine/Source/Runtime/TypedElementFramework/Public/Elements/Interfaces/TypedElementDataStorageCompatibilityInterface.h
--- a/Engine/Source/Runtime/TypedElementFramework/Public/Elements/Interfaces/TypedElementDataStorageCompatibilityInterface.h
+++ b/Engine/Source/Runtime/TypedElementFramework/Public/Elements/Interfaces/TypedElementDataStorageCompatibilityInterface.h
@@ -107,6 +107,9 @@ template<typename Type> Type* GetRawPointer(const TStrongObjectPtr<Type> Object)
 template<typename Type> Type* GetRawPointer(const TObjectKey<Type> Object)		{ return Object.ResolveObjectPtr(); }
 template<typename Type> Type* GetRawPointer(const TUniquePtr<Type> Object)		{ return Object.Get(); }
 template<typename Type> Type* GetRawPointer(const TSharedPtr<Type> Object)		{ return Object.Get(); }
+#if UE_ENABLE_NOTNULL_WRAPPER
+template<typename Type> Type* GetRawPointer(TNotNull<Type*> Object)				{ return Object; }
+#endif
 template<typename Type> Type* GetRawPointer(Type* Object)						{ return Object; }
 template<typename Type> Type* GetRawPointer(Type& Object)						{ return &Object; }
 
