[TEDS] Three new TEDS Query Stack nodes.

Three new nodes were added to the Query Stack:
- Row Copy - Copies the content for a row locally. For memory space and performance reasons, row nodes can now share their data with other nodes. This is not always desirable, so this node was added to allow a local copy to be created. This for instance allows a parent node to be sorted one way and a child node in another way.
- Row handle sorter - Can sort row handles specifically. This is faster than using the upcoming sort node in general but even faster if the row handle array is already sorted as won't do any work. Unlike the sort node, this node also maintains/sets the sorted status on the row handle array, allowing for further optimizations down when using the row handle array.
- Row inversion - Inverts the order of the rows. Makes sorting from ascending to descending order really fast.

Additional changes made:
- When adding values to the row handle array, users can now provide a guarantee that the new value maintains order and/or uniqueness.
- A function was added to invert the order of the rows.
- Mutable access to the array of rows in the row handle array can be requested. By default this resets the sorting and uniqueness status, but users can provide a guarantee they'll maintain those.
- The row handle array view got a cast operator to convert to a standard const array view for row handles. This makes the view slightly easier to use with systems that (still) use a plain row handle view.

[RN] [TEDS] Three new TEDS Query Stack nodes to copy rows, sort by row handle and to invert the list of rows.

diff --git a/Engine/Plugins/Experimental/EditorDataStorageFeatures/Source/TedsQueryStack/Private/TedsRowCopyNode.cpp b/Engine/Plugins/Experimental/EditorDataStorageFeatures/Source/TedsQueryStack/Private/TedsRowCopyNode.cpp
new file mode 100644
--- /dev/null
+++ b/Engine/Plugins/Experimental/EditorDataStorageFeatures/Source/TedsQueryStack/Private/TedsRowCopyNode.cpp
@@ -0,0 +1,45 @@
+// Copyright Epic Games, Inc. All Rights Reserved.
+
+#include "TedsRowCopyNode.h"
+
+namespace UE::Editor::DataStorage::QueryStack
+{
+	FRowCopyNode::FRowCopyNode(TSharedPtr<IRowNode> InParent)
+		: Parent(MoveTemp(InParent))
+	{
+		Rows.Append(Parent->GetRows());
+	}
+
+	void FRowCopyNode::Reset()
+	{
+		Rows.Empty();
+		Rows.Append(Parent->GetRows());
+		Revision++;
+	}
+
+	INode::RevisionId FRowCopyNode::GetRevision() const
+	{
+		return Revision;
+	}
+
+	void FRowCopyNode::Update()
+	{
+		Parent->Update();
+
+		if (Parent->GetRevision() != ParentRevision)
+		{
+			Reset();
+			ParentRevision = Parent->GetRevision();
+		}
+	}
+
+	FRowHandleArrayView FRowCopyNode::GetRows() const
+	{
+		return Rows.GetRows();
+	}
+
+	FRowHandleArray* FRowCopyNode::GetMutableRows()
+	{
+		return &Rows;
+	}
+} // namespace UE::Editor::DataStorage::QueryStack
diff --git a/Engine/Plugins/Experimental/EditorDataStorageFeatures/Source/TedsQueryStack/Private/TedsRowHandleSortNode.cpp b/Engine/Plugins/Experimental/EditorDataStorageFeatures/Source/TedsQueryStack/Private/TedsRowHandleSortNode.cpp
new file mode 100644
--- /dev/null
+++ b/Engine/Plugins/Experimental/EditorDataStorageFeatures/Source/TedsQueryStack/Private/TedsRowHandleSortNode.cpp
@@ -0,0 +1,43 @@
+// Copyright Epic Games, Inc. All Rights Reserved.
+
+#include "TedsRowHandleSortNode.h"
+
+namespace UE::Editor::DataStorage::QueryStack
+{
+	FRowHandleSortNode::FRowHandleSortNode(TSharedPtr<IRowNode> InParent)
+		: Parent(MoveTemp(InParent))
+	{
+		FRowHandleArray* ParentRows = Parent->GetMutableRows();
+		checkf(ParentRows != nullptr, 
+			TEXT("FRowHandleSortNode requires a parent node that provides write access to its row handle array."));
+
+		ParentRows->Sort();
+		ParentRevision = Parent->GetRevision();
+	}
+
+	INode::RevisionId FRowHandleSortNode::GetRevision() const
+	{
+		return ParentRevision;
+	}
+
+	void FRowHandleSortNode::Update()
+	{
+		Parent->Update();
+
+		if (Parent->GetRevision() != ParentRevision)
+		{
+			Parent->GetMutableRows()->Sort();
+			ParentRevision = Parent->GetRevision();
+		}
+	}
+
+	FRowHandleArrayView FRowHandleSortNode::GetRows() const
+	{
+		return Parent->GetRows();
+	}
+
+	FRowHandleArray* FRowHandleSortNode::GetMutableRows()
+	{
+		return Parent->GetMutableRows();
+	}
+} // namespace UE::Editor::DataStorage::QueryStack
diff --git a/Engine/Plugins/Experimental/EditorDataStorageFeatures/Source/TedsQueryStack/Private/TedsRowOrderInversionNode.cpp b/Engine/Plugins/Experimental/EditorDataStorageFeatures/Source/TedsQueryStack/Private/TedsRowOrderInversionNode.cpp
new file mode 100644
--- /dev/null
+++ b/Engine/Plugins/Experimental/EditorDataStorageFeatures/Source/TedsQueryStack/Private/TedsRowOrderInversionNode.cpp
@@ -0,0 +1,79 @@
+// Copyright Epic Games, Inc. All Rights Reserved.
+
+#include "TedsRowOrderInversionNode.h"
+
+namespace UE::Editor::DataStorage::QueryStack
+{
+	FRowOrderInversionNode::FRowOrderInversionNode(TSharedPtr<IRowNode> InParent, bool bInIsEnabled)
+		: Parent(MoveTemp(InParent))
+		, bIsEnabled(bInIsEnabled)
+	{
+		FRowHandleArray* Rows = Parent->GetMutableRows();
+		checkf(Rows != nullptr, TEXT("FRowOrderInversionNode requires a parent that can provide mutable rows."));
+
+		if (bInIsEnabled)
+		{
+			ApplyOrder();
+		}
+		else
+		{
+			bRequiresSync = true;
+		}
+	}
+
+	void FRowOrderInversionNode::Enable(bool bInIsEnabled)
+	{
+		bIsEnabled = bInIsEnabled;
+	}
+
+	INode::RevisionId FRowOrderInversionNode::GetRevision() const
+	{
+		return Revision;
+	}
+
+	void FRowOrderInversionNode::Update()
+	{
+		Parent->Update();
+
+		RevisionId LocalParentRevision = Parent->GetRevision();
+		if (bIsEnabled)
+		{
+			if (LocalParentRevision != ParentRevision || bRequiresSync)
+			{
+				ParentRevision = LocalParentRevision;
+				Revision++;
+				bRequiresSync = false;
+				ApplyOrder();
+			}
+		}
+		else
+		{
+			if (LocalParentRevision != ParentRevision)
+			{
+				ParentRevision = LocalParentRevision;
+				Revision++;
+				bRequiresSync = true;
+			}
+		}
+	}
+
+	FRowHandleArrayView FRowOrderInversionNode::GetRows() const
+	{
+		return Parent->GetRows();
+	}
+
+	FRowHandleArray* FRowOrderInversionNode::GetMutableRows()
+	{
+		return Parent->GetMutableRows();
+	}
+
+	void FRowOrderInversionNode::ApplyOrder()
+	{
+		if (FRowHandleArray* Rows = Parent->GetMutableRows())
+		{
+			Rows->InvertOrder();
+		}
+
+		Revision++;
+	}
+} // namespace UE::Editor::DataStorage::QueryStack
diff --git a/Engine/Plugins/Experimental/EditorDataStorageFeatures/Source/TedsQueryStack/Public/TedsRowCopyNode.h b/Engine/Plugins/Experimental/EditorDataStorageFeatures/Source/TedsQueryStack/Public/TedsRowCopyNode.h
new file mode 100644
--- /dev/null
+++ b/Engine/Plugins/Experimental/EditorDataStorageFeatures/Source/TedsQueryStack/Public/TedsRowCopyNode.h
@@ -0,0 +1,37 @@
+// Copyright Epic Games, Inc. All Rights Reserved.
+
+#pragma once
+
+#include "Containers/ContainersFwd.h"
+#include "TedsQueryStackInterfaces.h"
+#include "Elements/Framework//TypedElementRowHandleArray.h"
+
+namespace UE::Editor::DataStorage::QueryStack
+{
+	/**
+	 * Copies the content of the parent node to a local array. This node can be used
+	 * to avoid changes to the parent node being made. An example usage would be if
+	 * the parent node is sorted by row handle for performance reasons but the final
+	 * list needs to be sorted based on other criteria.
+	 */
+	class FRowCopyNode : public IRowNode
+	{
+	public:
+		TEDSQUERYSTACK_API explicit FRowCopyNode(TSharedPtr<IRowNode> InParent);
+		virtual ~FRowCopyNode() override = default;
+
+		/** Clears out the current row list and gets a fresh copy from the parent. */
+		TEDSQUERYSTACK_API void Reset();
+
+		TEDSQUERYSTACK_API virtual RevisionId GetRevision() const override;
+		TEDSQUERYSTACK_API virtual void Update() override;
+		TEDSQUERYSTACK_API virtual FRowHandleArrayView GetRows() const override;
+		TEDSQUERYSTACK_API virtual FRowHandleArray* GetMutableRows() override;
+
+	private:
+		FRowHandleArray Rows;
+		TSharedPtr<IRowNode> Parent;
+		RevisionId ParentRevision = 0;
+		RevisionId Revision = 0;
+	};
+} // namespace UE::Editor::DataStorage::QueryStack
diff --git a/Engine/Plugins/Experimental/EditorDataStorageFeatures/Source/TedsQueryStack/Public/TedsRowHandleSortNode.h b/Engine/Plugins/Experimental/EditorDataStorageFeatures/Source/TedsQueryStack/Public/TedsRowHandleSortNode.h
new file mode 100644
--- /dev/null
+++ b/Engine/Plugins/Experimental/EditorDataStorageFeatures/Source/TedsQueryStack/Public/TedsRowHandleSortNode.h
@@ -0,0 +1,30 @@
+// Copyright Epic Games, Inc. All Rights Reserved.
+
+#pragma once
+
+#include "TedsQueryStackInterfaces.h"
+
+namespace UE::Editor::DataStorage::QueryStack
+{
+	/**
+	 * Sorts the rows in the parent by row handle. Prefer this over using the FRowSortNode when sorting
+	 * by row handle as this is somewhat faster and allows the row handle array to get its sorted status.
+	 * If the row handle array is already sorted this node doesn't do any work, making it even faster than
+	 * using FRowSortNode.
+	 */
+	class FRowHandleSortNode : public IRowNode
+	{
+	public:
+		TEDSQUERYSTACK_API explicit FRowHandleSortNode(TSharedPtr<IRowNode> InParent);
+		virtual ~FRowHandleSortNode() override = default;
+
+		TEDSQUERYSTACK_API virtual RevisionId GetRevision() const override;
+		TEDSQUERYSTACK_API virtual void Update() override;
+		TEDSQUERYSTACK_API virtual FRowHandleArrayView GetRows() const override;
+		TEDSQUERYSTACK_API virtual FRowHandleArray* GetMutableRows() override;
+
+	private:
+		TSharedPtr<IRowNode> Parent;
+		RevisionId ParentRevision = 0;
+	};
+} // namespace UE::Editor::DataStorage::QueryStack
diff --git a/Engine/Plugins/Experimental/EditorDataStorageFeatures/Source/TedsQueryStack/Public/TedsRowOrderInversionNode.h b/Engine/Plugins/Experimental/EditorDataStorageFeatures/Source/TedsQueryStack/Public/TedsRowOrderInversionNode.h
new file mode 100644
--- /dev/null
+++ b/Engine/Plugins/Experimental/EditorDataStorageFeatures/Source/TedsQueryStack/Public/TedsRowOrderInversionNode.h
@@ -0,0 +1,38 @@
+// Copyright Epic Games, Inc. All Rights Reserved.
+
+#pragma once
+
+#include "Containers/ContainersFwd.h"
+#include "TedsQueryStackInterfaces.h"
+#include "Elements/Framework//TypedElementRowHandleArray.h"
+
+namespace UE::Editor::DataStorage::QueryStack
+{
+	/**
+	 * When enabled, orders the rows on the parent in reverse. If disabled, this row has no effect, but if the order on the parent
+	 * was inverted it will remain inverted.
+	 */
+	class FRowOrderInversionNode : public IRowNode
+	{
+	public:
+		TEDSQUERYSTACK_API FRowOrderInversionNode(TSharedPtr<IRowNode> InParent, bool bInIsEnabled);
+		virtual ~FRowOrderInversionNode() override = default;
+
+		/** Clears out the current row list and gets a fresh copy from the parent. */
+		TEDSQUERYSTACK_API void Enable(bool bInIsEnabled);
+
+		TEDSQUERYSTACK_API virtual RevisionId GetRevision() const override;
+		TEDSQUERYSTACK_API virtual void Update() override;
+		TEDSQUERYSTACK_API virtual FRowHandleArrayView GetRows() const override;
+		TEDSQUERYSTACK_API virtual FRowHandleArray* GetMutableRows() override;
+
+	private:
+		void ApplyOrder();
+
+		TSharedPtr<IRowNode> Parent;
+		RevisionId ParentRevision = 0;
+		RevisionId Revision = 0;
+		bool bIsEnabled = true;
+		bool bRequiresSync = false;
+	};
+} // namespace UE::Editor::DataStorage::QueryStack
diff --git a/Engine/Source/Runtime/TypedElementFramework/Private/Elements/Framework/TypedElementRowHandleArray.cpp b/Engine/Source/Runtime/TypedElementFramework/Private/Elements/Framework/TypedElementRowHandleArray.cpp
--- a/Engine/Source/Runtime/TypedElementFramework/Private/Elements/Framework/TypedElementRowHandleArray.cpp
+++ b/Engine/Source/Runtime/TypedElementFramework/Private/Elements/Framework/TypedElementRowHandleArray.cpp
@@ -9,6 +9,7 @@
 #include "Algo/IndexOf.h"
 #include "HAL/UnrealMemory.h"
 #include "Templates/Sorting.h"
+#include "Templates/UnrealTemplate.h"
 
 namespace UE::Editor::DataStorage
 {
@@ -35,19 +36,30 @@ namespace UE::Editor::DataStorage
 		Rows.Shrink();
 	}
 
-	void FRowHandleArray::Add(RowHandle Row)
+	void FRowHandleArray::Add(RowHandle Row, EFlags Guarantee)
 	{
 		bIsSorted = bIsSorted && (Rows.IsEmpty() || Rows.Last() <= Row);
-		bIsUnique = bIsUnique && bIsSorted && (Rows.IsEmpty() || Rows.Last() < Row);
+		bIsUnique = bIsUnique && (bIsSorted ? (Rows.IsEmpty() || Rows.Last() < Row) : (EnumHasAllFlags(Guarantee, EFlags::IsUnique)));
 		Rows.Add(Row);
 	}
 
-	void FRowHandleArray::Append(TConstArrayView<RowHandle> AdditionalRows)
+	void FRowHandleArray::Append(TConstArrayView<RowHandle> AdditionalRows, EFlags Guarantee)
 	{
 		if (!AdditionalRows.IsEmpty())
 		{
-			bIsSorted = false;
-			bIsUnique = false;
+			if (!Rows.IsEmpty())
+			{
+				RowHandle CurrentLast = Rows.Last();
+				RowHandle FirstNew = *AdditionalRows.begin();
+				bIsSorted = bIsSorted && EnumHasAllFlags(Guarantee, EFlags::IsSorted) && (CurrentLast <= FirstNew);
+				bIsUnique = bIsSorted && bIsUnique && EnumHasAllFlags(Guarantee, EFlags::IsUnique) && (CurrentLast < FirstNew);
+			}
+			else
+			{
+				bIsSorted = EnumHasAllFlags(Guarantee, EFlags::IsSorted);
+				bIsUnique = EnumHasAllFlags(Guarantee, EFlags::IsUnique);
+			}
+
 			int32 StartIndex = Rows.AddUninitialized(AdditionalRows.Num());
 			FMemory::Memcpy(Rows.GetData() + StartIndex, AdditionalRows.GetData(), AdditionalRows.NumBytes());
 		}
@@ -365,6 +377,20 @@ namespace UE::Editor::DataStorage
 		}
 	}
 
+	void FRowHandleArray::InvertOrder()
+	{
+		if (!Rows.IsEmpty())
+		{
+			RowHandle* Front = Rows.GetData();
+			RowHandle* End = Front + Rows.Num() - 1;
+			while (Front < End)
+			{
+				Swap(*Front++, *End--);
+			}
+			bIsSorted = false;
+		}
+	}
+
 	FRowHandleArrayView FRowHandleArray::GetRows() const
 	{
 		FRowHandleArrayView::EFlags Flags = FRowHandleArrayView::EFlags::None;
@@ -373,6 +399,13 @@ namespace UE::Editor::DataStorage
 		return FRowHandleArrayView(Rows, Flags);
 	}
 
+	TArrayView<RowHandle> FRowHandleArray::GetMutableRows(EFlags Guarantee)
+	{
+		bIsSorted = bIsSorted && EnumHasAllFlags(Guarantee, EFlags::IsSorted);
+		bIsUnique = bIsUnique && EnumHasAllFlags(Guarantee, EFlags::IsUnique);
+		return TArrayView<RowHandle>(Rows);
+	}
+
 	bool FRowHandleArray::IsSorted() const
 	{
 		return bIsSorted;
diff --git a/Engine/Source/Runtime/TypedElementFramework/Private/Elements/Framework/TypedElementRowHandleArrayView.cpp b/Engine/Source/Runtime/TypedElementFramework/Private/Elements/Framework/TypedElementRowHandleArrayView.cpp
--- a/Engine/Source/Runtime/TypedElementFramework/Private/Elements/Framework/TypedElementRowHandleArrayView.cpp
+++ b/Engine/Source/Runtime/TypedElementFramework/Private/Elements/Framework/TypedElementRowHandleArrayView.cpp
@@ -70,6 +70,11 @@ namespace UE::Editor::DataStorage
 		return RowCount * sizeof(RowHandle);
 	}
 
+	FRowHandleArrayView::operator TConstArrayView<RowHandle>() const
+	{
+		return TConstArrayView<RowHandle>(Rows, RowCount);
+	}
+
 	bool FRowHandleArrayView::IsSorted() const
 	{
 		return bIsSorted;
diff --git a/Engine/Source/Runtime/TypedElementFramework/Public/Elements/Framework/TypedElementRowHandleArray.h b/Engine/Source/Runtime/TypedElementFramework/Public/Elements/Framework/TypedElementRowHandleArray.h
--- a/Engine/Source/Runtime/TypedElementFramework/Public/Elements/Framework/TypedElementRowHandleArray.h
+++ b/Engine/Source/Runtime/TypedElementFramework/Public/Elements/Framework/TypedElementRowHandleArray.h
@@ -15,6 +15,8 @@ namespace UE::Editor::DataStorage
 	class FRowHandleArray
 	{
 	public:
+		using EFlags = FRowHandleArrayView::EFlags;
+
 		FRowHandleArray() = default;
 		TYPEDELEMENTFRAMEWORK_API explicit FRowHandleArray(FRowHandleArrayView InRows);
 
@@ -30,11 +32,21 @@ namespace UE::Editor::DataStorage
 		/** Reduces the amount of memory used by the array to exactly fit the current number or elements. */
 		TYPEDELEMENTFRAMEWORK_API void Shrink();
 
-		/** Appends the provided row handle at the end of the array. This may result in the array becoming unsorted and/or no longer unique. */
-		TYPEDELEMENTFRAMEWORK_API void Add(RowHandle Row);
+		/** 
+		 * Appends the provided row handle at the end of the array. This may result in the array becoming unsorted and/or no longer unique.
+		 * If Guarantee is set to IsUnique the user takes responsibility to guaranteed that the new row is not already present. If array
+		 * is sorted, the IsUnique hint is ignored and the array will determine if the values are unique. The IsSorted guarantee is always
+		 * ignored. If the arrays was already not unique setting the Guarantee flags has no effect.
+		 */
+		TYPEDELEMENTFRAMEWORK_API void Add(RowHandle Row, EFlags Guarantee = EFlags::None);
 
-		/** Append the provided list of rows to the end of the array. This will make the array unsorted and not unique. */
-		TYPEDELEMENTFRAMEWORK_API void Append(TConstArrayView<RowHandle> AdditionalRows);
+		/** 
+		 * Append the provided list of rows to the end of the array. If Guarantee is set to IsSorted and/or IsUnique the user guarantees
+		 * that the provided row handle array view is sorted by row handle and/or only contains unique values. In this case this function
+		 * will try to append the additional rows and maintain the sorted status and if applicable the uniqueness state, otherwise the array 
+		 * will be marked as unsorted and no longer unique.
+		 */
+		TYPEDELEMENTFRAMEWORK_API void Append(TConstArrayView<RowHandle> AdditionalRows, EFlags Guarantee = EFlags::None);
 
 		/**
 		 * Append the provided list of rows to the end of the array. If the current array and the additional rows are both sorted then
@@ -104,8 +116,22 @@ namespace UE::Editor::DataStorage
 		/** Merges the provided array into the array while remaining sorted. */
 		TYPEDELEMENTFRAMEWORK_API void SortedMerge(FRowHandleArray&& AddedRows);
 
+		/**
+		 * Reverses the order in which the rows are stored, making the first element the last and vice versa. If the array
+		 * was marked as sorted it will no longer be considered sorted. Another call to InvertOrder will not restore the
+		 * sorted state.
+		 */
+		TYPEDELEMENTFRAMEWORK_API void InvertOrder();
+
 		/** Constructs a view of the array to provide access to its rows. */
 		TYPEDELEMENTFRAMEWORK_API FRowHandleArrayView GetRows() const;
+		/** 
+		 * Returns a mutable view of the array of rows. This will invalidate sorted and uniqueness status of the array.
+		 * If Guarantee is set to IsSorted and/or IsUnique the user takes responsibility to guaranteed that the rows remain sorted in
+		 * numeric order and/or all values remain unique after the new rows are modified. If the arrays was already not sorted and/or not unique
+		 * setting the Guarantee flags has no effect.
+		 */
+		TYPEDELEMENTFRAMEWORK_API TArrayView<RowHandle> GetMutableRows(EFlags Guarantee = EFlags::None);
 		/** Whether or not the contained rows are numerically ordered from smallest to largest. */
 		TYPEDELEMENTFRAMEWORK_API bool IsSorted() const;
 		/** Whether or all rows in the view only appear once. */
diff --git a/Engine/Source/Runtime/TypedElementFramework/Public/Elements/Framework/TypedElementRowHandleArrayView.h b/Engine/Source/Runtime/TypedElementFramework/Public/Elements/Framework/TypedElementRowHandleArrayView.h
--- a/Engine/Source/Runtime/TypedElementFramework/Public/Elements/Framework/TypedElementRowHandleArrayView.h
+++ b/Engine/Source/Runtime/TypedElementFramework/Public/Elements/Framework/TypedElementRowHandleArrayView.h
@@ -64,6 +64,9 @@ namespace UE::Editor::DataStorage
 		 */
 		TYPEDELEMENTFRAMEWORK_API int32 NumBytes() const;
 
+		/** Conversion operator to allow the view to be transparently used in places where a immutable array view of row handles is required. */
+		TYPEDELEMENTFRAMEWORK_API operator TConstArrayView<RowHandle>() const;
+
 		/** Whether or not the contained rows are numerically ordered from smallest to largest. */
 		TYPEDELEMENTFRAMEWORK_API bool IsSorted() const;
 		/** Whether or all rows in the view only appear once. */
diff --git a/Engine/Source/Runtime/TypedElementFramework/Tests/Elements/Framework/TypedElementRowHandleArrayTests.cpp b/Engine/Source/Runtime/TypedElementFramework/Tests/Elements/Framework/TypedElementRowHandleArrayTests.cpp
--- a/Engine/Source/Runtime/TypedElementFramework/Tests/Elements/Framework/TypedElementRowHandleArrayTests.cpp
+++ b/Engine/Source/Runtime/TypedElementFramework/Tests/Elements/Framework/TypedElementRowHandleArrayTests.cpp
@@ -477,7 +477,7 @@ TEST_CASE_NAMED(FRowHandleArray_Tests, "Editor::DataStorage::Row Handle Array (F
 	}
 
 
-	SECTION("Make v for array with several double and unique values at the start and end.")
+	SECTION("Reduce to duplicates for array with several double and unique values at the start and end.")
 	{
 		RowHandle ArrayData[] = { 1, 1, 2, 3, 3, 4, 5, 5 };
 		FRowHandleArray Array(FRowHandleArrayView(ArrayData, FRowHandleArrayView::EFlags::IsSorted));
@@ -728,5 +728,62 @@ TEST_CASE_NAMED(FRowHandleArray_Tests, "Editor::DataStorage::Row Handle Array (F
 		TestRange.Sort();
 		CHECK_MESSAGE(TEXT("Sorted merge not correct."), Compare(Array, TestRange));
 	}
+
+	SECTION("Invert order")
+	{
+		RowHandle ArrayData[] = { 1, 2, 3, 4, 5 };
+		FRowHandleArray Array(FRowHandleArrayView(ArrayData, FRowHandleArrayView::EFlags::IsSorted | FRowHandleArrayView::EFlags::IsUnique));
+
+		Array.InvertOrder();
+
+		FRowHandleArrayView View = Array.GetRows();
+		CHECK_MESSAGE(TEXT("First row doesn't have the expected value."),  View[0] == 5);
+		CHECK_MESSAGE(TEXT("Second row doesn't have the expected value."), View[1] == 4);
+		CHECK_MESSAGE(TEXT("Third row doesn't have the expected value."),  View[2] == 3);
+		CHECK_MESSAGE(TEXT("Fourth row doesn't have the expected value."), View[3] == 2);
+		CHECK_MESSAGE(TEXT("Fifth row doesn't have the expected value."),  View[4] == 1);
+		CHECK_MESSAGE(TEXT("Array still marked as sorted."), !Array.IsSorted());
+	}
+
+	SECTION("Invert order with even count")
+	{
+		RowHandle ArrayData[] = { 1, 2, 3, 4, 5, 6 };
+		FRowHandleArray Array(FRowHandleArrayView(ArrayData, FRowHandleArrayView::EFlags::IsSorted | FRowHandleArrayView::EFlags::IsUnique));
+
+		Array.InvertOrder();
+
+		FRowHandleArrayView View = Array.GetRows();
+		CHECK_MESSAGE(TEXT("First row doesn't have the expected value."),  View[0] == 6);
+		CHECK_MESSAGE(TEXT("Second row doesn't have the expected value."), View[1] == 5);
+		CHECK_MESSAGE(TEXT("Third row doesn't have the expected value."),  View[2] == 4);
+		CHECK_MESSAGE(TEXT("Fourth row doesn't have the expected value."), View[3] == 3);
+		CHECK_MESSAGE(TEXT("Fifth row doesn't have the expected value."),  View[4] == 2);
+		CHECK_MESSAGE(TEXT("Sixth row doesn't have the expected value."),  View[5] == 1);
+		CHECK_MESSAGE(TEXT("Array still marked as sorted."), !Array.IsSorted());
+	}
+
+	SECTION("Invert order, empty array")
+	{
+		FRowHandleArray Array;
+
+		Array.InvertOrder();
+
+		CHECK_MESSAGE(TEXT("Empty array should remain sorted."), Array.IsSorted());
+	}
+
+	SECTION("Invert order, called twice")
+	{
+		RowHandle ArrayData[] = { 1, 2, 3 };
+		FRowHandleArray Array(FRowHandleArrayView(ArrayData, FRowHandleArrayView::EFlags::IsSorted | FRowHandleArrayView::EFlags::IsUnique));
+
+		Array.InvertOrder();
+		Array.InvertOrder();
+
+		FRowHandleArrayView View = Array.GetRows();
+		CHECK_MESSAGE(TEXT("First row doesn't have the expected value."),  View[0] == 1);
+		CHECK_MESSAGE(TEXT("Second row doesn't have the expected value."), View[1] == 2);
+		CHECK_MESSAGE(TEXT("Third row doesn't have the expected value."),  View[2] == 3);
+		CHECK_MESSAGE(TEXT("Array still marked as sorted."), !Array.IsSorted()); // Array should still be unsorted.
+	}
 }
 #endif // #if WITH_TESTS
