[TEDS] Renamed SortPrefix to CalculatePrefix.

The original intention was for `SortPrefix` to mean "the prefix used for sorting", but it was too easily read as "sort by prefix". To remove the confusion the function was renamed to `CalculatePrefix`. The corresponding result struct was also updated.

#rnx

diff --git a/Engine/Source/Runtime/TypedElementFramework/Public/Elements/Framework/TypedElementSorter.h b/Engine/Source/Runtime/TypedElementFramework/Public/Elements/Framework/TypedElementSorter.h
--- a/Engine/Source/Runtime/TypedElementFramework/Public/Elements/Framework/TypedElementSorter.h
+++ b/Engine/Source/Runtime/TypedElementFramework/Public/Elements/Framework/TypedElementSorter.h
@@ -18,7 +18,7 @@ namespace UE::Editor::DataStorage
 	 */
 	class ICoreProvider;
 
-	struct FSortPrefixResult
+	struct FPrefixInfo
 	{
 		uint64 Prefix = TNumericLimits<uint64>::Max();
 		bool bHasRemainingBytes = false;
@@ -83,7 +83,7 @@ namespace UE::Editor::DataStorage
 		 * Note that using a prefix is limited to only one value that has a variable length (e.g. strings) and has to be the last value
 		 * to be sorted.
 		 */
-		virtual FSortPrefixResult SortPrefix(const ICoreProvider& Storage, RowHandle Row, uint32 ByteIndex) const = 0;
+		virtual FPrefixInfo CalculatePrefix(const ICoreProvider& Storage, RowHandle Row, uint32 ByteIndex) const = 0;
 	};
 
 	namespace Private
@@ -188,7 +188,7 @@ namespace UE::Editor::DataStorage
 
 		uint32 GetByteSize() const;
 		constexpr static uint32 GetElementSize();
-		FSortPrefixResult CalculatePrefix(int32 CurrentIndex, int32 ByteIndex) const;
+		FPrefixInfo CalculatePrefix(int32 CurrentIndex, int32 ByteIndex) const;
 
 		int32 Compare(const TSortNameView& Rhs) const;
 		bool operator==(const TSortNameView& Rhs) const;
@@ -210,7 +210,7 @@ namespace UE::Editor::DataStorage
 	TSortNameView(By, const FName&) -> TSortNameView<By>;
 
 	template<typename... ValueTypes>
-	FSortPrefixResult CreateSortPrefix(uint32 ByteIndex, ValueTypes&&... Values);
+	FPrefixInfo CreateSortPrefix(uint32 ByteIndex, ValueTypes&&... Values);
 
 	/**
 	 * @section Sort type information
@@ -246,11 +246,11 @@ namespace UE::Editor::DataStorage
 		static constexpr bool bIsFixedSize = true;
 		constexpr static uint32 GetByteSize(NumericType Value) { return sizeof(NumericType); }
 		constexpr static uint32 GetElementSize() { return sizeof(NumericType); }
-		constexpr static FSortPrefixResult CalculatePrefix(int32 CurrentIndex, int32 ByteIndex, NumericType Value)
+		constexpr static FPrefixInfo CalculatePrefix(int32 CurrentIndex, int32 ByteIndex, NumericType Value)
 		{
 			// Due to an issue with MSVC resolving forward declared partially specialized templates that use concepts, this 
 			// can't be moved to the .inl file.
-			return FSortPrefixResult
+			return FPrefixInfo
 			{
 				.Prefix = Private::MoveToLocation(CurrentIndex - ByteIndex, Private::Rebase(Value)),
 				.bHasRemainingBytes = false
@@ -265,7 +265,7 @@ namespace UE::Editor::DataStorage
 		static constexpr bool bIsFixedSize = false;
 		constexpr static uint32 GetByteSize(TSortStringView<Casing, T> Value);
 		constexpr static uint32 GetElementSize();
-		constexpr static FSortPrefixResult CalculatePrefix(
+		constexpr static FPrefixInfo CalculatePrefix(
 			int32 CurrentIndex, int32 ByteIndex, TSortStringView<Casing, T> Value);
 	};
 
@@ -284,7 +284,7 @@ namespace UE::Editor::DataStorage
 		static constexpr bool bIsFixedSize = TSortNameView<By>::bIsFixedSize;
 		static uint32 GetByteSize(const TSortNameView<By>& Value) { return Value.GetByteSize(); };
 		constexpr static uint32 GetElementSize() { return TSortNameView<By>::GetElementSize(); };
-		static FSortPrefixResult CalculatePrefix(
+		static FPrefixInfo CalculatePrefix(
 			int32 CurrentIndex, int32 ByteIndex, TSortNameView<By> Value) { return Value.CalculatePrefix(CurrentIndex, ByteIndex); };
 	};
 
diff --git a/Engine/Source/Runtime/TypedElementFramework/Public/Elements/Framework/TypedElementSorter.inl b/Engine/Source/Runtime/TypedElementFramework/Public/Elements/Framework/TypedElementSorter.inl
--- a/Engine/Source/Runtime/TypedElementFramework/Public/Elements/Framework/TypedElementSorter.inl
+++ b/Engine/Source/Runtime/TypedElementFramework/Public/Elements/Framework/TypedElementSorter.inl
@@ -60,11 +60,11 @@ namespace UE::Editor::DataStorage
 		}
 
 		template<typename ValueType>
-		void CalculatePrefix(FSortPrefixResult& Result, int32 CurrentIndex, int32 ByteIndex, const ValueType& Value)
+		void CalculatePrefix(FPrefixInfo& Result, int32 CurrentIndex, int32 ByteIndex, const ValueType& Value)
 		{
 			if constexpr (TSortTypeInfo<ValueType>::bIsSupportedType)
 			{
-				constexpr int32 ResultSize = sizeof(FSortPrefixResult::Prefix);
+				constexpr int32 ResultSize = sizeof(FPrefixInfo::Prefix);
 				int32 Size = static_cast<int32>(TSortTypeInfo<ValueType>::GetByteSize(Value));
 
 				if constexpr (TSortTypeInfo<ValueType>::bIsFixedSize)
@@ -73,7 +73,7 @@ namespace UE::Editor::DataStorage
 					if (Position < ResultSize)
 					{
 						Result.Prefix |= TSortTypeInfo<ValueType>::CalculatePrefix(CurrentIndex, ByteIndex, Value).Prefix;
-						Result.bHasRemainingBytes = (Position + Size) > sizeof(FSortPrefixResult::Prefix);
+						Result.bHasRemainingBytes = (Position + Size) > sizeof(FPrefixInfo::Prefix);
 					}
 				}
 				else
@@ -83,7 +83,7 @@ namespace UE::Editor::DataStorage
 					int32 Index = CurrentIndex - ByteIndex;
 					if (ResultSize - (Index & (ResultSize - 1)) >= static_cast<int32>(ElementSize))
 					{
-						FSortPrefixResult Intermediate =
+						FPrefixInfo Intermediate =
 							TSortTypeInfo<ValueType>::CalculatePrefix(CurrentIndex, ByteIndex, Value);
 						Result.Prefix |= Intermediate.Prefix;
 						Result.bHasRemainingBytes = Intermediate.bHasRemainingBytes;
@@ -93,12 +93,12 @@ namespace UE::Editor::DataStorage
 		}
 
 		template<typename ValueType, typename... ValueTypes>
-		void CalculatePrefix(FSortPrefixResult& Result, int32 CurrentIndex, int32 ByteIndex,
+		void CalculatePrefix(FPrefixInfo& Result, int32 CurrentIndex, int32 ByteIndex,
 			const ValueType& Value, ValueTypes&&... Values)
 		{
 			if constexpr (TSortTypeInfo<ValueType>::bIsSupportedType)
 			{
-				constexpr int32 ResultSize = sizeof(FSortPrefixResult::Prefix);
+				constexpr int32 ResultSize = sizeof(FPrefixInfo::Prefix);
 				int32 Size = static_cast<int32>(TSortTypeInfo<ValueType>::GetByteSize(Value));
 				// Skip over bytes that fall out of the requested range.
 				if (CurrentIndex + Size <= ByteIndex)
@@ -164,7 +164,7 @@ namespace UE::Editor::DataStorage
 	}
 
 	template<SortBy By>
-	FSortPrefixResult TSortNameView<By>::CalculatePrefix(int32 CurrentIndex, int32 ByteIndex) const
+	FPrefixInfo TSortNameView<By>::CalculatePrefix(int32 CurrentIndex, int32 ByteIndex) const
 	{
 		CacheCompareType();
 		if constexpr (bIsById)
@@ -269,7 +269,7 @@ namespace UE::Editor::DataStorage
 	}
 
 	template<SortCase Casing, typename T>
-	constexpr FSortPrefixResult TSortTypeInfo<TSortStringView<Casing, T>>::CalculatePrefix(
+	constexpr FPrefixInfo TSortTypeInfo<TSortStringView<Casing, T>>::CalculatePrefix(
 		int32 CurrentIndex, int32 ByteIndex, TSortStringView<Casing, T> Value)
 	{
 		constexpr bool bIsCaseSensitive = TSortStringView<Casing, T>::bIsCaseSensitive;
@@ -326,7 +326,7 @@ namespace UE::Editor::DataStorage
 			bHasRemainingCharacters = StringIndex + Min < Value.View.Len();
 		}
 
-		return FSortPrefixResult
+		return FPrefixInfo
 		{
 			.Prefix = Result,
 			.bHasRemainingBytes = bHasRemainingCharacters
@@ -334,9 +334,9 @@ namespace UE::Editor::DataStorage
 	}
 
 	template<typename... ValueTypes>
-	FSortPrefixResult CreateSortPrefix(uint32 ByteIndex, ValueTypes&&... Values)
+	FPrefixInfo CreateSortPrefix(uint32 ByteIndex, ValueTypes&&... Values)
 	{
-		FSortPrefixResult Result
+		FPrefixInfo Result
 		{
 			.Prefix = 0,
 			.bHasRemainingBytes = true
diff --git a/Engine/Source/Runtime/TypedElementFramework/Tests/Elements/Framework/TypedElementDataStorageUiSorterTests.cpp b/Engine/Source/Runtime/TypedElementFramework/Tests/Elements/Framework/TypedElementDataStorageUiSorterTests.cpp
--- a/Engine/Source/Runtime/TypedElementFramework/Tests/Elements/Framework/TypedElementDataStorageUiSorterTests.cpp
+++ b/Engine/Source/Runtime/TypedElementFramework/Tests/Elements/Framework/TypedElementDataStorageUiSorterTests.cpp
@@ -30,7 +30,7 @@ namespace UE::Editor::DataStorage::Tests
 			}
 		}
 
-		virtual FSortPrefixResult SortPrefix(const ICoreProvider& Storage, RowHandle Row, uint32 ByteIndex) const override
+		virtual FPrefixInfo CalculatePrefix(const ICoreProvider& Storage, RowHandle Row, uint32 ByteIndex) const override
 		{
 			return CreateSortPrefix(ByteIndex, Values[static_cast<int32>(Row)]);
 		}
@@ -59,7 +59,7 @@ namespace UE::Editor::DataStorage::Tests
 			{
 				Data.Add(FData
 					{
-						.Prefix = SortPrefix(Storage, Index, ByteIndex).Prefix,
+						.Prefix = CalculatePrefix(Storage, Index, ByteIndex).Prefix,
 						.Value = MoveTemp(Values[Index])
 					});
 			}
@@ -106,7 +106,7 @@ namespace UE::Editor::DataStorage::Tests
 			}
 		}
 
-		virtual FSortPrefixResult SortPrefix(const ICoreProvider& Storage, RowHandle Row, uint32 ByteIndex) const override
+		virtual FPrefixInfo CalculatePrefix(const ICoreProvider& Storage, RowHandle Row, uint32 ByteIndex) const override
 		{
 			return Values[static_cast<int32>(Row)].ApplyAfter([this, ByteIndex](auto&&... TupleValues)
 				{
@@ -314,8 +314,8 @@ namespace UE::Editor::DataStorage::Tests
 						static_cast<uint32>(0x11223344)
 					});
 
-				FSortPrefixResult Result;
-				Result = Sorter.SortPrefix(*DataStorage, 0, 0);
+				FPrefixInfo Result;
+				Result = Sorter.CalculatePrefix(*DataStorage, 0, 0);
 				CHECK(Result.Prefix == 0x1122334400000000);
 				CHECK(Result.bHasRemainingBytes == false);
 			}
@@ -328,8 +328,8 @@ namespace UE::Editor::DataStorage::Tests
 						static_cast<uint8>(0x55)
 					});
 
-				FSortPrefixResult Result;
-				Result = Sorter.SortPrefix(*DataStorage, 0, 0);
+				FPrefixInfo Result;
+				Result = Sorter.CalculatePrefix(*DataStorage, 0, 0);
 				CHECK(Result.Prefix == 0x1122334455000000);
 				CHECK(Result.bHasRemainingBytes == false);
 			}
@@ -343,8 +343,8 @@ namespace UE::Editor::DataStorage::Tests
 						static_cast<uint16>(0x6677)
 					});
 
-				FSortPrefixResult Result;
-				Result = Sorter.SortPrefix(*DataStorage, 0, 0);
+				FPrefixInfo Result;
+				Result = Sorter.CalculatePrefix(*DataStorage, 0, 0);
 				CHECK(Result.Prefix == 0x1122334455667700);
 				CHECK(Result.bHasRemainingBytes == false);
 			}
@@ -359,12 +359,12 @@ namespace UE::Editor::DataStorage::Tests
 						static_cast<uint32>(0x8899aabb)
 					});
 
-				FSortPrefixResult Result;
-				Result = Sorter.SortPrefix(*DataStorage, 0, 0);
+				FPrefixInfo Result;
+				Result = Sorter.CalculatePrefix(*DataStorage, 0, 0);
 				CHECK(Result.Prefix == 0x1122334455667788);
 				CHECK(Result.bHasRemainingBytes == true);
 
-				Result = Sorter.SortPrefix(*DataStorage, 0, 8);
+				Result = Sorter.CalculatePrefix(*DataStorage, 0, 8);
 				CHECK(Result.Prefix == 0x99aabb0000000000);
 				CHECK(Result.bHasRemainingBytes == false);
 			}
@@ -380,16 +380,16 @@ namespace UE::Editor::DataStorage::Tests
 						static_cast<uint64>(0xccddeeff11223344)
 					});
 
-				FSortPrefixResult Result;
-				Result = Sorter.SortPrefix(*DataStorage, 0, 0);
+				FPrefixInfo Result;
+				Result = Sorter.CalculatePrefix(*DataStorage, 0, 0);
 				CHECK(Result.Prefix == 0x1122334455667788);
 				CHECK(Result.bHasRemainingBytes == true);
 
-				Result = Sorter.SortPrefix(*DataStorage, 0, 8);
+				Result = Sorter.CalculatePrefix(*DataStorage, 0, 8);
 				CHECK(Result.Prefix == 0x99aabbccddeeff11);
 				CHECK(Result.bHasRemainingBytes == true);
 
-				Result = Sorter.SortPrefix(*DataStorage, 0, 16);
+				Result = Sorter.CalculatePrefix(*DataStorage, 0, 16);
 				CHECK(Result.Prefix == 0x2233440000000000);
 				CHECK(Result.bHasRemainingBytes == false);
 			}
@@ -406,16 +406,16 @@ namespace UE::Editor::DataStorage::Tests
 						static_cast<uint16>(0x5566)
 					});
 
-				FSortPrefixResult Result;
-				Result = Sorter.SortPrefix(*DataStorage, 0, 0);
+				FPrefixInfo Result;
+				Result = Sorter.CalculatePrefix(*DataStorage, 0, 0);
 				CHECK(Result.Prefix == 0x1122334455667788);
 				CHECK(Result.bHasRemainingBytes == true);
 
-				Result = Sorter.SortPrefix(*DataStorage, 0, 8);
+				Result = Sorter.CalculatePrefix(*DataStorage, 0, 8);
 				CHECK(Result.Prefix == 0x99aabbccddeeff11);
 				CHECK(Result.bHasRemainingBytes == true);
 
-				Result = Sorter.SortPrefix(*DataStorage, 0, 16);
+				Result = Sorter.CalculatePrefix(*DataStorage, 0, 16);
 				CHECK(Result.Prefix == 0x2233445566000000);
 				CHECK(Result.bHasRemainingBytes == false);
 			}
@@ -433,16 +433,16 @@ namespace UE::Editor::DataStorage::Tests
 						static_cast<uint8>(0x77)
 					});
 
-				FSortPrefixResult Result;
-				Result = Sorter.SortPrefix(*DataStorage, 0, 0);
+				FPrefixInfo Result;
+				Result = Sorter.CalculatePrefix(*DataStorage, 0, 0);
 				CHECK(Result.Prefix == 0x1122334455667788);
 				CHECK(Result.bHasRemainingBytes == true);
 
-				Result = Sorter.SortPrefix(*DataStorage, 0, 8);
+				Result = Sorter.CalculatePrefix(*DataStorage, 0, 8);
 				CHECK(Result.Prefix == 0x99aabbccddeeff11);
 				CHECK(Result.bHasRemainingBytes == true);
 
-				Result = Sorter.SortPrefix(*DataStorage, 0, 16);
+				Result = Sorter.CalculatePrefix(*DataStorage, 0, 16);
 				CHECK(Result.Prefix == 0x2233445566770000);
 				CHECK(Result.bHasRemainingBytes == false);
 			}
@@ -461,16 +461,16 @@ namespace UE::Editor::DataStorage::Tests
 						static_cast<uint16>(0x8899)
 					});
 
-				FSortPrefixResult Result;
-				Result = Sorter.SortPrefix(*DataStorage, 0, 0);
+				FPrefixInfo Result;
+				Result = Sorter.CalculatePrefix(*DataStorage, 0, 0);
 				CHECK(Result.Prefix == 0x1122334455667788);
 				CHECK(Result.bHasRemainingBytes == true);
 
-				Result = Sorter.SortPrefix(*DataStorage, 0, 8);
+				Result = Sorter.CalculatePrefix(*DataStorage, 0, 8);
 				CHECK(Result.Prefix == 0x99aabbccddeeff11);
 				CHECK(Result.bHasRemainingBytes == true);
 
-				Result = Sorter.SortPrefix(*DataStorage, 0, 16);
+				Result = Sorter.CalculatePrefix(*DataStorage, 0, 16);
 				CHECK(Result.Prefix == 0x2233445566778899);
 				CHECK(Result.bHasRemainingBytes == false);
 			}
@@ -484,11 +484,11 @@ namespace UE::Editor::DataStorage::Tests
 		{
 			using SortStringView = TSortStringView<FSortCaseSensitive, FWideStringView>;
 			FSingleValueTestSorter<SortStringView> Sorter;
-			FSortPrefixResult Result;
+			FPrefixInfo Result;
 			{
 				FWideString String = TEXT("a");
 				Sorter.Values.Add(SortStringView(String));
-				Result = Sorter.SortPrefix(*DataStorage, 0, 0);
+				Result = Sorter.CalculatePrefix(*DataStorage, 0, 0);
 				CHECK(Result.Prefix == 0x0061000000000000);
 				CHECK(Result.bHasRemainingBytes == false);
 				Sorter.Values.Empty();
@@ -496,7 +496,7 @@ namespace UE::Editor::DataStorage::Tests
 			{
 				FWideString String = TEXT("aa");
 				Sorter.Values.Add(SortStringView(String));
-				Result = Sorter.SortPrefix(*DataStorage, 0, 0);
+				Result = Sorter.CalculatePrefix(*DataStorage, 0, 0);
 				CHECK(Result.Prefix == 0x0061006100000000);
 				CHECK(Result.bHasRemainingBytes == false);
 				Sorter.Values.Empty();
@@ -504,7 +504,7 @@ namespace UE::Editor::DataStorage::Tests
 			{
 				FWideString String = TEXT("aaa");
 				Sorter.Values.Add(SortStringView(String));
-				Result = Sorter.SortPrefix(*DataStorage, 0, 0);
+				Result = Sorter.CalculatePrefix(*DataStorage, 0, 0);
 				CHECK(Result.Prefix == 0x0061006100610000);
 				CHECK(Result.bHasRemainingBytes == false);
 				Sorter.Values.Empty();
@@ -512,7 +512,7 @@ namespace UE::Editor::DataStorage::Tests
 			{
 				FWideString String = TEXT("aaaa");
 				Sorter.Values.Add(SortStringView(String));
-				Result = Sorter.SortPrefix(*DataStorage, 0, 0);
+				Result = Sorter.CalculatePrefix(*DataStorage, 0, 0);
 				CHECK(Result.Prefix == 0x0061006100610061);
 				CHECK(Result.bHasRemainingBytes == false);
 				Sorter.Values.Empty();
@@ -520,10 +520,10 @@ namespace UE::Editor::DataStorage::Tests
 			{
 				FWideString String = TEXT("aaaab");
 				Sorter.Values.Add(SortStringView(String));
-				Result = Sorter.SortPrefix(*DataStorage, 0, 0);
+				Result = Sorter.CalculatePrefix(*DataStorage, 0, 0);
 				CHECK(Result.Prefix == 0x0061006100610061);
 				CHECK(Result.bHasRemainingBytes == true);
-				Result = Sorter.SortPrefix(*DataStorage, 0, 8);
+				Result = Sorter.CalculatePrefix(*DataStorage, 0, 8);
 				CHECK(Result.Prefix == 0x0062000000000000);
 				CHECK(Result.bHasRemainingBytes == false);
 				Sorter.Values.Empty();
@@ -531,10 +531,10 @@ namespace UE::Editor::DataStorage::Tests
 			{
 				FWideString String = TEXT("aaaabb");
 				Sorter.Values.Add(SortStringView(String));
-				Result = Sorter.SortPrefix(*DataStorage, 0, 0);
+				Result = Sorter.CalculatePrefix(*DataStorage, 0, 0);
 				CHECK(Result.Prefix == 0x0061006100610061);
 				CHECK(Result.bHasRemainingBytes == true);
-				Result = Sorter.SortPrefix(*DataStorage, 0, 8);
+				Result = Sorter.CalculatePrefix(*DataStorage, 0, 8);
 				CHECK(Result.Prefix == 0x0062006200000000);
 				CHECK(Result.bHasRemainingBytes == false);
 				Sorter.Values.Empty();
@@ -549,11 +549,11 @@ namespace UE::Editor::DataStorage::Tests
 		{
 			using SortStringView = TSortStringView<FSortCaseSensitive, FWideStringView>;
 			FMultiValueTestSorter<uint16, uint32, SortStringView> Sorter;
-			FSortPrefixResult Result;
+			FPrefixInfo Result;
 			{
 				FWideString String = TEXT("a");
 				Sorter.Values.Add({ static_cast<uint16>(0xaaaa), 0xbbbbcccc, SortStringView(String) });
-				Result = Sorter.SortPrefix(*DataStorage, 0, 0);
+				Result = Sorter.CalculatePrefix(*DataStorage, 0, 0);
 				CHECK(Result.Prefix == 0xaaaabbbbcccc0061);
 				CHECK(Result.bHasRemainingBytes == false);
 				Sorter.Values.Empty();
@@ -561,10 +561,10 @@ namespace UE::Editor::DataStorage::Tests
 			{
 				FWideString String = TEXT("ab");
 				Sorter.Values.Add({ static_cast<uint16>(0xaaaa), 0xbbbbcccc, SortStringView(String) });
-				Result = Sorter.SortPrefix(*DataStorage, 0, 0);
+				Result = Sorter.CalculatePrefix(*DataStorage, 0, 0);
 				CHECK(Result.Prefix == 0xaaaabbbbcccc0061);
 				CHECK(Result.bHasRemainingBytes == true);
-				Result = Sorter.SortPrefix(*DataStorage, 0, 8);
+				Result = Sorter.CalculatePrefix(*DataStorage, 0, 8);
 				CHECK(Result.Prefix == 0x0062000000000000);
 				CHECK(Result.bHasRemainingBytes == false);
 				Sorter.Values.Empty();
@@ -572,10 +572,10 @@ namespace UE::Editor::DataStorage::Tests
 			{
 				FWideString String = TEXT("abc");
 				Sorter.Values.Add({ static_cast<uint16>(0xaaaa), 0xbbbbcccc, SortStringView(String) });
-				Result = Sorter.SortPrefix(*DataStorage, 0, 0);
+				Result = Sorter.CalculatePrefix(*DataStorage, 0, 0);
 				CHECK(Result.Prefix == 0xaaaabbbbcccc0061);
 				CHECK(Result.bHasRemainingBytes == true);
-				Result = Sorter.SortPrefix(*DataStorage, 0, 8);
+				Result = Sorter.CalculatePrefix(*DataStorage, 0, 8);
 				CHECK(Result.Prefix == 0x0062006300000000);
 				CHECK(Result.bHasRemainingBytes == false);
 				Sorter.Values.Empty();
@@ -583,10 +583,10 @@ namespace UE::Editor::DataStorage::Tests
 			{
 				FWideString String = TEXT("abcd");
 				Sorter.Values.Add({ static_cast<uint16>(0xaaaa), 0xbbbbcccc, SortStringView(String) });
-				Result = Sorter.SortPrefix(*DataStorage, 0, 0);
+				Result = Sorter.CalculatePrefix(*DataStorage, 0, 0);
 				CHECK(Result.Prefix == 0xaaaabbbbcccc0061);
 				CHECK(Result.bHasRemainingBytes == true);
-				Result = Sorter.SortPrefix(*DataStorage, 0, 8);
+				Result = Sorter.CalculatePrefix(*DataStorage, 0, 8);
 				CHECK(Result.Prefix == 0x0062006300640000);
 				CHECK(Result.bHasRemainingBytes == false);
 				Sorter.Values.Empty();
@@ -594,10 +594,10 @@ namespace UE::Editor::DataStorage::Tests
 			{
 				FWideString String = TEXT("abcde");
 				Sorter.Values.Add({ static_cast<uint16>(0xaaaa), 0xbbbbcccc, SortStringView(String) });
-				Result = Sorter.SortPrefix(*DataStorage, 0, 0);
+				Result = Sorter.CalculatePrefix(*DataStorage, 0, 0);
 				CHECK(Result.Prefix == 0xaaaabbbbcccc0061);
 				CHECK(Result.bHasRemainingBytes == true);
-				Result = Sorter.SortPrefix(*DataStorage, 0, 8);
+				Result = Sorter.CalculatePrefix(*DataStorage, 0, 8);
 				CHECK(Result.Prefix == 0x0062006300640065);
 				CHECK(Result.bHasRemainingBytes == false);
 				Sorter.Values.Empty();
@@ -605,13 +605,13 @@ namespace UE::Editor::DataStorage::Tests
 			{
 				FWideString String = TEXT("abcdef");
 				Sorter.Values.Add({ static_cast<uint16>(0xaaaa), 0xbbbbcccc, SortStringView(String) });
-				Result = Sorter.SortPrefix(*DataStorage, 0, 0);
+				Result = Sorter.CalculatePrefix(*DataStorage, 0, 0);
 				CHECK(Result.Prefix == 0xaaaabbbbcccc0061);
 				CHECK(Result.bHasRemainingBytes == true);
-				Result = Sorter.SortPrefix(*DataStorage, 0, 8);
+				Result = Sorter.CalculatePrefix(*DataStorage, 0, 8);
 				CHECK(Result.Prefix == 0x0062006300640065);
 				CHECK(Result.bHasRemainingBytes == true);
-				Result = Sorter.SortPrefix(*DataStorage, 0, 16);
+				Result = Sorter.CalculatePrefix(*DataStorage, 0, 16);
 				CHECK(Result.Prefix == 0x0066000000000000);
 				CHECK(Result.bHasRemainingBytes == false);
 				Sorter.Values.Empty();
