[TEDS] Added custom sorting support to a few TEDS UI columns.

[RN] Added custom sorting support to a few TEDS UI columns.

diff --git a/Engine/Plugins/Experimental/EditorDataStorage/Source/TedsUI/Private/Widgets/NameWidget.cpp b/Engine/Plugins/Experimental/EditorDataStorage/Source/TedsUI/Private/Widgets/NameWidget.cpp
--- a/Engine/Plugins/Experimental/EditorDataStorage/Source/TedsUI/Private/Widgets/NameWidget.cpp
+++ b/Engine/Plugins/Experimental/EditorDataStorage/Source/TedsUI/Private/Widgets/NameWidget.cpp
@@ -4,10 +4,63 @@
 
 #include "Elements/Framework/TypedElementAttributeBinding.h"
 #include "Elements/Columns/TypedElementMiscColumns.h"
+#include "Internationalization/Internationalization.h"
 #include "Widgets/Text/STextBlock.h"
 
 #define LOCTEXT_NAMESPACE "NameWidget"
 
+namespace UE::Editor::DataStorage
+{
+	//
+	// No slash
+	//
+
+	FText FNameWidgetSorter_NoSlash::GetShortName() const
+	{
+		return LOCTEXT("FNameWidgetSorter_NoSlashName", "No slash");
+	}
+	
+	int32 FNameWidgetSorter_NoSlash::Compare(const FNameColumn& Left, const FNameColumn& Right) const
+	{
+		FString LeftString = Left.Name.ToString();
+		LeftString.RemoveFromStart(TEXT("/"));
+		
+		FString RightString = Right.Name.ToString();
+		RightString.RemoveFromStart(TEXT("/"));
+		
+		return LeftString.Compare(RightString, ESearchCase::IgnoreCase);
+	}
+
+	FPrefixInfo FNameWidgetSorter_NoSlash::CalculatePrefix(const FNameColumn& Column, uint32 ByteIndex) const
+	{
+		FString String = Column.Name.ToString();
+		String.RemoveFromStart(TEXT("/"));
+		return CreateSortPrefix(ByteIndex, TSortStringView(FSortCaseInsensitive{}, String));
+	}
+
+
+
+
+	//
+	// With slash
+	//
+
+	FText FNameWidgetSorter_WithSlash::GetShortName() const
+	{
+		return LOCTEXT("FNameWidgetSorter_WithSlashName", "With slash");
+	}
+	
+	int32 FNameWidgetSorter_WithSlash::Compare(const FNameColumn& Left, const FNameColumn& Right) const
+	{
+		return Left.Name.Compare(Right.Name);
+	}
+
+	FPrefixInfo FNameWidgetSorter_WithSlash::CalculatePrefix(const FNameColumn& Column, uint32 ByteIndex) const
+	{
+		return CreateSortPrefix(ByteIndex, TSortNameView(FSortByName{}, Column.Name));
+	}
+} // namespace UE::Editor::DataStorage
+
 void UNameWidgetFactory::RegisterWidgetConstructors(UE::Editor::DataStorage::ICoreProvider& DataStorage,
 	UE::Editor::DataStorage::IUiProvider& DataStorageUi) const
 {
@@ -33,4 +86,16 @@ TSharedPtr<SWidget> FNameWidgetConstructor::CreateWidget(UE::Editor::DataStorage
 				.ToolTipText(Binder.BindText(&FNameColumn::Name));
 }
 
+TArray<TSharedPtr<const UE::Editor::DataStorage::FColumnSorterInterface>> FNameWidgetConstructor::ConstructColumnSorters(
+	UE::Editor::DataStorage::ICoreProvider* DataStorage,
+	const UE::Editor::DataStorage::FMetaDataView& Arguments)
+{
+	using namespace UE::Editor::DataStorage;
+	return TArray<TSharedPtr<const FColumnSorterInterface>>(
+		{
+			MakeShared<FNameWidgetSorter_NoSlash>(),
+			MakeShared<FNameWidgetSorter_WithSlash>()
+		});
+}
+
 #undef LOCTEXT_NAMESPACE
diff --git a/Engine/Plugins/Experimental/EditorDataStorage/Source/TedsUI/Public/Widgets/NameWidget.h b/Engine/Plugins/Experimental/EditorDataStorage/Source/TedsUI/Public/Widgets/NameWidget.h
--- a/Engine/Plugins/Experimental/EditorDataStorage/Source/TedsUI/Public/Widgets/NameWidget.h
+++ b/Engine/Plugins/Experimental/EditorDataStorage/Source/TedsUI/Public/Widgets/NameWidget.h
@@ -2,6 +2,7 @@
 
 #pragma once
 
+#include "Elements/Columns/TypedElementMiscColumns.h"
 #include "Elements/Interfaces/TypedElementDataStorageFactory.h"
 #include "Elements/Interfaces/TypedElementDataStorageUiInterface.h"
 #include "Internationalization/Text.h"
@@ -12,6 +13,29 @@
 class IEditorDataStorageProvider;
 class UScriptStruct;
 
+namespace UE::Editor::DataStorage
+{
+	class FNameWidgetSorter_NoSlash final : public TColumnSorterInterface<FColumnSorterInterface::ESortType::HybridSort, FNameColumn>
+	{
+	public:
+		virtual FText GetShortName() const override;
+
+	protected:
+		virtual int32 Compare(const FNameColumn& Left, const FNameColumn& Right) const override;
+		virtual FPrefixInfo CalculatePrefix(const FNameColumn& Column, uint32 ByteIndex) const override;
+	};
+
+	class FNameWidgetSorter_WithSlash final : public TColumnSorterInterface<FColumnSorterInterface::ESortType::HybridSort, FNameColumn>
+	{
+	public:
+		virtual FText GetShortName() const override;
+
+	protected:
+		virtual int32 Compare(const FNameColumn& Left, const FNameColumn& Right) const override;
+		virtual FPrefixInfo CalculatePrefix(const FNameColumn& Column, uint32 ByteIndex) const override;
+	};
+}
+
 UCLASS()
 class UNameWidgetFactory : public UEditorDataStorageFactory
 {
@@ -39,4 +63,8 @@ public:
 		UE::Editor::DataStorage::RowHandle TargetRow,
 		UE::Editor::DataStorage::RowHandle WidgetRow, 
 		const UE::Editor::DataStorage::FMetaDataView& Arguments) override;
+
+	virtual TArray<TSharedPtr<const UE::Editor::DataStorage::FColumnSorterInterface>> ConstructColumnSorters(
+		UE::Editor::DataStorage::ICoreProvider* DataStorage,
+		const UE::Editor::DataStorage::FMetaDataView& Arguments) override;
 };
diff --git a/Engine/Plugins/Experimental/EditorDataStorageFeatures/Source/TedsAssetData/Private/Widgets/AssetDataVirtualPathWidget.cpp b/Engine/Plugins/Experimental/EditorDataStorageFeatures/Source/TedsAssetData/Private/Widgets/AssetDataVirtualPathWidget.cpp
--- a/Engine/Plugins/Experimental/EditorDataStorageFeatures/Source/TedsAssetData/Private/Widgets/AssetDataVirtualPathWidget.cpp
+++ b/Engine/Plugins/Experimental/EditorDataStorageFeatures/Source/TedsAssetData/Private/Widgets/AssetDataVirtualPathWidget.cpp
@@ -10,6 +10,66 @@
 #include "Widgets/Layout/SBox.h"
 #include "Widgets/Text/STextBlock.h"
 
+
+namespace UE::Editor::DataStorage
+{
+	int32 FAssetDataVirtualPathWidgetSorter::Compare(const ICoreProvider& Storage, RowHandle Left, RowHandle Right) const
+	{
+		FString LeftPath;
+		bool bLeftPathFound = ExtractString(LeftPath, Storage, Left);
+		FString RightPath;
+		bool bRightPathFound = ExtractString(RightPath, Storage, Right);
+
+		return (bLeftPathFound && bRightPathFound) 
+			? LeftPath.Compare(RightPath, ESearchCase::IgnoreCase)
+			: (static_cast<int32>(bLeftPathFound) - static_cast<int32>(bRightPathFound));
+	}
+
+	FPrefixInfo FAssetDataVirtualPathWidgetSorter::CalculatePrefix(
+		const ICoreProvider& Storage, RowHandle Row, uint32 ByteIndex) const
+	{
+		FString Path;
+		return ExtractString(Path, Storage, Row) 
+			? CreateSortPrefix(ByteIndex, TSortStringView(FSortCaseInsensitive{}, Path))
+			: FPrefixInfo{};
+	}
+
+	bool FAssetDataVirtualPathWidgetSorter::ExtractString(FString& Result, const ICoreProvider& Storage, RowHandle Row) const
+	{
+		bool bIsAsset = Storage.HasColumns<FAssetTag>(Row);
+		if (bIsAsset)
+		{
+			if (const FVirtualPathColumn_Experimental* VirtualPathColumn = Storage.GetColumn<FVirtualPathColumn_Experimental>(Row))
+			{
+				Result = VirtualPathColumn->VirtualPath.ToString();
+				Result = TedsAssetDataHelper::RemoveSlashFromStart(Result);
+				return true;
+			}
+		}
+		else
+		{
+			if (const FAssetPathColumn_Experimental* AssetPathColumn = Storage.GetColumn<FAssetPathColumn_Experimental>(Row))
+			{
+				Result = AssetPathColumn->Path.ToString();
+				Result = TedsAssetDataHelper::RemoveSlashFromStart(Result);
+				Result = TedsAssetDataHelper::RemoveAllFromLastSlash(Result);
+				return true;
+			}
+		}
+		return false;
+	}
+
+	FColumnSorterInterface::ESortType FAssetDataVirtualPathWidgetSorter::GetSortType() const
+	{
+		return ESortType::HybridSort;
+	}
+
+	FText FAssetDataVirtualPathWidgetSorter::GetShortName() const
+	{
+		return FText::GetEmpty();
+	}
+}
+
 void UAssetDataVirtualPathWidgetFactory::RegisterWidgetConstructors(UE::Editor::DataStorage::ICoreProvider& DataStorage,
 	UE::Editor::DataStorage::IUiProvider& DataStorageUi) const
 {
@@ -60,3 +120,14 @@ TSharedPtr<SWidget> FAssetDataVirtualPathWidgetConstructor::CreateWidget(UE::Edi
 				.Text(PathText)
 			];
 }
+
+TArray<TSharedPtr<const UE::Editor::DataStorage::FColumnSorterInterface>> FAssetDataVirtualPathWidgetConstructor::ConstructColumnSorters(
+	UE::Editor::DataStorage::ICoreProvider* DataStorage,
+	const UE::Editor::DataStorage::FMetaDataView& Arguments)
+{
+	using namespace UE::Editor::DataStorage;
+	return TArray<TSharedPtr<const FColumnSorterInterface>>(
+		{
+			MakeShared<FAssetDataVirtualPathWidgetSorter>()
+		});
+}
diff --git a/Engine/Plugins/Experimental/EditorDataStorageFeatures/Source/TedsAssetData/Private/Widgets/AssetDataVirtualPathWidget.h b/Engine/Plugins/Experimental/EditorDataStorageFeatures/Source/TedsAssetData/Private/Widgets/AssetDataVirtualPathWidget.h
--- a/Engine/Plugins/Experimental/EditorDataStorageFeatures/Source/TedsAssetData/Private/Widgets/AssetDataVirtualPathWidget.h
+++ b/Engine/Plugins/Experimental/EditorDataStorageFeatures/Source/TedsAssetData/Private/Widgets/AssetDataVirtualPathWidget.h
@@ -7,6 +7,21 @@
 
 #include "AssetDataVirtualPathWidget.generated.h"
 
+namespace UE::Editor::DataStorage
+{
+	class FAssetDataVirtualPathWidgetSorter final : public FColumnSorterInterface
+	{
+	public:
+		virtual int32 Compare(const ICoreProvider& Storage, RowHandle Left, RowHandle Right) const override;
+		virtual FPrefixInfo CalculatePrefix(const ICoreProvider& Storage, RowHandle Row, uint32 ByteIndex) const override;
+		virtual ESortType GetSortType() const override;
+		virtual FText GetShortName() const override;
+
+	private:
+		bool ExtractString(FString& Result, const ICoreProvider& Storage, RowHandle Row) const;
+	};
+}
+
 UCLASS()
 class UAssetDataVirtualPathWidgetFactory : public UEditorDataStorageFactory
 {
@@ -15,7 +30,7 @@ class UAssetDataVirtualPathWidgetFactory : public UEditorDataStorageFactory
 public:
 	~UAssetDataVirtualPathWidgetFactory() override = default;
 
-	TEDSASSETDATA_API void RegisterWidgetConstructors(UE::Editor::DataStorage::ICoreProvider& DataStorage,
+	void RegisterWidgetConstructors(UE::Editor::DataStorage::ICoreProvider& DataStorage,
 		UE::Editor::DataStorage::IUiProvider& DataStorageUi) const override;
 };
 
@@ -26,13 +41,17 @@ struct FAssetDataVirtualPathWidgetConstructor : public FSimpleWidgetConstructor
 	GENERATED_BODY()
 
 public:
-	TEDSASSETDATA_API FAssetDataVirtualPathWidgetConstructor();
+	FAssetDataVirtualPathWidgetConstructor();
 	~FAssetDataVirtualPathWidgetConstructor() override = default;
 
-	TEDSASSETDATA_API virtual TSharedPtr<SWidget> CreateWidget(
+	virtual TSharedPtr<SWidget> CreateWidget(
 		UE::Editor::DataStorage::ICoreProvider* DataStorage,
 		UE::Editor::DataStorage::IUiProvider* DataStorageUi,
 		UE::Editor::DataStorage::RowHandle TargetRow,
 		UE::Editor::DataStorage::RowHandle WidgetRow,
 		const UE::Editor::DataStorage::FMetaDataView& Arguments) override;
+
+	virtual TArray<TSharedPtr<const UE::Editor::DataStorage::FColumnSorterInterface>> ConstructColumnSorters(
+		UE::Editor::DataStorage::ICoreProvider* DataStorage,
+		const UE::Editor::DataStorage::FMetaDataView& Arguments) override;
 };
diff --git a/Engine/Plugins/Experimental/EditorDataStorageFeatures/Source/TedsTableViewer/Private/Widgets/RowHandleWidget.cpp b/Engine/Plugins/Experimental/EditorDataStorageFeatures/Source/TedsTableViewer/Private/Widgets/RowHandleWidget.cpp
--- a/Engine/Plugins/Experimental/EditorDataStorageFeatures/Source/TedsTableViewer/Private/Widgets/RowHandleWidget.cpp
+++ b/Engine/Plugins/Experimental/EditorDataStorageFeatures/Source/TedsTableViewer/Private/Widgets/RowHandleWidget.cpp
@@ -8,6 +8,32 @@
 
 #define LOCTEXT_NAMESPACE "RowHandleWidget"
 
+namespace UE::Editor::DataStorage
+{
+	//
+	// FRowHandleSorter
+	//
+
+	int32 FRowHandleSorter::Compare(const ICoreProvider& Storage, RowHandle Left, RowHandle Right) const
+	{
+		return (Left > Right) - (Left < Right);
+	}
+
+	FPrefixInfo FRowHandleSorter::CalculatePrefix(const ICoreProvider& Storage, RowHandle Row, uint32 ByteIndex) const
+	{
+		return CreateSortPrefix(ByteIndex, Row);
+	}
+
+	FColumnSorterInterface::ESortType FRowHandleSorter::GetSortType() const
+	{
+		return ESortType::FixedSize64;
+	}
+
+	FText FRowHandleSorter::GetShortName() const
+	{
+		return FText::GetEmpty();
+	}
+}
 
 void URowHandleWidgetFactory::RegisterWidgetConstructors(UE::Editor::DataStorage::ICoreProvider& DataStorage, UE::Editor::DataStorage::IUiProvider& DataStorageUi) const
 {
@@ -81,4 +107,15 @@ bool FRowHandleWidgetConstructor::FinalizeWidget(UE::Editor::DataStorage::ICoreP
 	return true;
 }
 
+TArray<TSharedPtr<const UE::Editor::DataStorage::FColumnSorterInterface>> FRowHandleWidgetConstructor::ConstructColumnSorters(
+	UE::Editor::DataStorage::ICoreProvider* DataStorage,
+	const UE::Editor::DataStorage::FMetaDataView& Arguments)
+{
+	using namespace UE::Editor::DataStorage;
+	return TArray<TSharedPtr<const FColumnSorterInterface>>(
+		{
+			MakeShared<FRowHandleSorter>()
+		});
+}
+
 #undef LOCTEXT_NAMESPACE //"RowHandleWidget"
diff --git a/Engine/Plugins/Experimental/EditorDataStorageFeatures/Source/TedsTableViewer/Public/Widgets/RowHandleWidget.h b/Engine/Plugins/Experimental/EditorDataStorageFeatures/Source/TedsTableViewer/Public/Widgets/RowHandleWidget.h
--- a/Engine/Plugins/Experimental/EditorDataStorageFeatures/Source/TedsTableViewer/Public/Widgets/RowHandleWidget.h
+++ b/Engine/Plugins/Experimental/EditorDataStorageFeatures/Source/TedsTableViewer/Public/Widgets/RowHandleWidget.h
@@ -15,6 +15,18 @@ namespace UE::Editor::DataStorage
 
 class UScriptStruct;
 
+namespace UE::Editor::DataStorage
+{
+	class FRowHandleSorter final : public FColumnSorterInterface
+	{
+	public:
+		virtual int32 Compare(const ICoreProvider& Storage, RowHandle Left, RowHandle Right) const override;
+		virtual FPrefixInfo CalculatePrefix(const ICoreProvider& Storage, RowHandle Row, uint32 ByteIndex) const override;
+		virtual ESortType GetSortType() const override;
+		virtual FText GetShortName() const override;
+	};
+}
+
 UCLASS()
 class TEDSTABLEVIEWER_API URowHandleWidgetFactory : public UEditorDataStorageFactory
 {
@@ -46,4 +58,8 @@ protected:
 		UE::Editor::DataStorage::IUiProvider* DataStorageUi,
 		UE::Editor::DataStorage::RowHandle Row,
 		const TSharedPtr<SWidget>& Widget) override;
+
+	virtual TArray<TSharedPtr<const UE::Editor::DataStorage::FColumnSorterInterface>> ConstructColumnSorters(
+		UE::Editor::DataStorage::ICoreProvider* DataStorage,
+		const UE::Editor::DataStorage::FMetaDataView& Arguments) override;
 };
diff --git a/Engine/Source/Runtime/TypedElementFramework/Public/Elements/Framework/TypedElementSorter.h b/Engine/Source/Runtime/TypedElementFramework/Public/Elements/Framework/TypedElementSorter.h
--- a/Engine/Source/Runtime/TypedElementFramework/Public/Elements/Framework/TypedElementSorter.h
+++ b/Engine/Source/Runtime/TypedElementFramework/Public/Elements/Framework/TypedElementSorter.h
@@ -61,7 +61,10 @@ namespace UE::Editor::DataStorage
 			 * increased peak memory usage and limits the number of variable sized variables to one, which also have to be the last 
 			 * variable in the chain.
 			 */
-			HybridSort
+			HybridSort,
+
+			/** Indicates the final value, cannot be used as input. */
+			Max
 		};
 
 		/** Returns the type of sorting required for this column. */
@@ -86,6 +89,69 @@ namespace UE::Editor::DataStorage
 		virtual FPrefixInfo CalculatePrefix(const ICoreProvider& Storage, RowHandle Row, uint32 ByteIndex) const = 0;
 	};
 
+	template<FColumnSorterInterface::ESortType SortType, typename ColumnType>
+	class TColumnSorterInterface : public FColumnSorterInterface
+	{
+		static_assert(static_cast<uint32>(SortType) < static_cast<uint32>(ESortType::Max), "Unexpected sort type for TColumnSorterInterface.");
+	};
+
+	template<typename ColumnType>
+	class TColumnSorterInterface<FColumnSorterInterface::ESortType::FixedSize64, ColumnType> : public FColumnSorterInterface
+	{
+	public:
+		virtual ~TColumnSorterInterface() override = default;
+
+		int32 Compare(const ICoreProvider& Storage, RowHandle Left, RowHandle Right) const override final;
+		FPrefixInfo CalculatePrefix(const ICoreProvider& Storage, RowHandle Row, uint32 ByteIndex) const override final;
+		ESortType GetSortType() const override final;
+
+	protected:
+		virtual FPrefixInfo CalculatePrefix(const ColumnType& Column, uint32 ByteIndex) const = 0;
+	};
+
+	template<typename ColumnType>
+	class TColumnSorterInterface<FColumnSorterInterface::ESortType::FixedSizeOnly, ColumnType> : public FColumnSorterInterface
+	{
+	public:
+		virtual ~TColumnSorterInterface() override = default;
+
+		int32 Compare(const ICoreProvider& Storage, RowHandle Left, RowHandle Right) const override final;
+		FPrefixInfo CalculatePrefix(const ICoreProvider& Storage, RowHandle Row, uint32 ByteIndex) const override final;
+		ESortType GetSortType() const override final;
+
+	protected:
+		virtual FPrefixInfo CalculatePrefix(const ColumnType& Column, uint32 ByteIndex) const = 0;
+	};
+
+	template<typename ColumnType>
+	class TColumnSorterInterface<FColumnSorterInterface::ESortType::ComparativeSort, ColumnType> : public FColumnSorterInterface
+	{
+	public:
+		virtual ~TColumnSorterInterface() override = default;
+
+		int32 Compare(const ICoreProvider& Storage, RowHandle Left, RowHandle Right) const override final;
+		FPrefixInfo CalculatePrefix(const ICoreProvider& Storage, RowHandle Row, uint32 ByteIndex) const override final;
+		ESortType GetSortType() const override final;
+
+	protected:
+		virtual int32 Compare(const ColumnType& Left, const ColumnType& Right) const = 0;
+	};
+
+	template<typename ColumnType>
+	class TColumnSorterInterface<FColumnSorterInterface::ESortType::HybridSort, ColumnType> : public FColumnSorterInterface
+	{
+	public:
+		virtual ~TColumnSorterInterface() override = default;
+
+		int32 Compare(const ICoreProvider& Storage, RowHandle Left, RowHandle Right) const override final;
+		FPrefixInfo CalculatePrefix(const ICoreProvider& Storage, RowHandle Row, uint32 ByteIndex) const override final;
+		ESortType GetSortType() const override final;
+
+	protected:
+		virtual int32 Compare(const ColumnType& Left, const ColumnType& Right) const = 0;
+		virtual FPrefixInfo CalculatePrefix(const ColumnType& Column, uint32 ByteIndex) const = 0;
+	};
+
 	namespace Private
 	{
 		template<typename T>
diff --git a/Engine/Source/Runtime/TypedElementFramework/Public/Elements/Framework/TypedElementSorter.inl b/Engine/Source/Runtime/TypedElementFramework/Public/Elements/Framework/TypedElementSorter.inl
--- a/Engine/Source/Runtime/TypedElementFramework/Public/Elements/Framework/TypedElementSorter.inl
+++ b/Engine/Source/Runtime/TypedElementFramework/Public/Elements/Framework/TypedElementSorter.inl
@@ -1,9 +1,124 @@
 // Copyright Epic Games, Inc. All Rights Reserved.
 
 #include "Misc/Char.h"
+#include "Elements/Interfaces/TypedElementDataStorageInterface.h"
 
 namespace UE::Editor::DataStorage
 {
+	//
+	// FixedSizeOnly Sort
+	//
+
+	template<typename ColumnType>
+	int32 TColumnSorterInterface<FColumnSorterInterface::ESortType::FixedSize64, ColumnType > ::Compare(
+		const ICoreProvider& Storage, RowHandle Left, RowHandle Right) const
+	{
+		return 0;
+	}
+
+	template<typename ColumnType>
+	FPrefixInfo TColumnSorterInterface<FColumnSorterInterface::ESortType::FixedSize64, ColumnType>::CalculatePrefix(
+		const ICoreProvider& Storage, RowHandle Row, uint32 ByteIndex) const
+	{
+		const ColumnType* Column = Storage.GetColumn<ColumnType>(Row);
+		return Column ? CalculatePrefix(*Column, ByteIndex) : FPrefixInfo{};
+	}
+
+	template<typename ColumnType>
+	FColumnSorterInterface::ESortType TColumnSorterInterface<FColumnSorterInterface::ESortType::FixedSize64, ColumnType>::GetSortType() const
+	{
+		return FColumnSorterInterface::ESortType::FixedSize64;
+	}
+
+
+
+	//
+	// FixedSizeOnly Sort
+	//
+
+	template<typename ColumnType>
+	int32 TColumnSorterInterface<FColumnSorterInterface::ESortType::FixedSizeOnly, ColumnType > ::Compare(
+		const ICoreProvider& Storage, RowHandle Left, RowHandle Right) const
+	{
+		return 0;
+	}
+
+	template<typename ColumnType>
+	FPrefixInfo TColumnSorterInterface<FColumnSorterInterface::ESortType::FixedSizeOnly, ColumnType>::CalculatePrefix(
+		const ICoreProvider& Storage, RowHandle Row, uint32 ByteIndex) const
+	{
+		const ColumnType* Column = Storage.GetColumn<ColumnType>(Row);
+		return Column ? CalculatePrefix(*Column, ByteIndex) : FPrefixInfo{};
+	}
+
+	template<typename ColumnType>
+	FColumnSorterInterface::ESortType TColumnSorterInterface<FColumnSorterInterface::ESortType::FixedSizeOnly, ColumnType>::GetSortType() const
+	{
+		return FColumnSorterInterface::ESortType::FixedSizeOnly;
+	}
+
+	
+	
+	//
+	// Comparative Sort
+	//
+
+	template<typename ColumnType>
+	int32 TColumnSorterInterface<FColumnSorterInterface::ESortType::ComparativeSort, ColumnType > ::Compare(
+		const ICoreProvider& Storage, RowHandle Left, RowHandle Right) const
+	{
+		const ColumnType* LeftColumn = Storage.GetColumn<ColumnType>(Left);
+		const ColumnType* RigthColumn = Storage.GetColumn<ColumnType>(Right);
+
+		return (LeftColumn && RigthColumn)
+			? Compare(*LeftColumn, *RigthColumn)
+			: ((LeftColumn == nullptr) - (RigthColumn == nullptr));
+	}
+
+	template<typename ColumnType>
+	FPrefixInfo TColumnSorterInterface<FColumnSorterInterface::ESortType::ComparativeSort, ColumnType>::CalculatePrefix(
+		const ICoreProvider& Storage, RowHandle Row, uint32 ByteIndex) const
+	{
+		return FPrefixInfo{};
+	}
+
+	template<typename ColumnType>
+	FColumnSorterInterface::ESortType TColumnSorterInterface<FColumnSorterInterface::ESortType::ComparativeSort, ColumnType>::GetSortType() const
+	{
+		return FColumnSorterInterface::ESortType::ComparativeSort;
+	}
+
+
+	//
+	// Hybrid Sort
+	//
+	
+	template<typename ColumnType>
+	int32 TColumnSorterInterface<FColumnSorterInterface::ESortType::HybridSort, ColumnType>::Compare(
+		const ICoreProvider& Storage, RowHandle Left, RowHandle Right) const
+	{
+		const ColumnType* LeftColumn = Storage.GetColumn<ColumnType>(Left);
+		const ColumnType* RigthColumn = Storage.GetColumn<ColumnType>(Right);
+
+		return (LeftColumn && RigthColumn)
+			? Compare(*LeftColumn, *RigthColumn)
+			: ((LeftColumn == nullptr) - (RigthColumn == nullptr));
+	}
+
+	template<typename ColumnType>
+	FPrefixInfo TColumnSorterInterface<FColumnSorterInterface::ESortType::HybridSort, ColumnType>::CalculatePrefix(
+		const ICoreProvider& Storage, RowHandle Row, uint32 ByteIndex) const
+	{
+		const ColumnType* Column = Storage.GetColumn<ColumnType>(Row);
+		return Column ? CalculatePrefix(*Column, ByteIndex) : FPrefixInfo{};
+	}
+
+	template<typename ColumnType>
+	FColumnSorterInterface::ESortType TColumnSorterInterface<FColumnSorterInterface::ESortType::HybridSort, ColumnType>::GetSortType() const
+	{
+		return FColumnSorterInterface::ESortType::HybridSort;
+	}
+
 	namespace Private
 	{
 		template<typename T>
