ICoreProvider : Add non-const FindFactory overloads

diff --git a/Engine/Plugins/Experimental/EditorDataStorage/Source/TedsCore/Private/TypedElementDatabase.cpp b/Engine/Plugins/Experimental/EditorDataStorage/Source/TedsCore/Private/TypedElementDatabase.cpp
--- a/Engine/Plugins/Experimental/EditorDataStorage/Source/TedsCore/Private/TypedElementDatabase.cpp
+++ b/Engine/Plugins/Experimental/EditorDataStorage/Source/TedsCore/Private/TypedElementDatabase.cpp
@@ -235,6 +235,11 @@ const UEditorDataStorageFactory* UEditorDataStorage::FindFactory(const UClass* F
 	return nullptr;
 }
 
+UEditorDataStorageFactory* UEditorDataStorage::FindFactory(const UClass* FactoryType)
+{
+	return const_cast<UEditorDataStorageFactory*>(static_cast<const UEditorDataStorage*>(this)->FindFactory(FactoryType));
+}
+
 void UEditorDataStorage::Deinitialize()
 {
 	checkf(Factories.IsEmpty(), TEXT("ResetFactories should have been called before deinitialized"));
diff --git a/Engine/Plugins/Experimental/EditorDataStorage/Source/TedsCore/Private/TypedElementDatabase.h b/Engine/Plugins/Experimental/EditorDataStorage/Source/TedsCore/Private/TypedElementDatabase.h
--- a/Engine/Plugins/Experimental/EditorDataStorage/Source/TedsCore/Private/TypedElementDatabase.h
+++ b/Engine/Plugins/Experimental/EditorDataStorage/Source/TedsCore/Private/TypedElementDatabase.h
@@ -69,9 +69,8 @@ public:
 
 	/** Returns factory instance given the type of factory */
 	virtual const UEditorDataStorageFactory* FindFactory(const UClass* FactoryType) const override;
-	/** Helper for FindFactory(const UClass*) */
-	template<typename FactoryTypeT>
-	const FactoryTypeT* FindFactory() const;
+
+	virtual UEditorDataStorageFactory* FindFactory(const UClass* FactoryType) override;
 	
 	void Deinitialize();
 
@@ -237,9 +236,3 @@ UEditorDataStorage::TFactoryIterator<FactoryType, DatabaseType>::operator bool()
 {
 	return Database != nullptr && Index < Database->Factories.Num();
 }
-
-template <typename FactoryTypeT>
-const FactoryTypeT* UEditorDataStorage::FindFactory() const
-{
-	return static_cast<const FactoryTypeT*>(FindFactory(FactoryTypeT::StaticClass()));
-}
diff --git a/Engine/Source/Runtime/TypedElementFramework/Public/Elements/Interfaces/TypedElementDataStorageInterface.h b/Engine/Source/Runtime/TypedElementFramework/Public/Elements/Interfaces/TypedElementDataStorageInterface.h
--- a/Engine/Source/Runtime/TypedElementFramework/Public/Elements/Interfaces/TypedElementDataStorageInterface.h
+++ b/Engine/Source/Runtime/TypedElementFramework/Public/Elements/Interfaces/TypedElementDataStorageInterface.h
@@ -60,10 +60,15 @@ public:
 	/** Finds a factory instance registered with TEDS */
 	virtual const UEditorDataStorageFactory* FindFactory(const UClass* FactoryType) const = 0;
 
+	virtual UEditorDataStorageFactory* FindFactory(const UClass* FactoryType) = 0;
+
 	/** Convenience function for FindFactory */
 	template<typename FactoryT>
 	const FactoryT* FindFactory() const;
 	
+	template<typename FactoryT>
+	FactoryT* FindFactory();
+
 	/**
 	 * @section Table management
 	 * 
@@ -475,6 +480,12 @@ const FactoryT* ICoreProvider::FindFactory() const
 	return static_cast<const FactoryT*>(FindFactory(FactoryT::StaticClass()));
 }
 
+template <typename FactoryT>
+FactoryT* ICoreProvider::FindFactory()
+{
+	return static_cast<FactoryT*>(FindFactory(FactoryT::StaticClass()));
+}
+
 template<TColumnType... Columns>
 TableHandle ICoreProvider::RegisterTable(const FName& Name)
 {
