Teds Query Conditions: Convert the Column and Token array to a TArray with an inline allocator

diff --git a/Engine/Source/Runtime/TypedElementFramework/Private/Elements/Common/TypedElementQueryConditions.cpp b/Engine/Source/Runtime/TypedElementFramework/Private/Elements/Common/TypedElementQueryConditions.cpp
--- a/Engine/Source/Runtime/TypedElementFramework/Private/Elements/Common/TypedElementQueryConditions.cpp
+++ b/Engine/Source/Runtime/TypedElementFramework/Private/Elements/Common/TypedElementQueryConditions.cpp
@@ -23,12 +23,17 @@ namespace UE::Editor::DataStorage::Queries
 	FConditions::FConditions()
 	{
 		Identifiers.Init(NAME_None, MaxColumnCount);
+		Columns.Init(nullptr, MaxColumnCount);
+		Tokens.Init(Token::None, MaxTokenCount);
 	}
 
 	FConditions::FConditions(FColumnBase Column)
 		: ColumnCount(1)
 	{
 		Identifiers.Init(NAME_None, MaxColumnCount);
+		Columns.Init(nullptr, MaxColumnCount);
+		Tokens.Init(Token::None, MaxTokenCount);
+		
 		Columns[0] = Column.TypeInfo;
 		Identifiers[0] = Column.Identifier;
 	}
@@ -352,7 +357,7 @@ namespace UE::Editor::DataStorage::Queries
 	{
 		checkf(bIsCompiled, TEXT("Query Conditions must call Compile() before you can use them"));
 		
-		return TConstArrayView<TWeakObjectPtr<const UScriptStruct>>(Columns, ColumnCount);
+		return TConstArrayView<TWeakObjectPtr<const UScriptStruct>>(Columns.GetData(), ColumnCount);
 	}
 
 	bool FConditions::IsEmpty() const
@@ -519,7 +524,7 @@ namespace UE::Editor::DataStorage::Queries
 
 	void FConditions::ConvertColumnBitToArray(TArray<TWeakObjectPtr<const UScriptStruct>>& MatchedColumns, uint64 ColumnBits) const
 	{
-		uint64 Index = 0;
+		int Index = 0;
 		while (ColumnBits)
 		{
 			if (ColumnBits & 1)
diff --git a/Engine/Source/Runtime/TypedElementFramework/Public/Elements/Common/TypedElementQueryConditions.h b/Engine/Source/Runtime/TypedElementFramework/Public/Elements/Common/TypedElementQueryConditions.h
--- a/Engine/Source/Runtime/TypedElementFramework/Public/Elements/Common/TypedElementQueryConditions.h
+++ b/Engine/Source/Runtime/TypedElementFramework/Public/Elements/Common/TypedElementQueryConditions.h
@@ -169,14 +169,14 @@ namespace Queries
 		};
 		FRIEND_ENUM_CLASS_FLAGS(EColumnFlags);
 		
-		using ColumnArray = TWeakObjectPtr<const UScriptStruct>[MaxColumnCount];
-		using TokenArray = Token[MaxTokenCount];
+		using ColumnArray = TArray<TWeakObjectPtr<const UScriptStruct>, TInlineAllocator<MaxColumnCount>>;
+		using TokenArray = TArray<Token, TInlineAllocator<MaxTokenCount>>;
 		using IdentifierArray = TArray<FName, TInlineAllocator<MaxColumnCount>>;
 		using ColumnFlagArray = TArray<EColumnFlags, TInlineAllocator<MaxColumnCount>>;
 
-		ColumnArray Columns{ nullptr };
+		ColumnArray Columns;
 		ColumnFlagArray ColumnFlags;
-		TokenArray Tokens{ Token::None };
+		TokenArray Tokens;
 		IdentifierArray Identifiers;
 		uint8_t ColumnCount = 0;
 		uint8_t TokenCount = 0;
