TEDS:

Add a new query utils library with a util function to merge two queries
Add a new query stack node to merge queries

diff --git a/Engine/Plugins/Experimental/EditorDataStorageFeatures/Source/TedsQueryStack/Private/TedsQueryMergeNode.cpp b/Engine/Plugins/Experimental/EditorDataStorageFeatures/Source/TedsQueryStack/Private/TedsQueryMergeNode.cpp
new file mode 100644
--- /dev/null
+++ b/Engine/Plugins/Experimental/EditorDataStorageFeatures/Source/TedsQueryStack/Private/TedsQueryMergeNode.cpp
@@ -0,0 +1,83 @@
+// Copyright Epic Games, Inc. All Rights Reserved.
+
+#include "TedsQueryMergeNode.h"
+
+#include "Elements/Framework/TypedElementQueryBuilder.h"
+#include "Elements/Interfaces/TypedElementDataStorageInterface.h"
+#include "TedsQueryStackLog.h"
+
+namespace UE::Editor::DataStorage::QueryStack
+{
+	FQueryMergeNode::FQueryMergeNode(ICoreProvider& InStorage, TConstArrayView<TSharedPtr<IQueryNode>> InParents)
+		: Storage(InStorage)
+	{
+		checkf(!InParents.IsEmpty(), TEXT("FQueryMergeNode requires at least one parent query."));
+
+		Parents.Reserve(InParents.Num());
+		for (const TSharedPtr<IQueryNode>& Parent : InParents)
+		{
+			Parents.Emplace<FParentInfo>({ .Parent = Parent, .Revision = Parent->GetRevision()});
+		}
+		Rebuild();
+	}
+
+	FQueryMergeNode::~FQueryMergeNode()
+	{
+		Storage.UnregisterQuery(QueryHandle);
+	}
+
+	INode::RevisionId FQueryMergeNode::GetRevision() const
+	{
+		return Revision;
+	}
+
+	void FQueryMergeNode::Update()
+	{
+		bool bRebuild = false;
+		for (FParentInfo& ParentInfo : Parents)
+		{
+			ParentInfo.Parent->Update();
+			RevisionId ParentRevision = ParentInfo.Parent->GetRevision();
+			if (ParentInfo.Revision != ParentRevision)
+			{
+				bRebuild = true;
+				ParentInfo.Revision = ParentRevision;
+				break;
+			}
+		}
+
+		if (bRebuild)
+		{
+			Rebuild();
+			Revision++;
+		}
+	}
+
+	QueryHandle FQueryMergeNode::GetQuery() const
+	{
+		return QueryHandle;
+	}
+
+	void FQueryMergeNode::Rebuild()
+	{
+		Storage.UnregisterQuery(QueryHandle);
+		QueryHandle = InvalidQueryHandle;
+
+		FQueryDescription FinalQueryDescription;
+
+		FText ErrorText;
+		
+		for (const FParentInfo& Parent : Parents)
+		{
+			FQueryDescription ParentQueryDescription = Storage.GetQueryDescription(Parent.Parent->GetQuery());
+
+			if (!Queries::MergeQueries(FinalQueryDescription, ParentQueryDescription, &ErrorText))
+			{
+				UE_LOG(LogTedsQueryStack, Error, TEXT("Skipping Parent Query %s due to error: %s"),
+					*ParentQueryDescription.Callback.Name.ToString(), *ErrorText.ToString());
+			}
+		}
+
+		QueryHandle = Storage.RegisterQuery(MoveTemp(FinalQueryDescription));
+	}
+}
diff --git a/Engine/Plugins/Experimental/EditorDataStorageFeatures/Source/TedsQueryStack/Private/TedsQueryStackLog.cpp b/Engine/Plugins/Experimental/EditorDataStorageFeatures/Source/TedsQueryStack/Private/TedsQueryStackLog.cpp
new file mode 100644
--- /dev/null
+++ b/Engine/Plugins/Experimental/EditorDataStorageFeatures/Source/TedsQueryStack/Private/TedsQueryStackLog.cpp
@@ -0,0 +1,5 @@
+// Copyright Epic Games, Inc. All Rights Reserved.
+
+#include "TedsQueryStackLog.h"
+
+DEFINE_LOG_CATEGORY(LogTedsQueryStack);
diff --git a/Engine/Plugins/Experimental/EditorDataStorageFeatures/Source/TedsQueryStack/Private/TedsQueryStackLog.h b/Engine/Plugins/Experimental/EditorDataStorageFeatures/Source/TedsQueryStack/Private/TedsQueryStackLog.h
new file mode 100644
--- /dev/null
+++ b/Engine/Plugins/Experimental/EditorDataStorageFeatures/Source/TedsQueryStack/Private/TedsQueryStackLog.h
@@ -0,0 +1,7 @@
+// Copyright Epic Games, Inc. All Rights Reserved.
+
+#pragma once
+
+#include "Logging/LogMacros.h"
+
+DECLARE_LOG_CATEGORY_EXTERN(LogTedsQueryStack, Log, All);
diff --git a/Engine/Plugins/Experimental/EditorDataStorageFeatures/Source/TedsQueryStack/Private/TedsRowMonitorNode.cpp b/Engine/Plugins/Experimental/EditorDataStorageFeatures/Source/TedsQueryStack/Private/TedsRowMonitorNode.cpp
--- a/Engine/Plugins/Experimental/EditorDataStorageFeatures/Source/TedsQueryStack/Private/TedsRowMonitorNode.cpp
+++ b/Engine/Plugins/Experimental/EditorDataStorageFeatures/Source/TedsQueryStack/Private/TedsRowMonitorNode.cpp
@@ -118,6 +118,16 @@ namespace UE::Editor::DataStorage::QueryStack
 			TargetRows.SortedMerge(AddedRows);
 			TargetRows.Remove(RemovedRows.GetRows());
 
+			if (!AddedRows.IsEmpty())
+			{
+				RowsAddedEvent.Broadcast(AddedRows.GetRows());
+			}
+			
+			if (!RemovedRows.IsEmpty())
+			{
+				RowsRemovedEvent.Broadcast(RemovedRows.GetRows());
+			}
+			
 			Revision++;
 		}
 		AddedRows.Empty();
@@ -134,6 +144,16 @@ namespace UE::Editor::DataStorage::QueryStack
 		return &ResolveRows();
 	}
 
+	FRowMonitorNode::FOnRowsChanged& FRowMonitorNode::OnMonitoredRowsAdded()
+	{
+		return RowsAddedEvent;
+	}
+
+	FRowMonitorNode::FOnRowsChanged& FRowMonitorNode::OnMonitoredRowsRemoved()
+	{
+		return RowsRemovedEvent;
+	}
+
 	const FRowHandleArray& FRowMonitorNode::ResolveRows() const
 	{
 		return const_cast<FRowMonitorNode*>(this)->ResolveRows();
diff --git a/Engine/Plugins/Experimental/EditorDataStorageFeatures/Source/TedsQueryStack/Public/TedsQueryMergeNode.h b/Engine/Plugins/Experimental/EditorDataStorageFeatures/Source/TedsQueryStack/Public/TedsQueryMergeNode.h
new file mode 100644
--- /dev/null
+++ b/Engine/Plugins/Experimental/EditorDataStorageFeatures/Source/TedsQueryStack/Public/TedsQueryMergeNode.h
@@ -0,0 +1,45 @@
+// Copyright Epic Games, Inc. All Rights Reserved.
+
+#pragma once
+
+#include "Elements/Common/TypedElementQueryDescription.h"
+#include "TedsQueryNode.h"
+
+namespace UE::Editor::DataStorage
+{
+	class ICoreProvider;
+} // namespace UE::Editor::DataStorage
+
+namespace UE::Editor::DataStorage::QueryStack
+{
+	/**
+	 * Merges one or more parent queries into a composite query
+	 * A merge of two query nodes has the same result of the intersection of the rows when you run the queries individually
+	 */
+	class FQueryMergeNode : public IQueryNode
+	{
+	public:
+		TEDSQUERYSTACK_API explicit FQueryMergeNode(ICoreProvider& InStorage, TConstArrayView<TSharedPtr<IQueryNode>> InParents);
+		TEDSQUERYSTACK_API virtual ~FQueryMergeNode() override;
+
+		TEDSQUERYSTACK_API virtual RevisionId GetRevision() const override;
+		TEDSQUERYSTACK_API virtual void Update() override;
+		TEDSQUERYSTACK_API virtual QueryHandle GetQuery() const override;
+
+	protected:
+		void Rebuild();
+		
+	private:
+
+		struct FParentInfo
+		{
+			TSharedPtr<IQueryNode> Parent;
+			RevisionId Revision;
+		};
+
+		QueryHandle QueryHandle = InvalidQueryHandle;
+		TArray<FParentInfo> Parents;
+		ICoreProvider& Storage;
+		RevisionId Revision = 0;
+	};
+} // namespace UE::Editor::DataStorage::QueryStack
diff --git a/Engine/Plugins/Experimental/EditorDataStorageFeatures/Source/TedsQueryStack/Public/TedsRowMonitorNode.h b/Engine/Plugins/Experimental/EditorDataStorageFeatures/Source/TedsQueryStack/Public/TedsRowMonitorNode.h
--- a/Engine/Plugins/Experimental/EditorDataStorageFeatures/Source/TedsQueryStack/Public/TedsRowMonitorNode.h
+++ b/Engine/Plugins/Experimental/EditorDataStorageFeatures/Source/TedsQueryStack/Public/TedsRowMonitorNode.h
@@ -20,6 +20,9 @@ namespace UE::Editor::DataStorage::QueryStack
 	class FRowMonitorNode : public IRowNode
 	{
 	public:
+
+		DECLARE_MULTICAST_DELEGATE_OneParam(FOnRowsChanged, FRowHandleArrayView /* InRows */);
+		
 		TEDSQUERYSTACK_API FRowMonitorNode(
 			ICoreProvider& Storage, TSharedPtr<IRowNode> ParentRow, TSharedPtr<IQueryNode> QueryNode);
 		TEDSQUERYSTACK_API FRowMonitorNode(
@@ -41,6 +44,15 @@ namespace UE::Editor::DataStorage::QueryStack
 		TEDSQUERYSTACK_API virtual FRowHandleArrayView GetRows() const override;
 		TEDSQUERYSTACK_API virtual FRowHandleArray* GetMutableRows() override;
 
+		// NOTE: These events are not meant to be a wholesale solution for tracking changes to this node, it only fires when monitored rows are modified
+		// through observers. They are also deferred to Update() and not fired immediately.
+		// For full tracking the regular pattern of checking the node's revision and calling GetRows() should be used
+		
+		// Event that is called during Update() to allow systems to respond to monitored rows being added
+		TEDSQUERYSTACK_API FOnRowsChanged& OnMonitoredRowsAdded();
+
+		// Event that is called during Update() to allow systems to respond to monitored rows being removed
+		TEDSQUERYSTACK_API FOnRowsChanged& OnMonitoredRowsRemoved();
 	private:
 		const FRowHandleArray& ResolveRows() const;
 		FRowHandleArray& ResolveRows();
@@ -62,5 +74,8 @@ namespace UE::Editor::DataStorage::QueryStack
 		RevisionId ParentRevision = 0;
 		RevisionId Revision = 0;
 		bool bFixedColumns = false;
+		
+		FOnRowsChanged RowsAddedEvent;
+		FOnRowsChanged RowsRemovedEvent;
 	};
 } // namespace UE::Editor::DataStorage::QueryStack
diff --git a/Engine/Source/Runtime/TypedElementFramework/Private/Elements/Framework/TypedElementMetaData.cpp b/Engine/Source/Runtime/TypedElementFramework/Private/Elements/Framework/TypedElementMetaData.cpp
--- a/Engine/Source/Runtime/TypedElementFramework/Private/Elements/Framework/TypedElementMetaData.cpp
+++ b/Engine/Source/Runtime/TypedElementFramework/Private/Elements/Framework/TypedElementMetaData.cpp
@@ -42,6 +42,10 @@ namespace UE::Editor::DataStorage
 		MutableData.Shrink();
 	}
 
+	bool FMetaDataBase::IsEmpty() const
+	{
+		return ImmutableData.IsEmpty() && MutableData.IsEmpty();
+	}
 
 
 	//
@@ -79,6 +83,11 @@ namespace UE::Editor::DataStorage
 		return Result;
 	}
 
+	bool FColumnMetaData::IsEmpty() const
+	{
+		return FMetaDataBase::IsEmpty() && ColumnType == nullptr;
+	}
+
 
 	//
 	// FMetaDataEntryView
diff --git a/Engine/Source/Runtime/TypedElementFramework/Private/Elements/Framework/TypedElementQueryBuilder.cpp b/Engine/Source/Runtime/TypedElementFramework/Private/Elements/Framework/TypedElementQueryBuilder.cpp
--- a/Engine/Source/Runtime/TypedElementFramework/Private/Elements/Framework/TypedElementQueryBuilder.cpp
+++ b/Engine/Source/Runtime/TypedElementFramework/Private/Elements/Framework/TypedElementQueryBuilder.cpp
@@ -7,10 +7,11 @@
 #include "Elements/Framework/TypedElementMetaData.h"
 #include "GenericPlatform/GenericPlatformMath.h"
 
+#define LOCTEXT_NAMESPACE "TypedElementQueryUtils"
+
+
 namespace UE::Editor::DataStorage::Queries
 {
-	using namespace UE::Editor::DataStorage;
-
 	const UScriptStruct* Type(FTopLevelAssetPath Name)
 	{
 		const UScriptStruct* StructInfo = TypeOptional(Name);
@@ -551,4 +552,469 @@ namespace UE::Editor::DataStorage::Queries
 	{
 		return FDependency{ &Query };
 	}
+	
+	/**
+	 * MergeQueries
+	 */
+
+	namespace Private
+	{
+		// Enum to get the result when two TEDS operators need to be tested against each other
+		enum EConflictResolution
+		{
+			First, // The first operator is stronger
+			Second, // The second operator is stronger
+			Conflict, // There is a conflict between the two operators
+			Unknown // An unknown operator was found
+		};
+
+		EConflictResolution Resolve(EQueryAccessType First, EQueryAccessType Second)
+		{
+			switch (First)
+			{
+			case EQueryAccessType::ReadWrite:
+
+				// ReadWrite is the strongest access type
+				return EConflictResolution::First;
+				
+			case EQueryAccessType::ReadOnly:
+
+				// ReadWrite is a stronger access type than ReadOnly
+				return Second == EQueryAccessType::ReadWrite ? EConflictResolution::Second : EConflictResolution::First;
+				
+			case EQueryAccessType::OptionalReadOnly:
+
+				// OptionalReadOnly is the weakest access type
+				return EConflictResolution::Second;
+				
+			default:
+				checkf(false, TEXT("Unknown EQueryAccessType value found when trying to resolve conditions."))
+				return EConflictResolution::Unknown;
+			}
+		}
+		
+		EConflictResolution Resolve(FQueryDescription::EOperatorType First, FQueryDescription::EOperatorType Second)
+		{
+			switch (First)
+			{
+			case FQueryDescription::EOperatorType::SimpleAll:
+
+				// All is the strongest operator, but it conflicts with None
+				return Second == FQueryDescription::EOperatorType::SimpleNone ? EConflictResolution::Conflict : EConflictResolution::First;
+				
+			case FQueryDescription::EOperatorType::SimpleAny:
+
+				// Any conflicts with None, and is weaker than All
+				if (Second == FQueryDescription::EOperatorType::SimpleNone)
+				{
+					return EConflictResolution::Conflict;
+				}
+				else if (Second == FQueryDescription::EOperatorType::SimpleAll)
+				{
+					return EConflictResolution::Second;
+				}
+
+				return EConflictResolution::First;
+				
+			case FQueryDescription::EOperatorType::SimpleNone:
+
+				// None conflicts with everything except itself
+				return Second != FQueryDescription::EOperatorType::SimpleNone ? EConflictResolution::Conflict : EConflictResolution::First;
+				
+			
+			case FQueryDescription::EOperatorType::SimpleOptional:
+				
+				// SimpleOptional is currently unused, so we'll treat it as the weakest
+				return EConflictResolution::Second;
+				
+			default:
+				checkf(false, TEXT("Unknown EOperatorType value found when trying to resolve conditions."))
+				return EConflictResolution::Unknown;
+			}
+		}
+		
+		EConflictResolution Resolve(EQueryAccessType First, FQueryDescription::EOperatorType Second)
+		{
+			switch (First)
+			{
+			case EQueryAccessType::ReadWrite:
+
+				// Selection takes precedence over conditions, but ReadWrite conflicts with SimpleNone
+				return Second == FQueryDescription::EOperatorType::SimpleNone ? EConflictResolution::Conflict : EConflictResolution::First;
+				
+			case EQueryAccessType::ReadOnly:
+
+				// Selection takes precedence over conditions, but ReadOnly conflicts with SimpleNone
+				return Second == FQueryDescription::EOperatorType::SimpleNone ? EConflictResolution::Conflict : EConflictResolution::First;
+				
+			case EQueryAccessType::OptionalReadOnly:
+
+				// OptionalReadOnly conflicts with a SimpleNone
+				if (Second == FQueryDescription::EOperatorType::SimpleNone)
+				{
+					return EConflictResolution::Conflict;
+				}
+
+				// OptionalReadOnly is weaker than SimpleAll or SimpleAny, but stronger than SimpleOptional because selection wins over conditions
+				// NOTE: SimpleOptional is currently unused
+				return Second == FQueryDescription::EOperatorType::SimpleOptional ?	EConflictResolution::First : EConflictResolution::Second;
+				
+			default:
+				checkf(false, TEXT("Unknown EQueryAccessType value found when trying to resolve conditions."))
+				return EConflictResolution::Unknown;
+			}
+		}
+
+		void RemoveSelection(FQueryDescription& Query, int32 SelectionIndex, bool bDynamic)
+		{
+			if (bDynamic)
+			{
+				Query.DynamicSelectionAccessTypes.RemoveAt(SelectionIndex);
+				Query.DynamicSelectionTypes.RemoveAt(SelectionIndex);
+				Query.DynamicSelectionMetaData.RemoveAt(SelectionIndex);
+			}
+			else
+			{
+				Query.SelectionAccessTypes.RemoveAt(SelectionIndex);
+				Query.SelectionTypes.RemoveAt(SelectionIndex);
+				Query.SelectionMetaData.RemoveAt(SelectionIndex);
+			}
+		};
+
+		void RemoveCondition(FQueryDescription& Query, int32 ConditionIndex, bool bDynamic)
+		{
+			if (bDynamic)
+			{
+				Query.DynamicConditionOperations.RemoveAt(ConditionIndex);
+				Query.DynamicConditionDescriptions.RemoveAt(ConditionIndex);
+			}
+			else
+			{
+				Query.ConditionOperators.RemoveAt(ConditionIndex);
+				Query.ConditionTypes.RemoveAt(ConditionIndex);
+			}
+		};
+		
+		// Util function to resolve duplicate columns/tags at a specific index inside the Select() part of a query
+		void ResolveSelectDuplicate(FQueryDescription& Query, int32 OriginalIndex, int32 DuplicateIndex, bool bDynamic)
+		{
+			EQueryAccessType OriginalAccessType = bDynamic ?
+				Query.DynamicSelectionAccessTypes[OriginalIndex] : Query.SelectionAccessTypes[OriginalIndex];
+			
+			EQueryAccessType DuplicateAccessType = bDynamic ? Query.DynamicSelectionAccessTypes[DuplicateIndex] : Query.SelectionAccessTypes[DuplicateIndex];
+
+			// Remove the weaker of the two operators to resolve the duplicate
+			switch (Resolve(OriginalAccessType, DuplicateAccessType))
+			{
+			case EConflictResolution::First:
+				RemoveSelection(Query, DuplicateIndex, bDynamic);
+				break;
+				
+			case EConflictResolution::Second:
+				RemoveSelection(Query, OriginalIndex, bDynamic);
+				break;
+
+			// Two Selects can never conflict with each other, something went wrong (i.e one of them was unknown)
+			case EConflictResolution::Conflict:
+				checkf(false, TEXT("Two Selects() cannot have a conflict as it only supports simple requirements, one of the queries is probably invalid."))
+				break;
+			case EConflictResolution::Unknown:
+			default:
+
+				checkf(false, TEXT("Unknown EQueryAccessType value found when trying to merge queries."))
+				break;
+			}
+		}
+		
+		// Util function to resolve duplicate columns/tags at a specific index inside the conditions of two queries
+		bool ResolveConditionDuplicate(FQueryDescription& Query, int32 OriginalIndex, int32 DuplicateIndex, bool bDynamic, FText* OutErrorMessage)
+		{
+			FQueryDescription::EOperatorType OriginalOperatorType = bDynamic ?
+				Query.DynamicConditionOperations[OriginalIndex] : Query.ConditionTypes[OriginalIndex];
+
+			FQueryDescription::EOperatorType DuplicateOperatorType = bDynamic ?
+				Query.DynamicConditionOperations[DuplicateIndex] : Query.ConditionTypes[DuplicateIndex];
+			
+			// Remove the weaker of the two conditions, or report an error if they conflict
+			switch (Resolve(OriginalOperatorType, DuplicateOperatorType))
+			{
+			case EConflictResolution::First:
+				RemoveCondition(Query, DuplicateIndex, bDynamic);
+				return true;
+				
+			case EConflictResolution::Second:
+				RemoveCondition(Query, OriginalIndex, bDynamic);
+				return true;
+
+			case EConflictResolution::Conflict:
+
+				if (OutErrorMessage)
+				{
+					*OutErrorMessage = LOCTEXT("CannotMerge_ConditionConflict",
+						"Cannot merge two queries if they contain the same requirement under conflicting conditions (None and All/Any)");
+				}
+
+				return false;
+			case EConflictResolution::Unknown:
+			default:
+
+				checkf(false, TEXT("Unknown EOperatorType value found when trying to merge queries."))
+				return false;
+			}
+		}
+
+		// Util function to resolve duplicates between a condition type and a selection type on a single query
+		bool ResolveConditionSelectionDuplicate(FQueryDescription& Query, int32 SelectionIndex, int32 ConditionIndex, bool bDynamic, FText* OutErrorMessage)
+		{
+			EQueryAccessType SelectionAccessType = bDynamic ? Query.DynamicSelectionAccessTypes[SelectionIndex] : Query.SelectionAccessTypes[SelectionIndex];
+
+			FQueryDescription::EOperatorType ConditionOperatorType = bDynamic ?
+				Query.DynamicConditionOperations[ConditionIndex] : Query.ConditionTypes[ConditionIndex];
+
+			// Remove the weaker of the two, or report any conflicts
+			switch (Resolve(SelectionAccessType, ConditionOperatorType))
+			{
+			case EConflictResolution::First:
+				RemoveCondition(Query, ConditionIndex, bDynamic);
+				return true;
+				
+			case EConflictResolution::Second:
+				RemoveSelection(Query, SelectionIndex, bDynamic);
+				return true;
+
+			case EConflictResolution::Conflict:
+
+				if (OutErrorMessage)
+				{
+					*OutErrorMessage = LOCTEXT("CannotMerge_SelectAndConditionConflict",
+						"Cannot merge two queries if they contain conflicting requirements in the Select() and Where() clauses");
+				}
+				return false;
+				
+			case EConflictResolution::Unknown:
+			default:
+
+				checkf(false, TEXT("Unknown EOperatorType value found when trying to merge queries."))
+				return false;
+			}
+		}
+
+		// Resolve any duplicates or conflicts that may have resulted from the merge
+		bool ResolveDuplicatesInQuery(FQueryDescription& Query, FText* OutErrorMessage)
+		{
+			// NOTE: We don't need to check the FConditions because they currently don't support NOT and cannot conflict, and the query matching logic
+			// for them is also overriden by TEDS so they can work with duplicates
+
+			// Check the selection types
+			for (int32 SelectionIndex = 0; SelectionIndex < Query.SelectionTypes.Num(); SelectionIndex++)
+			{
+				// Check if the same selection type is present twice in the query.
+				// Using FindLastByPredicate is fine because we can only have two copies of the same type
+				int32 DuplicateIndex = Query.SelectionTypes.FindLastByPredicate([&Query, SelectionIndex](const TWeakObjectPtr<const UScriptStruct>& SelectionType)
+				{
+					return SelectionType == Query.SelectionTypes[SelectionIndex];
+				});
+
+				if (DuplicateIndex != INDEX_NONE && DuplicateIndex != SelectionIndex)
+				{
+					ResolveSelectDuplicate(Query, SelectionIndex, DuplicateIndex, false);
+				}
+
+				// Check if the selection type is present in the conditions of the query
+				int32 ConditionDuplicateIndex = Query.ConditionOperators.IndexOfByPredicate([&Query, SelectionIndex](const FQueryDescription::FOperator& Operator)
+				{
+					return Operator.Type == Query.SelectionTypes[SelectionIndex];
+				});
+
+				if (ConditionDuplicateIndex != INDEX_NONE)
+				{
+					if (!ResolveConditionSelectionDuplicate(Query, SelectionIndex, ConditionDuplicateIndex, false, OutErrorMessage))
+					{
+						return false;
+					}
+				}
+			}
+
+			// Do the same for the dynamic selection types
+			for (int32 SelectionIndex = 0; SelectionIndex < Query.DynamicSelectionTypes.Num(); SelectionIndex++)
+			{
+				// Check if the same selection type is present twice in the query.
+				// Using FindLastByPredicate is fine because we can only have two copies of the same type
+				int32 DuplicateIndex = Query.DynamicSelectionTypes.FindLastByPredicate([&Query, SelectionIndex](const FDynamicColumnDescription& Description)
+				{
+					return Description == Query.DynamicSelectionTypes[SelectionIndex];
+				});
+
+				if (DuplicateIndex != INDEX_NONE && DuplicateIndex != SelectionIndex)
+				{
+					ResolveSelectDuplicate(Query, SelectionIndex, DuplicateIndex, true);
+				}
+
+				// Check if the selection type is present in the conditions of the query
+				int32 ConditionDuplicateIndex = Query.DynamicConditionDescriptions.IndexOfByPredicate([&Query, SelectionIndex](const FDynamicColumnDescription& Description)
+				{
+					return Description == Query.DynamicSelectionTypes[SelectionIndex];
+				});
+
+				if (ConditionDuplicateIndex != INDEX_NONE)
+				{
+					if (!ResolveConditionSelectionDuplicate(Query, SelectionIndex, ConditionDuplicateIndex, true, OutErrorMessage))
+					{
+						return false;
+					}
+				}
+			}
+
+			// Check if the same condition operator is present twice in the query.
+			for (int32 ConditionIndex = 0; ConditionIndex < Query.ConditionOperators.Num(); ConditionIndex++)
+			{
+				// NOTE: We don't need to test against selection again because the selection loops handle that
+
+				// Using FindLastByPredicate is fine because we can only have two copies of the same type
+				int32 ConditionDuplicateIndex = Query.ConditionOperators.FindLastByPredicate([&Query, ConditionIndex](const FQueryDescription::FOperator& Operator)
+				{
+					return Operator.Type == Query.ConditionOperators[ConditionIndex].Type;
+				});
+
+				if (ConditionDuplicateIndex != INDEX_NONE && ConditionDuplicateIndex != ConditionIndex)
+				{
+					if (!ResolveConditionDuplicate(Query, ConditionIndex, ConditionDuplicateIndex, false, OutErrorMessage))
+					{
+						return false;
+					}
+				}
+			}
+
+			// Check if the same dynamic condition operator is present twice in the query.
+			for (int32 ConditionIndex = 0; ConditionIndex < Query.DynamicConditionDescriptions.Num(); ConditionIndex++)
+			{
+				// NOTE: We don't need to test against selection again because the selection loops handle that
+
+				// Using FindLastByPredicate is fine because we can only have two copies of the same type
+				int32 ConditionDuplicateIndex = Query.DynamicConditionDescriptions.FindLastByPredicate([&Query, ConditionIndex](const FDynamicColumnDescription& Description)
+				{
+					return Description == Query.DynamicConditionDescriptions[ConditionIndex];
+				});
+
+				if (ConditionDuplicateIndex != INDEX_NONE && ConditionDuplicateIndex != ConditionIndex)
+				{
+					if (!ResolveConditionDuplicate(Query, ConditionIndex, ConditionDuplicateIndex, true, OutErrorMessage))
+					{
+						return false;
+					}
+				}
+			}
+			return true;
+		}
+
+		// Util function for some upfront checks to see if the two queries can be merged
+		bool AreMergeablePrecheck(const FQueryDescription& Destination, const FQueryDescription& Source, FText* OutErrorMessage)
+		{
+			if (Destination.Callback.Function.IsSet() && Source.Callback.Function.IsSet())
+			{
+				if (OutErrorMessage)
+				{
+					*OutErrorMessage = LOCTEXT("CannotMerge_Callbacks", "Cannot merge two queries if both contain a callback");
+				}
+				return false;
+			}
+
+			if ((Destination.Conditions && !Source.Conditions) || (!Destination.Conditions && Source.Conditions))
+			{
+				if (OutErrorMessage)
+				{
+					*OutErrorMessage = LOCTEXT("CannotMerge_FConditions",
+						"Cannot merge two queries if one of them uses the FConditionsSyntax and the other uses the legacy .All/.Any/.None syntax");
+				}
+				return false;
+			}
+
+			if (!Destination.MetaData.IsEmpty() && !Source.MetaData.IsEmpty())
+			{
+				if (OutErrorMessage)
+				{
+					*OutErrorMessage = LOCTEXT("CannotMerge_Metadata",
+						"Cannot merge two queries if they both contain metadata, since FMetaData currently does not support merging");
+				}
+				return false;
+			}
+
+			return true;
+		}
+
+		// Merge the Select() section of the two queries
+		void MergeSelect(FQueryDescription& Destination, const FQueryDescription& Source)
+		{
+			Destination.SelectionTypes.Append(Source.SelectionTypes);
+			Destination.SelectionAccessTypes.Append(Source.SelectionAccessTypes);
+			Destination.SelectionMetaData.Append(Source.SelectionMetaData);
+			
+			Destination.DynamicSelectionTypes.Append(Source.DynamicSelectionTypes);
+			Destination.DynamicSelectionAccessTypes.Append(Source.DynamicSelectionAccessTypes);
+			Destination.DynamicSelectionMetaData.Append(Source.DynamicSelectionMetaData);
+		}
+
+		// Merge the Where() section of the two queries
+		void MergeConditions(FQueryDescription& Destination, const FQueryDescription& Source)
+		{
+			Destination.ConditionOperators.Append(Source.ConditionOperators);
+			Destination.ConditionTypes.Append(Source.ConditionTypes);
+			
+			Destination.DynamicConditionDescriptions.Append(Source.DynamicConditionDescriptions);
+			Destination.DynamicConditionOperations.Append(Source.DynamicConditionOperations);
+		}
+	}
+	
+	bool MergeQueries(FQueryDescription& Destination, const FQueryDescription& Source, FText* OutErrorMessage)
+	{
+		if (!Private::AreMergeablePrecheck(Destination, Source, OutErrorMessage))
+		{
+			return false;
+		}
+
+		// Keep a copy of the original destination in case we detect that the merge is not possible partway through the process (e.g conflicting conditions)
+		FQueryDescription OriginalDestination = Destination;
+
+		// If we are combining a select and a count query, the final query will be treated as a select
+		Destination.Action = Source.Action == FQueryDescription::EActionType::Select ? Source.Action : Destination.Action;
+
+		// If we are at this point we are guaranteed only one of the two queries has a callback
+		if (Source.Callback.Function.IsSet())
+		{
+			Destination.Callback.Function = Source.Callback.Function;
+			Destination.bShouldBatchModifications = Source.bShouldBatchModifications;
+		}
+
+		// Copy over the trivially copyable data
+		Destination.ValueTags.Append(Source.ValueTags);
+		Destination.DependencyTypes.Append(Source.DependencyTypes);
+		Destination.DependencyFlags.Append(Source.DependencyFlags);
+		Destination.CachedDependencies.Append(Source.CachedDependencies);
+		Destination.Subqueries.Append(Source.Subqueries);
+
+		// Merge the Select() portion of the queries
+		Private::MergeSelect(Destination, Source);
+
+		// Since the FConditions cannot co-exist with the ConditionOperators, we only need to merge one of the two
+		if (Destination.Conditions)
+		{
+			Destination.Conditions = Destination.Conditions.GetValue() && Source.Conditions.GetValue();
+		}
+		else
+		{
+			Private::MergeConditions(Destination, Source);
+		}
+
+		// Validate the merged query
+		if (!Private::ResolveDuplicatesInQuery(Destination, OutErrorMessage))
+		{
+			// Restore the destination back to the original value if the merge failed
+			Destination = OriginalDestination;
+			return false;
+		}
+		
+		return true;
+	}
 } // namespace UE::Editor::DataStorage::Queries
+
+#undef LOCTEXT_NAMESPACE
\ No newline at end of file
diff --git a/Engine/Source/Runtime/TypedElementFramework/Public/Elements/Framework/TypedElementMetaData.h b/Engine/Source/Runtime/TypedElementFramework/Public/Elements/Framework/TypedElementMetaData.h
--- a/Engine/Source/Runtime/TypedElementFramework/Public/Elements/Framework/TypedElementMetaData.h
+++ b/Engine/Source/Runtime/TypedElementFramework/Public/Elements/Framework/TypedElementMetaData.h
@@ -63,6 +63,8 @@ namespace UE::Editor::DataStorage
 		TYPEDELEMENTFRAMEWORK_API virtual FMetaDataEntryView Find(FName Name) const;
 
 		TYPEDELEMENTFRAMEWORK_API virtual void Shrink();
+		
+		TYPEDELEMENTFRAMEWORK_API virtual bool IsEmpty() const;
 
 	protected:
 		FMetaDataBase() = default;
@@ -93,6 +95,8 @@ namespace UE::Editor::DataStorage
 		TYPEDELEMENTFRAMEWORK_API FColumnMetaData(const UScriptStruct* InColumnType, EFlags InFlags);
 
 		TYPEDELEMENTFRAMEWORK_API FMetaDataEntryView Find(FName Name) const override;
+		
+		TYPEDELEMENTFRAMEWORK_API bool IsEmpty() const override;
 
 	private:
 		/** If set, properties on the column will also be included if a value isn't found in the (im)mutable data map. */
diff --git a/Engine/Source/Runtime/TypedElementFramework/Public/Elements/Framework/TypedElementQueryBuilder.h b/Engine/Source/Runtime/TypedElementFramework/Public/Elements/Framework/TypedElementQueryBuilder.h
--- a/Engine/Source/Runtime/TypedElementFramework/Public/Elements/Framework/TypedElementQueryBuilder.h
+++ b/Engine/Source/Runtime/TypedElementFramework/Public/Elements/Framework/TypedElementQueryBuilder.h
@@ -456,6 +456,10 @@ namespace UE::Editor::DataStorage::Queries
 	template<typename Function>
 	SubqueryCallback CreateSubqueryCallbackBinding(Function&& Callback);
 
+	// Merge the two input queries by copying the source query into the destination query. If the merge is not possible the original data will be
+	// preserved and an optional error message can be requested
+	TYPEDELEMENTFRAMEWORK_API bool MergeQueries(FQueryDescription& Destination, const FQueryDescription& Source, FText* OutErrorMessage = nullptr);
+
 } // namespace UE::Editor::DataStorage::Queries
 
 #include "Elements/Framework/TypedElementQueryBuilder.inl"
diff --git a/Engine/Source/Runtime/TypedElementFramework/Tests/Elements/Framework/TypedElementQueryMergeTests.cpp b/Engine/Source/Runtime/TypedElementFramework/Tests/Elements/Framework/TypedElementQueryMergeTests.cpp
new file mode 100644
--- /dev/null
+++ b/Engine/Source/Runtime/TypedElementFramework/Tests/Elements/Framework/TypedElementQueryMergeTests.cpp
@@ -0,0 +1,435 @@
+// Copyright Epic Games, Inc. All Rights Reserved.
+
+#include "Elements/Framework/TypedElementTestColumns.h"
+#include "Elements/Interfaces/TypedElementDataStorageInterface.h"
+#if WITH_TESTS
+
+#include "Elements/Common/EditorDataStorageFeatures.h"
+#include "Elements/Framework/TypedElementQueryBuilder.h"
+#include "Elements/Interfaces/TypedElementDataStorageInterface.h"
+
+#include "Misc/AutomationTest.h"
+
+namespace UE::Editor::DataStorage::Tests
+{
+	using namespace UE::Editor::DataStorage::Queries;
+	
+	BEGIN_DEFINE_SPEC(QueryMergeTestFixture, "Editor.DataStorage.QueryMerge", EAutomationTestFlags::EditorContext | EAutomationTestFlags::EngineFilter)
+		ICoreProvider* TedsInterface = nullptr;
+		const FName TestTableName = TEXT("TestTable_QueryMergeTest");
+
+		TableHandle TestTable;
+		TArray<RowHandle> Rows;
+		TArray<QueryHandle> QueryHandles;
+		FName Identifier1;
+		FName Identifier2;
+		FName Identifier3;
+
+		TableHandle RegisterTestTable() const
+		{
+			const TableHandle Table = TedsInterface->FindTable(TestTableName);
+				
+			if (Table != InvalidTableHandle)
+			{
+				return Table;
+			}
+				
+			return TedsInterface->RegisterTable(
+			{
+				FTestColumnA::StaticStruct(),
+				FTestTagColumnA::StaticStruct()
+			},
+			TestTableName);
+		}
+
+		RowHandle CreateTestRow(TableHandle InTableHandle)
+		{
+			RowHandle Row = TedsInterface->AddRow(InTableHandle);
+			Rows.Add(Row);
+			return Row;
+		}
+
+		QueryHandle RegisterQuery(FQueryDescription&& Query)
+		{
+			QueryHandle QueryHandle = TedsInterface->RegisterQuery(MoveTemp(Query));
+			QueryHandles.Add(QueryHandle);
+			return QueryHandle;
+		}
+
+		void RegisterTestTableRows()
+		{
+			RowHandle RowA = CreateTestRow(TestTable);
+			TedsInterface->AddColumns<FTestColumnB, FTestColumnC, FTestColumnD>(RowA);
+			TedsInterface->AddColumns<FTestTagColumnB, FTestTagColumnC, FTestTagColumnD>(RowA);
+			TedsInterface->AddColumn<FTestColumnDynamic>(RowA, Identifier1);
+			TedsInterface->AddColumn<FTestColumnDynamic>(RowA, Identifier2);
+			TedsInterface->AddColumn<FTestColumnDynamic>(RowA, Identifier3);
+
+			RowHandle RowB = CreateTestRow(TestTable);
+			TedsInterface->AddColumns<FTestColumnB, FTestColumnD>(RowB);
+			TedsInterface->AddColumns<FTestTagColumnB, FTestTagColumnC>(RowB);
+			TedsInterface->AddColumn<FTestColumnDynamic>(RowB, Identifier1);
+
+			RowHandle RowC = CreateTestRow(TestTable);
+			TedsInterface->AddColumn<FTestColumnDynamic>(RowC, Identifier2);
+
+			RowHandle RowD = CreateTestRow(TestTable);
+			TedsInterface->AddColumns<FTestColumnC, FTestColumnD>(RowD);
+			TedsInterface->AddColumns<FTestTagColumnD>(RowD);
+			TedsInterface->AddColumn<FTestColumnDynamic>(RowD, Identifier3);
+
+			RowHandle RowE = CreateTestRow(TestTable);
+			TedsInterface->AddColumns<FTestColumnB, FTestColumnC, FTestColumnD>(RowE);
+			TedsInterface->AddColumn<FTestColumnDynamic>(RowE, Identifier2);
+			TedsInterface->AddColumn<FTestColumnDynamic>(RowE, Identifier3);
+		}
+
+		TSet<RowHandle> RegisterAndRunQuery(FQueryDescription& Query)
+		{
+			QueryHandle QueryHandle = RegisterQuery(MoveTemp(Query));
+			TSet<RowHandle> OutputRows;
+			
+			TedsInterface->RunQuery(QueryHandle, CreateDirectQueryCallbackBinding([&OutputRows](const IDirectQueryContext& Context, const RowHandle* RowArray)
+			{
+				OutputRows.Append(Context.GetRowHandles());
+			}));
+
+			return OutputRows;
+		}
+
+		void TestQueryMerge_Success(FQueryDescription& Query1, FQueryDescription& Query2)
+		{
+			// Copy Query1 since the merge operates directly on the first query
+			FQueryDescription MergedQuery = Query1;
+			FText Error;
+
+			bool bMerged = MergeQueries(MergedQuery, Query2, &Error);
+				
+			TestTrue(TEXT("Expected successful merge between simple queries"), bMerged);
+
+			TSet<RowHandle> Rows1 = RegisterAndRunQuery(Query1);
+			TSet<RowHandle> Rows2 = RegisterAndRunQuery(Query2);
+			TSet<RowHandle> RowsMerged = RegisterAndRunQuery(MergedQuery);
+
+			TestTrue(TEXT("Expected Query1 to have valid results"), !Rows1.IsEmpty());
+			TestTrue(TEXT("Expected Query2 to have valid results"), !Rows2.IsEmpty());
+
+			TSet<RowHandle> RowsMerged_Expected = Rows1.Intersect(Rows2);
+			bool bResultsSame = RowsMerged.Difference(RowsMerged_Expected).IsEmpty();
+
+			TestTrue(TEXT("Expected merge query output to be the same as the intersection of the individual query outputs"), bResultsSame);
+		}
+	
+		void TestQueryMerge_Failure(FQueryDescription& Query1, FQueryDescription& Query2, const FString& InvalidReason)
+		{
+			// Copy Query1 since the merge operates directly on the first query
+			FQueryDescription MergedQuery = Query1;
+			FText Error;
+
+			bool bMerged = MergeQueries(MergedQuery, Query2, &Error);
+
+			TestFalse(FString::Printf(TEXT("Expected merge to fail because of %s"), *InvalidReason), bMerged);
+		}
+
+	
+	END_DEFINE_SPEC(QueryMergeTestFixture)
+
+	void QueryMergeTestFixture::Define()
+	{
+		BeforeEach([this]()
+		{
+			Identifier1 = FName("Identifier1");
+			Identifier2 = FName("Identifier2");
+			Identifier3 = FName("Identifier3");
+			TedsInterface = GetMutableDataStorageFeature<ICoreProvider>(StorageFeatureName);
+			TestTrue("", TedsInterface != nullptr);
+			
+			TestTable = RegisterTestTable();
+			RegisterTestTableRows();
+		});
+			
+		Describe("QueryMerge", [this]
+		{
+			Describe("ValidMerge", [this]
+			{
+				Describe("Where()", [this]
+				{
+					It(TEXT("Simple"), EAsyncExecution::TaskGraphMainTick, [this]()
+					{
+						FQueryDescription Query1 =
+							Select()
+							.Where()
+								.All<FTestColumnA, FTestTagColumnA>()
+							.Compile();
+
+						FQueryDescription Query2 =
+							Select()
+							.Where()
+								.All<FTestColumnB, FTestTagColumnB>()
+							.Compile();
+
+						TestQueryMerge_Success(Query1, Query2);
+					});
+					
+					It(TEXT("Complex"), EAsyncExecution::TaskGraphMainTick, [this]()
+					{
+						FQueryDescription Query1 =
+							Select()
+							.Where()
+								.Any<FTestColumnB, FTestColumnD>()
+								.None<FTestColumnC>()
+							.Compile();
+
+						FQueryDescription Query2 =
+							Select()
+							.Where()
+								.All<FTestColumnB>()
+								.None<FTestTagColumnC>()
+							.Compile();
+
+						FQueryDescription MergedQuery = Query1;
+						FText Error;
+
+						TestQueryMerge_Success(Query1, Query2);
+					});
+				});
+				
+				Describe("Select()", [this]
+				{
+					It(TEXT("Simple"), EAsyncExecution::TaskGraphMainTick, [this]()
+					{
+						FQueryDescription Query1 =
+							Select()
+								.ReadOnly<FTestColumnA>()
+							.Compile();
+
+						FQueryDescription Query2 =
+							Select()
+								.ReadOnly<FTestColumnB>()
+							.Compile();
+
+						TestQueryMerge_Success(Query1, Query2);
+					});
+						
+					It(TEXT("Complex"), EAsyncExecution::TaskGraphMainTick, [this]()
+					{
+						FQueryDescription Query1 =
+							Select()
+								.ReadOnly<FTestColumnA>()
+								.ReadWrite<FTestColumnB>()
+							.Compile();
+
+						FQueryDescription Query2 =
+							Select()
+								.ReadOnly<FTestColumnB>()
+								.ReadWrite<FTestColumnD>()
+							.Compile();
+						
+						TestQueryMerge_Success(Query1, Query2);
+					});
+				});
+				
+				Describe("FConditions", [this]
+                {
+					It(TEXT("Simple"), EAsyncExecution::TaskGraphMainTick, [this]()
+					{
+						FQueryDescription Query1 =
+							Select()
+							.Where(TColumn<FTestColumnA>() && TColumn<FTestTagColumnA>())
+							.Compile();
+
+						FQueryDescription Query2 =
+							Select()
+							.Where(TColumn<FTestColumnB>() && TColumn<FTestTagColumnB>())
+							.Compile();
+
+						TestQueryMerge_Success(Query1, Query2);
+					});
+
+					It(TEXT("Complex"), EAsyncExecution::TaskGraphMainTick, [this]()
+					{
+						FQueryDescription Query1 =
+							Select()
+							.Where(TColumn<FTestColumnA>() && TColumn<FTestTagColumnA>())
+							.Compile();
+
+						FQueryDescription Query2 =
+							Select()
+							.Where(TColumn<FTestColumnB>() && TColumn<FTestTagColumnB>())
+							.Compile();
+
+						TestQueryMerge_Success(Query1, Query2);
+					});
+					
+					It(TEXT("Simple::Dynamic"), EAsyncExecution::TaskGraphMainTick, [this]()
+					{
+						FQueryDescription Query1 =
+							Select()
+							.Where(TColumn<FTestColumnDynamic>(Identifier1))
+							.Compile();
+
+						FQueryDescription Query2 =
+							Select()
+							.Where(TColumn<FTestColumnDynamic>(Identifier2))
+							.Compile();
+
+						TestQueryMerge_Success(Query1, Query2);
+					});
+					
+					It(TEXT("Complex::Dynamic"), EAsyncExecution::TaskGraphMainTick, [this]()
+					{
+						FQueryDescription Query1 =
+							Select()
+							.Where(TColumn<FTestColumnB>() ||  TColumn<FTestColumnDynamic>(Identifier2))
+							.Compile();
+
+						FQueryDescription Query2 =
+							Select()
+							.Where(TColumn<FTestColumnDynamic>(Identifier1) &&  TColumn<FTestColumnDynamic>(Identifier2))
+							.Compile();
+
+						TestQueryMerge_Success(Query1, Query2);
+					});
+                });
+                				
+				Describe("Combination", [this]
+				{
+					It(TEXT("Select()+Where()"), EAsyncExecution::TaskGraphMainTick, [this]()
+					{
+						FQueryDescription Query1 =
+							Select()
+								.ReadOnly<FTestColumnA>()
+							.Where()
+								.All<FTestColumnB, FTestTagColumnA>()
+							.Compile();
+
+						FQueryDescription Query2 =
+							Select()
+								.ReadWrite<FTestColumnDynamic>(Identifier1)
+						.Where()
+							.Any<FTestColumnB, FTestColumnC, FTestColumnD>()
+							.Compile();
+						
+						TestQueryMerge_Success(Query1, Query2);
+					});
+
+					It(TEXT("Select()+FConditions"), EAsyncExecution::TaskGraphMainTick, [this]()
+					{
+						FQueryDescription Query1 =
+							Select()
+								.ReadOnly<FTestColumnA>()
+							.Where(TColumn<FTestTagColumnB>() || TColumn<FTestColumnC>())
+							.Compile();
+
+						FQueryDescription Query2 =
+							Select()
+								.ReadWrite<FTestColumnDynamic>(Identifier1)
+							.Where(TColumn<FTestColumnDynamic>(Identifier2) || TColumn<FTestColumnDynamic>(Identifier3))
+							.Compile();
+						
+						TestQueryMerge_Success(Query1, Query2);
+					});				
+				});
+			});
+			
+			Describe("InvalidMerge", [this]
+			{
+				It(TEXT("Where()"), EAsyncExecution::TaskGraphMainTick, [this]()
+				{
+					FQueryDescription Query1 =
+							Select()
+							.Where()
+								.All<FTestColumnA>()
+							.Compile();
+
+					FQueryDescription Query2 =
+							Select()
+							.Where()
+								.None<FTestColumnA>()
+							.Compile();
+
+					TestQueryMerge_Failure(Query1, Query2, TEXT("Conflicting conditions .All<FTestColumnA>() and .None<FTestColumnA>()"));
+				});
+
+				It(TEXT("Where()::Dynamic"), EAsyncExecution::TaskGraphMainTick, [this]()
+				{
+					FQueryDescription Query1 =
+							Select()
+							.Where()
+								.All<FTestColumnDynamic>(Identifier1)
+							.Compile();
+
+					FQueryDescription Query2 =
+							Select()
+							.Where()
+								.None<FTestColumnDynamic>(Identifier1)
+							.Compile();
+
+					TestQueryMerge_Failure(Query1, Query2,
+						TEXT("Conflicting conditions .All<FTestColumnDynamic>(Identifier1) and .None<FTestColumnDynamic>(Identifier1)"));
+				});
+
+				It(TEXT("Select()+Where()"), EAsyncExecution::TaskGraphMainTick, [this]()
+				{
+					FQueryDescription Query1 =
+							Select()
+								.ReadWrite<FTestColumnA, FTestColumnB>()
+							.Where()
+							.Compile();
+
+					FQueryDescription Query2 =
+							Select()
+								.ReadOnly<FTestColumnC>()
+							.Where()
+								.None<FTestColumnB>()
+								.All<FTestColumnA>()
+							.Compile();
+
+					TestQueryMerge_Failure(Query1, Query2, TEXT("Conflicting conditions .ReadWrite<FTestColumnB>() and .None<FTestColumnB>()"));
+				});
+				
+				It(TEXT("Select()+Where()::Dynamic"), EAsyncExecution::TaskGraphMainTick, [this]()
+				{
+					FQueryDescription Query1 =
+							Select()
+								.ReadWrite<FTestColumnDynamic>(Identifier2)
+							.Where()
+							.Compile();
+
+					FQueryDescription Query2 =
+							Select()
+								.ReadOnly<FTestColumnC>()
+							.Where()
+								.None<FTestColumnDynamic>(Identifier2)
+								.All<FTestColumnDynamic>(Identifier1)
+							.Compile();
+
+					TestQueryMerge_Failure(Query1, Query2, TEXT("Conflicting conditions .ReadWrite<FTestColumnB>() and .None<FTestColumnB>()"));
+				});
+			});
+			
+		});
+		
+		AfterEach([this]()
+		{
+			for (RowHandle Row : Rows)
+			{
+				TedsInterface->RemoveRow(Row);
+			}
+			Rows.Empty(Rows.Num());
+
+			for (QueryHandle QueryHandle : QueryHandles)
+			{
+				TedsInterface->UnregisterQuery(QueryHandle);
+			}
+			QueryHandles.Empty(QueryHandles.Num());
+
+			TedsInterface = nullptr;
+			Identifier1 = NAME_None;
+			Identifier2 = NAME_None;
+			Identifier3 = NAME_None;
+		});
+	}
+} // namespace UE::Editor::DataStorage::Tests
+
+#endif // WITH_TESTS
