TEDS:
- Fix incorrect logic for empty query conditions in the FConditions merge operators
- Fix incorrect merge logic for checking whether 2 queries have incompatible syntax to account for empty FConditions

diff --git a/Engine/Source/Runtime/TypedElementFramework/Private/Elements/Common/TypedElementQueryConditions.cpp b/Engine/Source/Runtime/TypedElementFramework/Private/Elements/Common/TypedElementQueryConditions.cpp
--- a/Engine/Source/Runtime/TypedElementFramework/Private/Elements/Common/TypedElementQueryConditions.cpp
+++ b/Engine/Source/Runtime/TypedElementFramework/Private/Elements/Common/TypedElementQueryConditions.cpp
@@ -619,6 +619,15 @@ namespace UE::Editor::DataStorage::Queries
 
 	FConditions operator&&(const FConditions& Lhs, const FConditions& Rhs)
 	{
+		if (Lhs.IsEmpty())
+		{
+			return Rhs;
+		}
+		else if (Rhs.IsEmpty())
+		{
+			return Lhs;
+		}
+		
 		FConditions Result;
 		Result.Tokens[Result.TokenCount++] = FConditions::Token::ScopeOpen;
 		FConditions::AppendQuery(Result, Lhs);
@@ -645,10 +654,15 @@ namespace UE::Editor::DataStorage::Queries
 	FConditions operator&&(FColumnBase Lhs, const FConditions& Rhs)
 	{
 		FConditions Result(Lhs);
-		Result.Tokens[Result.TokenCount++] = FConditions::Token::And;
-		Result.Tokens[Result.TokenCount++] = FConditions::Token::ScopeOpen;
-		FConditions::AppendQuery(Result, Rhs);
-		Result.Tokens[Result.TokenCount++] = FConditions::Token::ScopeClose;
+
+		if (!Rhs.IsEmpty())
+		{
+			Result.Tokens[Result.TokenCount++] = FConditions::Token::And;
+			Result.Tokens[Result.TokenCount++] = FConditions::Token::ScopeOpen;
+			FConditions::AppendQuery(Result, Rhs);
+			Result.Tokens[Result.TokenCount++] = FConditions::Token::ScopeClose;
+		}
+		
 		return Result;
 	}
 
@@ -663,6 +677,15 @@ namespace UE::Editor::DataStorage::Queries
 
 	FConditions operator||(const FConditions& Lhs, const FConditions& Rhs)
 	{
+		if (Lhs.IsEmpty())
+		{
+			return Rhs;
+		}
+		else if (Rhs.IsEmpty())
+		{
+			return Lhs;
+		}
+		
 		FConditions Result;
 		Result.Tokens[Result.TokenCount++] = FConditions::Token::ScopeOpen;
 		FConditions::AppendQuery(Result, Lhs);
@@ -689,10 +712,15 @@ namespace UE::Editor::DataStorage::Queries
 	FConditions operator||(FColumnBase Lhs, const FConditions& Rhs)
 	{
 		FConditions Result(Lhs);
-		Result.Tokens[Result.TokenCount++] = FConditions::Token::Or;
-		Result.Tokens[Result.TokenCount++] = FConditions::Token::ScopeOpen;
-		FConditions::AppendQuery(Result, Rhs);
-		Result.Tokens[Result.TokenCount++] = FConditions::Token::ScopeClose;
+		
+		if (!Rhs.IsEmpty())
+		{
+			Result.Tokens[Result.TokenCount++] = FConditions::Token::Or;
+			Result.Tokens[Result.TokenCount++] = FConditions::Token::ScopeOpen;
+			FConditions::AppendQuery(Result, Rhs);
+			Result.Tokens[Result.TokenCount++] = FConditions::Token::ScopeClose;
+
+		}
 		return Result;
 	}
 
diff --git a/Engine/Source/Runtime/TypedElementFramework/Private/Elements/Framework/TypedElementQueryBuilder.cpp b/Engine/Source/Runtime/TypedElementFramework/Private/Elements/Framework/TypedElementQueryBuilder.cpp
--- a/Engine/Source/Runtime/TypedElementFramework/Private/Elements/Framework/TypedElementQueryBuilder.cpp
+++ b/Engine/Source/Runtime/TypedElementFramework/Private/Elements/Framework/TypedElementQueryBuilder.cpp
@@ -919,7 +919,8 @@ namespace UE::Editor::DataStorage::Queries
 				return false;
 			}
 
-			if ((Destination.Conditions && !Source.Conditions) || (!Destination.Conditions && Source.Conditions))
+			// If one of the queries is using FConditions and the other one has anything in the old Where() syntax, they are incompatible
+			if ((Destination.Conditions && Source.ConditionTypes.Num()) || (Source.Conditions && Destination.ConditionTypes.Num()))
 			{
 				if (OutErrorMessage)
 				{
@@ -939,6 +940,25 @@ namespace UE::Editor::DataStorage::Queries
 				return false;
 			}
 
+			// If both queries have a condition in .Any<>, they cannot be merged as combining .Any<> conditions results in an OR operation
+			// e.g .Any<ColumnA, ColumnB> + .Any<ColumnC, ColumnD> results in .Any<ColumnA, ColumnB, ColumnC, ColumnD> which is incorrect since
+			// the expected result is (ColumnA || ColumnB) && (ColumnC || ColumnD)
+			auto HasAnyCondition = [](const FQueryDescription& QueryDescription)
+			{
+				return QueryDescription.ConditionTypes.FindByKey(FQueryDescription::EOperatorType::SimpleAny)
+					|| QueryDescription.DynamicConditionOperations.FindByKey(FQueryDescription::EOperatorType::SimpleAny);
+			};
+
+			if (HasAnyCondition(Source) && HasAnyCondition(Destination))
+			{
+				if (OutErrorMessage)
+				{
+					*OutErrorMessage = LOCTEXT("CannotMerge_Any",
+						"Cannot merge two queries if they both .Any<> conditions, use the FConditions syntax instead to merge them properly");
+				}
+				return false;
+			}
+			
 			return true;
 		}
 
@@ -995,10 +1015,10 @@ namespace UE::Editor::DataStorage::Queries
 		// Merge the Select() portion of the queries
 		Private::MergeSelect(Destination, Source);
 
-		// Since the FConditions cannot co-exist with the ConditionOperators, we only need to merge one of the two
-		if (Destination.Conditions)
+		// Since the FConditions cannot co-exist with the ConditionOperators because of the pre-check, we only need to merge one of the two
+		if (Destination.Conditions || Source.Conditions)
 		{
-			Destination.Conditions = Destination.Conditions.GetValue() && Source.Conditions.GetValue();
+			Destination.Conditions = Destination.Conditions.Get(FConditions()) && Source.Conditions.Get(FConditions());
 		}
 		else
 		{
