[TEDS const API] Remaining cleanup to change usages of the API to be consistent and use const when possible.

diff --git a/Engine/Plugins/Experimental/EditorDataStorageFeatures/Source/TedsDebugger/Private/QueryEditor/TedsQueryEditorModel.cpp b/Engine/Plugins/Experimental/EditorDataStorageFeatures/Source/TedsDebugger/Private/QueryEditor/TedsQueryEditorModel.cpp
--- a/Engine/Plugins/Experimental/EditorDataStorageFeatures/Source/TedsDebugger/Private/QueryEditor/TedsQueryEditorModel.cpp
+++ b/Engine/Plugins/Experimental/EditorDataStorageFeatures/Source/TedsDebugger/Private/QueryEditor/TedsQueryEditorModel.cpp
@@ -346,12 +346,12 @@ namespace UE::Editor::DataStorage::Debug::QueryEditor
 		TArray<const UScriptStruct*> Columns;
 		TArray<const UScriptStruct*> Tags;
 
-		UScriptStruct* ColumnType = FColumn::StaticStruct();
-		UScriptStruct* TagType = FTag::StaticStruct();
+		const UScriptStruct* ColumnType = FColumn::StaticStruct();
+		const UScriptStruct* TagType = FTag::StaticStruct();
 		// Not sure if there is a faster way to do this.  Would be nice to iterate only the derived classes
-		for(TObjectIterator< UScriptStruct > It; It; ++It)
+		for(TObjectIterator< const UScriptStruct > It; It; ++It)
 		{
-			UScriptStruct* Struct = *It;
+			const UScriptStruct* Struct = *It;
 
 			if (Struct->IsChildOf(ColumnType) && Struct != ColumnType)
 			{
diff --git a/Engine/Plugins/Experimental/EditorDataStorageFeatures/Source/TedsOutliner/Private/TedsOutlinerImpl.cpp b/Engine/Plugins/Experimental/EditorDataStorageFeatures/Source/TedsOutliner/Private/TedsOutlinerImpl.cpp
--- a/Engine/Plugins/Experimental/EditorDataStorageFeatures/Source/TedsOutliner/Private/TedsOutlinerImpl.cpp
+++ b/Engine/Plugins/Experimental/EditorDataStorageFeatures/Source/TedsOutliner/Private/TedsOutlinerImpl.cpp
@@ -225,7 +225,7 @@ void FTedsOutlinerImpl::CreateFilterQueries()
 		// Grab all UStruct types to see if they derive from FColumn or FTag
 		ForEachObjectOfClass(UScriptStruct::StaticClass(), [&](UObject* Obj)
 		{
-			if (UScriptStruct* Struct = Cast<UScriptStruct>(Obj))
+			if (const UScriptStruct* Struct = Cast<const UScriptStruct>(Obj))
 			{
 				if (Struct->IsChildOf(TedsColumn) || Struct->IsChildOf(TedsTag))
 				{
diff --git a/Engine/Plugins/Experimental/EditorDataStorageFeatures/Source/TedsRevisionControl/Private/Processors/RevisionControlProcessors.cpp b/Engine/Plugins/Experimental/EditorDataStorageFeatures/Source/TedsRevisionControl/Private/Processors/RevisionControlProcessors.cpp
--- a/Engine/Plugins/Experimental/EditorDataStorageFeatures/Source/TedsRevisionControl/Private/Processors/RevisionControlProcessors.cpp
+++ b/Engine/Plugins/Experimental/EditorDataStorageFeatures/Source/TedsRevisionControl/Private/Processors/RevisionControlProcessors.cpp
@@ -22,7 +22,7 @@ namespace UE::Editor::RevisionControl::Private
 	extern FAutoConsoleVariableRef CVarAutoPopulateState;
 
 	// Update the overlay color for all rows with InColumn
-	void UpdateSCCOverlayStates(UScriptStruct* InColumn)
+	void UpdateSCCOverlayStates(const UScriptStruct* InColumn)
 	{
 		using namespace UE::Editor::DataStorage;
 		ICoreProvider* DataStorage = GetMutableDataStorageFeature<ICoreProvider>(StorageFeatureName);
diff --git a/Engine/Source/Runtime/TypedElementFramework/Private/Elements/Framework/TypedElementQueryBuilder.cpp b/Engine/Source/Runtime/TypedElementFramework/Private/Elements/Framework/TypedElementQueryBuilder.cpp
--- a/Engine/Source/Runtime/TypedElementFramework/Private/Elements/Framework/TypedElementQueryBuilder.cpp
+++ b/Engine/Source/Runtime/TypedElementFramework/Private/Elements/Framework/TypedElementQueryBuilder.cpp
@@ -21,7 +21,7 @@ namespace UE::Editor::DataStorage::Queries
 
 	const UScriptStruct* TypeOptional(FTopLevelAssetPath Name)
 	{
-		return static_cast<UScriptStruct*>(StaticFindObject(UScriptStruct::StaticClass(), Name, EFindObjectFlags::ExactClass));
+		return static_cast<const UScriptStruct*>(StaticFindObject(UScriptStruct::StaticClass(), Name, EFindObjectFlags::ExactClass));
 	}
 
 	const UScriptStruct* operator""_Type(const char* Name, std::size_t NameSize)
diff --git a/Engine/Source/Runtime/TypedElementFramework/Private/Elements/Framework/TypedElementtAttributeBindingText.cpp b/Engine/Source/Runtime/TypedElementFramework/Private/Elements/Framework/TypedElementtAttributeBindingText.cpp
--- a/Engine/Source/Runtime/TypedElementFramework/Private/Elements/Framework/TypedElementtAttributeBindingText.cpp
+++ b/Engine/Source/Runtime/TypedElementFramework/Private/Elements/Framework/TypedElementtAttributeBindingText.cpp
@@ -36,19 +36,19 @@ namespace UE::Editor::DataStorage
 				{
 					if (Property.Value.Property.IsBound())
 					{
-						const UScriptStruct* ColumnType = Property.Value.Property.GetObjectTypeInfo();
+						TWeakObjectPtr<const UScriptStruct> ColumnType = Property.Value.Property.GetObjectTypeInfo();
 
 						if(Property.Value.DynamicColumnIdentifier != NAME_None)
 						{
 							ColumnType = Storage->GenerateDynamicColumn(FDynamicColumnDescription{
-								.TemplateType = ColumnType,
+								.TemplateType = ColumnType.Get(),
 								.Identifier = Property.Value.DynamicColumnIdentifier
 							});
 						}
 
-						if(ColumnType)
+						if(ColumnType.IsValid())
 						{
-							if (void* Column = Storage->GetColumnData(Row, ColumnType))
+							if (void* Column = Storage->GetColumnData(Row, ColumnType.Get()))
 							{
 								// We use the original column type (Property.GetObjectTypeInfo()) instead of the dynamic ColumnType
 								// so that the type info in the property matches
diff --git a/Engine/Source/Runtime/TypedElementFramework/Private/Elements/Interfaces/TypedElementDataStorageUiInterface.cpp b/Engine/Source/Runtime/TypedElementFramework/Private/Elements/Interfaces/TypedElementDataStorageUiInterface.cpp
--- a/Engine/Source/Runtime/TypedElementFramework/Private/Elements/Interfaces/TypedElementDataStorageUiInterface.cpp
+++ b/Engine/Source/Runtime/TypedElementFramework/Private/Elements/Interfaces/TypedElementDataStorageUiInterface.cpp
@@ -49,7 +49,7 @@ bool FTypedElementWidgetConstructor::Initialize_Internal(const UE::Editor::DataS
 		FName Identifier = UE::Editor::DataStorage::ColumnUtils::GetDynamicColumnIdentifier(Column.Get());
 		if (!Identifier.IsNone())
 		{
-			MatchedDynamicTemplates.Emplace(Cast<UScriptStruct>(Column->GetSuperStruct()), Identifier);
+			MatchedDynamicTemplates.Emplace(Cast<const UScriptStruct>(Column->GetSuperStruct()), Identifier);
 		}
 	}
 	
diff --git a/Engine/Source/Runtime/TypedElementFramework/Public/Elements/Framework/TypedElementAttributeBindingProperty.h b/Engine/Source/Runtime/TypedElementFramework/Public/Elements/Framework/TypedElementAttributeBindingProperty.h
--- a/Engine/Source/Runtime/TypedElementFramework/Public/Elements/Framework/TypedElementAttributeBindingProperty.h
+++ b/Engine/Source/Runtime/TypedElementFramework/Public/Elements/Framework/TypedElementAttributeBindingProperty.h
@@ -99,13 +99,13 @@ namespace UE::Editor::DataStorage
 			const PropertyType& Get(const ObjectType& Object) const;
 
 			// Get the bound property for this specified object ptr by providing type information about the object
-			PropertyType& Get(void* Object, const UScriptStruct* Type) const;
+			PropertyType& Get(void* Object, TWeakObjectPtr<const UScriptStruct> Type) const;
 
 			// Get the bound property for this specified object ptr by providing type information about the object
-			const PropertyType& Get(const void* Object, const UScriptStruct* Type) const;
+			const PropertyType& Get(const void* Object, TWeakObjectPtr<const UScriptStruct> Type) const;
 
 			// Returns the stored type information or null if no type was set.
-			const UScriptStruct* GetObjectTypeInfo() const;
+			TWeakObjectPtr<const UScriptStruct> GetObjectTypeInfo() const;
 
 			// Whether or not this property has been bound.
 			bool IsBound() const;
@@ -121,7 +121,7 @@ namespace UE::Editor::DataStorage
 			InternalPropertyType InternalProperty;
 
 			// The type of the object we are bound to
-			UScriptStruct* ObjectTypeInfo = nullptr;
+			TWeakObjectPtr<const UScriptStruct> ObjectTypeInfo = nullptr;
 		};
 	
 		//
@@ -222,33 +222,33 @@ namespace UE::Editor::DataStorage
 		}
 
 		template <typename PropertyType>
-		PropertyType& TProperty<PropertyType>::Get(void* Object, const UScriptStruct* Type) const
+		PropertyType& TProperty<PropertyType>::Get(void* Object, TWeakObjectPtr<const UScriptStruct> Type) const
 		{
 			checkf(Object, TEXT("Nullptr pointer provided to Object while trying to retrieve a property value."));
-			checkf(Type, TEXT("Nullptr pointer provided to Type while trying to retrieve a property value."));
+			checkf(Type.IsValid(), TEXT("Nullptr pointer provided to Type while trying to retrieve a property value."));
 			checkf(IsBound(), TEXT("Attempting to retrieve the value of object type (%s) from a property that wasn't bound."),
 				*Type->GetFName().ToString());
 			ensureMsgf(Type == ObjectTypeInfo, TEXT("Provided object type (%s) did not match bound object type (%s)."),
-				*Type->GetFName().ToString(), ObjectTypeInfo ? *ObjectTypeInfo->GetFName().ToString() : TEXT("<no type>"));
+				*Type->GetFName().ToString(), ObjectTypeInfo.IsValid() ? *ObjectTypeInfo->GetFName().ToString() : TEXT("<no type>"));
 
 			return *Visit([Object](auto&& Prop) { return Prop.Get(Object); }, InternalProperty);
 		}
 
 		template <typename PropertyType>
-		const PropertyType& TProperty<PropertyType>::Get(const void* Object, const UScriptStruct* Type) const
+		const PropertyType& TProperty<PropertyType>::Get(const void* Object, TWeakObjectPtr<const UScriptStruct> Type) const
 		{
 			checkf(Object, TEXT("Nullptr pointer provided to Object while trying to retrieve a property value."));
-			checkf(Type, TEXT("Nullptr pointer provided to Type while trying to retrieve a property value."));
+			checkf(Type.IsValid(), TEXT("Nullptr pointer provided to Type while trying to retrieve a property value."));
 			checkf(IsBound(), TEXT("Attempting to retrieve the value of object type (%s) from a property that wasn't bound."),
 				*Type->GetFName().ToString());
 			ensureMsgf(Type == ObjectTypeInfo, TEXT("Provided object type (%s) did not match bound object type (%s)."),
-				*Type->GetFName().ToString(), ObjectTypeInfo ? *ObjectTypeInfo->GetFName().ToString() : TEXT("<no type>"));
+				*Type->GetFName().ToString(), ObjectTypeInfo.IsValid() ? *ObjectTypeInfo->GetFName().ToString() : TEXT("<no type>"));
 			
 			return *Visit([Object](auto&& Prop) { return Prop.Get(Object); }, InternalProperty);
 		}
 
 		template <typename PropertyType>
-		const UScriptStruct* TProperty<PropertyType>::GetObjectTypeInfo() const
+		TWeakObjectPtr<const UScriptStruct> TProperty<PropertyType>::GetObjectTypeInfo() const
 		{
 			return ObjectTypeInfo;
 		}
diff --git a/Engine/Source/Runtime/TypedElementFramework/Public/Elements/Interfaces/TypedElementQueryStorageInterfaces.h b/Engine/Source/Runtime/TypedElementFramework/Public/Elements/Interfaces/TypedElementQueryStorageInterfaces.h
--- a/Engine/Source/Runtime/TypedElementFramework/Public/Elements/Interfaces/TypedElementQueryStorageInterfaces.h
+++ b/Engine/Source/Runtime/TypedElementFramework/Public/Elements/Interfaces/TypedElementQueryStorageInterfaces.h
@@ -543,7 +543,7 @@ namespace UE::Editor::DataStorage
 	template<typename ColumnType>
 	ColumnType& ICommonQueryWithEnvironmentContext::AddColumn(RowHandle Row, ColumnType&& Column)
 	{
-		UScriptStruct* TypeInfo = ColumnType::StaticStruct();
+		const UScriptStruct* TypeInfo = ColumnType::StaticStruct();
 
 		if constexpr (std::is_move_constructible_v<ColumnType>)
 		{
@@ -565,7 +565,7 @@ namespace UE::Editor::DataStorage
 	template <typename ColumnTypeTemplate>
 	ColumnTypeTemplate* ICommonQueryWithEnvironmentContext::AddColumn(RowHandle Row, const FName& Identifier)
 	{
-		UScriptStruct* TemplateType = ColumnTypeTemplate::StaticStruct();
+		const UScriptStruct* TemplateType = ColumnTypeTemplate::StaticStruct();
 		const FDynamicColumnDescription Description
 		{
 			.TemplateType = TemplateType,
@@ -590,7 +590,7 @@ namespace UE::Editor::DataStorage
 	template <typename ColumnTypeTemplate>
 	ColumnTypeTemplate& ICommonQueryWithEnvironmentContext::AddColumn(RowHandle Row, const FName& Identifier, ColumnTypeTemplate&& Column)
 	{
-		UScriptStruct* TemplateType = ColumnTypeTemplate::StaticStruct();
+		const UScriptStruct* TemplateType = ColumnTypeTemplate::StaticStruct();
 		const FDynamicColumnDescription Description
 		{
 			.TemplateType = TemplateType,
diff --git a/Engine/Source/Runtime/TypedElementFramework/Tests/Elements/Framework/TypedElementDataStoragePerformanceTestCommands.cpp b/Engine/Source/Runtime/TypedElementFramework/Tests/Elements/Framework/TypedElementDataStoragePerformanceTestCommands.cpp
--- a/Engine/Source/Runtime/TypedElementFramework/Tests/Elements/Framework/TypedElementDataStoragePerformanceTestCommands.cpp
+++ b/Engine/Source/Runtime/TypedElementFramework/Tests/Elements/Framework/TypedElementDataStoragePerformanceTestCommands.cpp
@@ -90,7 +90,7 @@ void UTest_PingPongBetweenPhaseFactory::RegisterQueries(UE::Editor::DataStorage:
 				TArrayView<const RowHandle> Rows(RowPtr, Context.GetRowCount());
 				TArrayView<FTest_PingPongPrePhys> PingPongAs(PingPongAPtr, Context.GetRowCount());
 
-				UScriptStruct* RemovedColumnStruct = FTest_PingPongPrePhys::StaticStruct();
+				const UScriptStruct* RemovedColumnStruct = FTest_PingPongPrePhys::StaticStruct();
 				for (int32 Index = 0, End = Context.GetRowCount(); Index < End; ++Index)
 				{
 					uint64 Value = PingPongAs[Index].Value;
@@ -112,7 +112,7 @@ void UTest_PingPongBetweenPhaseFactory::RegisterQueries(UE::Editor::DataStorage:
 				TArrayView<const RowHandle> Rows(RowPtr, Context.GetRowCount());
 				TArrayView<FTest_PingPongDurPhys> PingPongAs(PingPongBPtr, Context.GetRowCount());
 	
-				UScriptStruct* RemovedColumnStruct = FTest_PingPongDurPhys::StaticStruct();
+				const UScriptStruct* RemovedColumnStruct = FTest_PingPongDurPhys::StaticStruct();
 				for (int32 Index = 0, End = Context.GetRowCount(); Index < End; ++Index)
 				{
 					uint64 Value = PingPongAs[Index].Value;
@@ -134,7 +134,7 @@ void UTest_PingPongBetweenPhaseFactory::RegisterQueries(UE::Editor::DataStorage:
 			TArrayView<const RowHandle> Rows(RowPtr, Context.GetRowCount());
 			TArrayView<FTest_PingPongPostPhys> PingPongAs(PingPongBPtr, Context.GetRowCount());
 
-			UScriptStruct* RemovedColumnStruct = FTest_PingPongPostPhys::StaticStruct();
+			const UScriptStruct* RemovedColumnStruct = FTest_PingPongPostPhys::StaticStruct();
 			for (int32 Index = 0, End = Context.GetRowCount(); Index < End; ++Index)
 			{
 				uint64 Value = PingPongAs[Index].Value;
