Replace some usages of FORCEINLINE with inline in TypedElement modules.

For context about this change, see this doc: https://docs.google.com/document/d/1LjpT7NpzTecfXdFNyv5KM6sCu67827BRljNi1iaWVSE/edit?tab=t.0#heading=h.fvtz8gjukz7x

Includes modules: TypedElementFramework, TypedElementRuntime

#rnx

diff --git a/Engine/Source/Runtime/TypedElementFramework/Public/Elements/Framework/TypedElementData.h b/Engine/Source/Runtime/TypedElementFramework/Public/Elements/Framework/TypedElementData.h
--- a/Engine/Source/Runtime/TypedElementFramework/Public/Elements/Framework/TypedElementData.h
+++ b/Engine/Source/Runtime/TypedElementFramework/Public/Elements/Framework/TypedElementData.h
@@ -196,12 +196,12 @@ public:
 #endif	// UE_TYPED_ELEMENT_HAS_REFTRACKING
 	}
 
-	FORCEINLINE const FTypedElementId& GetId() const
+	inline const FTypedElementId& GetId() const
 	{
 		return Id;
 	}
 
-	FORCEINLINE FTypedElementReferenceId AddRef(const bool bCanTrackReference) const
+	inline FTypedElementReferenceId AddRef(const bool bCanTrackReference) const
 	{
 		FTypedElementReferenceId ReferenceId = INDEX_NONE;
 #if UE_TYPED_ELEMENT_HAS_REFCOUNTING
@@ -217,7 +217,7 @@ public:
 		return ReferenceId;
 	}
 
-	FORCEINLINE void ReleaseRef(const FTypedElementReferenceId InReferenceId) const
+	inline void ReleaseRef(const FTypedElementReferenceId InReferenceId) const
 	{
 #if UE_TYPED_ELEMENT_HAS_REFCOUNTING
 		checkSlow(RefCount > 0);
@@ -231,7 +231,7 @@ public:
 #endif	// UE_TYPED_ELEMENT_HAS_REFTRACKING
 	}
 
-	FORCEINLINE FTypedElementRefCount GetRefCount() const
+	inline FTypedElementRefCount GetRefCount() const
 	{
 #if UE_TYPED_ELEMENT_HAS_REFCOUNTING
 		return FPlatformAtomics::AtomicRead(&RefCount);
@@ -307,12 +307,12 @@ public:
 		Data = ElementDataType();
 	}
 
-	FORCEINLINE const ElementDataType& GetData() const
+	inline const ElementDataType& GetData() const
 	{
 		return Data;
 	}
 
-	FORCEINLINE ElementDataType& GetMutableData()
+	inline ElementDataType& GetMutableData()
 	{
 		return Data;
 	}
@@ -403,7 +403,7 @@ public:
 		ControlBlock = nullptr;
 	}
 
-	FORCEINLINE const FTypedElementId& GetId() const
+	inline const FTypedElementId& GetId() const
 	{
 		if (ControlBlock)
 		{
@@ -416,7 +416,7 @@ public:
 		return FTypedElementId::Unset;
 	}
 
-	FORCEINLINE FTypedElementInternalData* GetInternalData() const
+	inline FTypedElementInternalData* GetInternalData() const
 	{
 		if (ControlBlock)
 		{
@@ -569,18 +569,18 @@ public:
 		return InternalDataArray[InElementId];
 	}
 
-	static FORCEINLINE void SetStaticDataTypeId(const FTypedHandleTypeId InTypeId)
+	static inline void SetStaticDataTypeId(const FTypedHandleTypeId InTypeId)
 	{
 		checkSlow(ElementDataType::Private_RegisteredTypeId == 0);
 		ElementDataType::Private_RegisteredTypeId = InTypeId;
 	}
 
-	static FORCEINLINE FTypedHandleTypeId StaticDataTypeId()
+	static inline FTypedHandleTypeId StaticDataTypeId()
 	{
 		return ElementDataType::StaticTypeId();
 	}
 
-	static FORCEINLINE FName StaticDataTypeName()
+	static inline FName StaticDataTypeName()
 	{
 		return ElementDataType::StaticTypeName();
 	}
@@ -665,16 +665,16 @@ public:
 		return InternalDataArray[*InternalDataArrayIndexPtr];
 	}
 
-	static FORCEINLINE void SetStaticDataTypeId(const FTypedHandleTypeId InTypeId)
+	static inline void SetStaticDataTypeId(const FTypedHandleTypeId InTypeId)
 	{
 	}
 
-	static FORCEINLINE FTypedHandleTypeId StaticDataTypeId()
+	static inline FTypedHandleTypeId StaticDataTypeId()
 	{
 		return 0;
 	}
 
-	static FORCEINLINE FName StaticDataTypeName()
+	static inline FName StaticDataTypeName()
 	{
 		return FName();
 	}
diff --git a/Engine/Source/Runtime/TypedElementFramework/Public/Elements/Framework/TypedElementHandle.h b/Engine/Source/Runtime/TypedElementFramework/Public/Elements/Framework/TypedElementHandle.h
--- a/Engine/Source/Runtime/TypedElementFramework/Public/Elements/Framework/TypedElementHandle.h
+++ b/Engine/Source/Runtime/TypedElementFramework/Public/Elements/Framework/TypedElementHandle.h
@@ -79,7 +79,7 @@ public:
 		Private_DestroyReleaseRef();
 	}
 
-	FORCEINLINE explicit operator bool() const
+	inline explicit operator bool() const
 	{
 		return IsSet();
 	}
@@ -87,7 +87,7 @@ public:
 	/**
 	 * Has this handle been initialized to a valid element?
 	 */
-	FORCEINLINE bool IsSet() const
+	inline bool IsSet() const
 	{
 		return DataPtr != nullptr;
 	}
@@ -95,7 +95,7 @@ public:
 	/**
 	 * Release this handle and set it back to an empty state.
 	 */
-	FORCEINLINE void Release()
+	inline void Release()
 	{
 		Private_DestroyReleaseRef();
 	}
@@ -103,7 +103,7 @@ public:
 	/**
 	 * Get the ID that this element represents.
 	 */
-	FORCEINLINE const FTypedElementId& GetId() const
+	inline const FTypedElementId& GetId() const
 	{
 		return DataPtr
 			? DataPtr->GetId()
@@ -115,7 +115,7 @@ public:
 	 * @note This is not typically something you'd want to query outside of data access within an interface implementation.
 	 */
 	template <typename ElementDataType>
-	FORCEINLINE bool IsDataOfType() const
+	inline bool IsDataOfType() const
 	{
 		return GetId().GetTypeId() == ElementDataType::StaticTypeId();
 	}
@@ -153,7 +153,7 @@ public:
 	 * @note This is not typically something you'd want to use outside of data access within an interface implementation.
 	 */
 	template <typename ElementDataType>
-	FORCEINLINE const ElementDataType& GetDataChecked() const
+	inline const ElementDataType& GetDataChecked() const
 	{
 		checkf(DataPtr, TEXT("Element handle data is null!"));
 		checkf(IsDataOfType<ElementDataType>(), TEXT("Element handle data type is '%d', but '%d' (%s) was requested!"), GetId().GetTypeId(), ElementDataType::StaticTypeId(), *ElementDataType::StaticTypeName().ToString());
@@ -189,12 +189,12 @@ public:
 		}
 	}
 
-	FORCEINLINE friend bool operator==(const FTypedElementHandle& InLHS, const FTypedElementHandle& InRHS)
+	inline friend bool operator==(const FTypedElementHandle& InLHS, const FTypedElementHandle& InRHS)
 	{
 		return InLHS.DataPtr == InRHS.DataPtr;
 	}
 
-	FORCEINLINE friend bool operator!=(const FTypedElementHandle& InLHS, const FTypedElementHandle& InRHS)
+	inline friend bool operator!=(const FTypedElementHandle& InLHS, const FTypedElementHandle& InRHS)
 	{
 		return !(InLHS == InRHS);
 	}
@@ -204,18 +204,18 @@ public:
 		return GetTypeHash(InElementHandle.GetId());
 	}
 
-	FORCEINLINE void Private_InitializeNoRef(const FTypedElementInternalData& InData)
+	inline void Private_InitializeNoRef(const FTypedElementInternalData& InData)
 	{
 		DataPtr = &InData;
 	}
 
-	FORCEINLINE void Private_InitializeAddRef(const FTypedElementInternalData& InData)
+	inline void Private_InitializeAddRef(const FTypedElementInternalData& InData)
 	{
 		Private_InitializeNoRef(InData);
 		RegisterRef();
 	}
 
-	FORCEINLINE void Private_DestroyNoRef()
+	inline void Private_DestroyNoRef()
 	{
 		DataPtr = nullptr;
 #if UE_TYPED_ELEMENT_HAS_REFTRACKING
@@ -223,19 +223,19 @@ public:
 #endif	// UE_TYPED_ELEMENT_HAS_REFTRACKING
 	}
 
-	FORCEINLINE void Private_DestroyReleaseRef()
+	inline void Private_DestroyReleaseRef()
 	{
 		UnregisterRef();
 		Private_DestroyNoRef();
 	}
 
-	FORCEINLINE const FTypedElementInternalData* Private_GetInternalData() const
+	inline const FTypedElementInternalData* Private_GetInternalData() const
 	{
 		return DataPtr;
 	}
 
 private:
-	FORCEINLINE void RegisterRef()
+	inline void RegisterRef()
 	{
 		if (DataPtr)
 		{
@@ -246,7 +246,7 @@ private:
 		}
 	}
 
-	FORCEINLINE void UnregisterRef()
+	inline void UnregisterRef()
 	{
 		if (DataPtr)
 		{
@@ -301,7 +301,7 @@ public:
 		Private_DestroyReleaseRef();
 	}
 
-	FORCEINLINE explicit operator bool() const
+	inline explicit operator bool() const
 	{
 		return IsSet();
 	}
@@ -309,7 +309,7 @@ public:
 	/**
 	 * Has this element been initialized to a valid handle and interface?
 	 */
-	FORCEINLINE bool IsSet() const
+	inline bool IsSet() const
 	{
 		return FTypedElementHandle::IsSet()
 			&& InterfacePtr;
@@ -318,7 +318,7 @@ public:
 	/**
 	 * Release this element and set it back to an empty state.
 	 */
-	FORCEINLINE void Release()
+	inline void Release()
 	{
 		Private_DestroyReleaseRef();
 	}
@@ -326,7 +326,7 @@ public:
 	/**
 	 * Attempt to access the interface stored within this element, returning null if it isn't set.
 	 */
-	FORCEINLINE BaseInterfaceType* GetInterface() const
+	inline BaseInterfaceType* GetInterface() const
 	{
 		return InterfacePtr;
 	}
@@ -335,20 +335,20 @@ public:
 	 * Attempt to access the interface stored within this element, asserting if it isn't set.
 	 */
 	template <typename U = BaseInterfaceType, std::enable_if_t<!std::is_void<U>::value>* = nullptr>
-	FORCEINLINE U& GetInterfaceChecked() const
+	inline U& GetInterfaceChecked() const
 	{
 		static_assert(std::is_same<U, BaseInterfaceType>::value, "Don't explicitly specify U!");
 		checkf(InterfacePtr, TEXT("Interface is null!"));
 		return *InterfacePtr;
 	}
 
-	FORCEINLINE friend bool operator==(const TTypedElementBase& InLHS, const TTypedElementBase& InRHS)
+	inline friend bool operator==(const TTypedElementBase& InLHS, const TTypedElementBase& InRHS)
 	{
 		return static_cast<const FTypedElementHandle&>(InLHS) == static_cast<const FTypedElementHandle&>(InRHS)
 			&& InLHS.InterfacePtr == InRHS.InterfacePtr;
 	}
 
-	FORCEINLINE friend bool operator!=(const TTypedElementBase& InLHS, const TTypedElementBase& InRHS)
+	inline friend bool operator!=(const TTypedElementBase& InLHS, const TTypedElementBase& InRHS)
 	{
 		return !(InLHS == InRHS);
 	}
@@ -358,25 +358,25 @@ public:
 		return HashCombine(GetTypeHash(static_cast<const FTypedElementHandle&>(InElement)), GetTypeHash(InElement.InterfacePtr));
 	}
 
-	FORCEINLINE void Private_InitializeNoRef(const FTypedElementInternalData& InData, BaseInterfaceType* InInterfacePtr)
+	inline void Private_InitializeNoRef(const FTypedElementInternalData& InData, BaseInterfaceType* InInterfacePtr)
 	{
 		FTypedElementHandle::Private_InitializeNoRef(InData);
 		InterfacePtr = InInterfacePtr;
 	}
 
-	FORCEINLINE void Private_InitializeAddRef(const FTypedElementInternalData& InData, BaseInterfaceType* InInterfacePtr)
+	inline void Private_InitializeAddRef(const FTypedElementInternalData& InData, BaseInterfaceType* InInterfacePtr)
 	{
 		FTypedElementHandle::Private_InitializeAddRef(InData);
 		InterfacePtr = InInterfacePtr;
 	}
 
-	FORCEINLINE void Private_DestroyNoRef()
+	inline void Private_DestroyNoRef()
 	{
 		FTypedElementHandle::Private_DestroyNoRef();
 		InterfacePtr = nullptr;
 	}
 
-	FORCEINLINE void Private_DestroyReleaseRef()
+	inline void Private_DestroyReleaseRef()
 	{
 		FTypedElementHandle::Private_DestroyReleaseRef();
 		InterfacePtr = nullptr;
@@ -437,7 +437,7 @@ public:
 		checkf(!IsSet(), TEXT("Element owner was still set during destruction! This will leak an element, and you should destroy this element prior to destruction!"));
 	}
 
-	FORCEINLINE explicit operator bool() const
+	inline explicit operator bool() const
 	{
 		return IsSet();
 	}
@@ -445,7 +445,7 @@ public:
 	/**
 	 * Has this owner been initialized to a valid element?
 	 */
-	FORCEINLINE bool IsSet() const
+	inline bool IsSet() const
 	{
 		return DataPtr != nullptr;
 	}
@@ -453,7 +453,7 @@ public:
 	/**
 	 * Get the ID that this element represents.
 	 */
-	FORCEINLINE const FTypedElementId& GetId() const
+	inline const FTypedElementId& GetId() const
 	{
 		return DataPtr
 			? DataPtr->GetId()
@@ -464,7 +464,7 @@ public:
 	 * Attempt to access the mutable data stored within this owner, returning null if it isn't possible.
 	 */
 	template <typename U = ElementDataType, std::enable_if_t<!std::is_void<U>::value>* = nullptr>
-	FORCEINLINE U* GetData() const
+	inline U* GetData() const
 	{
 		static_assert(std::is_same<U, ElementDataType>::value, "Don't explicitly specify U!");
 		return DataPtr
@@ -476,7 +476,7 @@ public:
 	 * Attempt to access the mutable data stored within this owner, asserting if it isn't possible.
 	 */
 	template <typename U = ElementDataType, std::enable_if_t<!std::is_void<U>::value>* = nullptr>
-	FORCEINLINE U& GetDataChecked() const
+	inline U& GetDataChecked() const
 	{
 		static_assert(std::is_same<U, ElementDataType>::value, "Don't explicitly specify U!");
 		checkf(DataPtr, TEXT("Handle data is null!"));
@@ -536,12 +536,12 @@ public:
 		InOutElementHandle.Release();
 	}
 
-	FORCEINLINE friend bool operator==(const TTypedElementOwner& InLHS, const TTypedElementOwner& InRHS)
+	inline friend bool operator==(const TTypedElementOwner& InLHS, const TTypedElementOwner& InRHS)
 	{
 		return InLHS.DataPtr == InRHS.DataPtr;
 	}
 
-	FORCEINLINE friend bool operator!=(const TTypedElementOwner& InLHS, const TTypedElementOwner& InRHS)
+	inline friend bool operator!=(const TTypedElementOwner& InLHS, const TTypedElementOwner& InRHS)
 	{
 		return !(InLHS == InRHS);
 	}
@@ -551,18 +551,18 @@ public:
 		return GetTypeHash(InElementOwner.GetId());
 	}
 
-	FORCEINLINE void Private_InitializeNoRef(TTypedElementInternalData<ElementDataType>& InData)
+	inline void Private_InitializeNoRef(TTypedElementInternalData<ElementDataType>& InData)
 	{
 		DataPtr = &InData;
 	}
 
-	FORCEINLINE void Private_InitializeAddRef(TTypedElementInternalData<ElementDataType>& InData)
+	inline void Private_InitializeAddRef(TTypedElementInternalData<ElementDataType>& InData)
 	{
 		Private_InitializeNoRef(InData);
 		RegisterRef();
 	}
 
-	FORCEINLINE void Private_DestroyNoRef()
+	inline void Private_DestroyNoRef()
 	{
 		DataPtr = nullptr;
 #if UE_TYPED_ELEMENT_HAS_REFTRACKING
@@ -570,19 +570,19 @@ public:
 #endif	// UE_TYPED_ELEMENT_HAS_REFTRACKING
 	}
 
-	FORCEINLINE void Private_DestroyReleaseRef()
+	inline void Private_DestroyReleaseRef()
 	{
 		UnregisterRef();
 		Private_DestroyNoRef();
 	}
 
-	FORCEINLINE const TTypedElementInternalData<ElementDataType>* Private_GetInternalData() const
+	inline const TTypedElementInternalData<ElementDataType>* Private_GetInternalData() const
 	{
 		return DataPtr;
 	}
 
 private:
-	FORCEINLINE void RegisterRef()
+	inline void RegisterRef()
 	{
 		if (DataPtr)
 		{
@@ -593,7 +593,7 @@ private:
 		}
 	}
 
-	FORCEINLINE void UnregisterRef()
+	inline void UnregisterRef()
 	{
 		if (DataPtr)
 		{
@@ -650,37 +650,37 @@ public:
 		return *this;
 	}
 
-	FORCEINLINE bool operator==(const FScriptTypedElementHandle& InOther) const
+	inline bool operator==(const FScriptTypedElementHandle& InOther) const
 	{
 		return InternalData == InOther.InternalData;
 	}
 
-	FORCEINLINE bool operator!=(const FScriptTypedElementHandle& InOther) const
+	inline bool operator!=(const FScriptTypedElementHandle& InOther) const
 	{
 		return !(*this == InOther);
 	}
 
-	FORCEINLINE explicit operator bool() const
+	inline explicit operator bool() const
 	{
 		return IsSet();
 	}
 
-	FORCEINLINE bool IsSet() const
+	inline bool IsSet() const
 	{
 		return InternalData.IsSet();
 	}
 
-	FORCEINLINE void Release()
+	inline void Release()
 	{
 		InternalData.Release();
 	}
 
-	FORCEINLINE void Private_Initialize(FScriptTypedElementInternalDataPtr&& InInternalData)
+	inline void Private_Initialize(FScriptTypedElementInternalDataPtr&& InInternalData)
 	{
 		InternalData = InInternalData;
 	}
 
-	FORCEINLINE const FTypedElementId& GetId() const
+	inline const FTypedElementId& GetId() const
 	{
 		return InternalData.GetId();
 	}
diff --git a/Engine/Source/Runtime/TypedElementFramework/Public/Elements/Framework/TypedElementId.h b/Engine/Source/Runtime/TypedElementFramework/Public/Elements/Framework/TypedElementId.h
--- a/Engine/Source/Runtime/TypedElementFramework/Public/Elements/Framework/TypedElementId.h
+++ b/Engine/Source/Runtime/TypedElementFramework/Public/Elements/Framework/TypedElementId.h
@@ -46,7 +46,7 @@ public:
 		checkf(!IsSet(), TEXT("Element ID was still set during destruction! This will leak an element reference, and you should release this ID prior to destruction!"));
 	}
 
-	FORCEINLINE explicit operator bool() const
+	inline explicit operator bool() const
 	{
 		return IsSet();
 	}
@@ -54,7 +54,7 @@ public:
 	/**
 	 * Has this ID been initialized to a valid element?
 	 */
-	FORCEINLINE bool IsSet() const
+	inline bool IsSet() const
 	{
 		return TypeId != 0;
 	}
@@ -62,7 +62,7 @@ public:
 	/**
 	 * Access the type ID portion of this element ID.
 	 */
-	FORCEINLINE FTypedHandleTypeId GetTypeId() const
+	inline FTypedHandleTypeId GetTypeId() const
 	{
 		return TypeId;
 	}
@@ -70,7 +70,7 @@ public:
 	/**
 	 * Access the element ID portion of this element ID.
 	 */
-	FORCEINLINE FTypedHandleElementId GetElementId() const
+	inline FTypedHandleElementId GetElementId() const
 	{
 		return ElementId;
 	}
@@ -80,17 +80,17 @@ public:
 	 * @note You typically don't want to store this directly as the element ID could be re-used.
 	 *       It is primarily useful as a secondary cache where something is keeping a reference to an element ID or element handle (eg, how FTypedElementList uses it internally).
 	 */
-	FORCEINLINE FTypedHandleCombinedId GetCombinedId() const
+	inline FTypedHandleCombinedId GetCombinedId() const
 	{
 		return CombinedId;
 	}
 
-	FORCEINLINE friend bool operator==(const FTypedElementId& InLHS, const FTypedElementId& InRHS)
+	inline friend bool operator==(const FTypedElementId& InLHS, const FTypedElementId& InRHS)
 	{
 		return InLHS.CombinedId == InRHS.CombinedId;
 	}
 
-	FORCEINLINE friend bool operator!=(const FTypedElementId& InLHS, const FTypedElementId& InRHS)
+	inline friend bool operator!=(const FTypedElementId& InLHS, const FTypedElementId& InRHS)
 	{
 		return !(InLHS == InRHS);
 	}
@@ -101,13 +101,13 @@ public:
 		
 	}
 
-	FORCEINLINE void Private_InitializeNoRef(const FTypedHandleTypeId InTypeId, const FTypedHandleElementId InElementId)
+	inline void Private_InitializeNoRef(const FTypedHandleTypeId InTypeId, const FTypedHandleElementId InElementId)
 	{
 		TypeId = InTypeId;
 		ElementId = InElementId;
 	}
 
-	FORCEINLINE void Private_DestroyNoRef()
+	inline void Private_DestroyNoRef()
 	{
 		CombinedId = 0;
 	}
diff --git a/Engine/Source/Runtime/TypedElementFramework/Public/Elements/Framework/TypedElementList.h b/Engine/Source/Runtime/TypedElementFramework/Public/Elements/Framework/TypedElementList.h
--- a/Engine/Source/Runtime/TypedElementFramework/Public/Elements/Framework/TypedElementList.h
+++ b/Engine/Source/Runtime/TypedElementFramework/Public/Elements/Framework/TypedElementList.h
@@ -36,14 +36,14 @@ TYPEDELEMENTFRAMEWORK_API void GetElementImpl(const UTypedElementRegistry* InReg
 TYPEDELEMENTFRAMEWORK_API void GetElementImpl(const UTypedElementRegistry* InRegistry, const FScriptTypedElementHandle& InElementHandle, const TSubclassOf<UInterface>& InBaseInterfaceType, FTypedElement& OutElement);
 
 template <typename BaseInterfaceType, class HandleType>
-FORCEINLINE void GetElement(const UTypedElementRegistry* InRegistry, const HandleType& InElementHandle, TTypedElement<BaseInterfaceType>& OutElement)
+inline void GetElement(const UTypedElementRegistry* InRegistry, const HandleType& InElementHandle, TTypedElement<BaseInterfaceType>& OutElement)
 {
 	static_assert(sizeof(TTypedElement<BaseInterfaceType>) == sizeof(FTypedElement), "All TTypedElement instances must be the same size for this cast implementation to work!");
 	GetElementImpl(InRegistry, InElementHandle, BaseInterfaceType::UClassType::StaticClass(), reinterpret_cast<FTypedElement&>(OutElement));
 }
 
 template <typename BaseInterfaceType, class HandleType>
-FORCEINLINE TTypedElement<BaseInterfaceType> GetElement(const UTypedElementRegistry* InRegistry, const HandleType& InElementHandle)
+inline TTypedElement<BaseInterfaceType> GetElement(const UTypedElementRegistry* InRegistry, const HandleType& InElementHandle)
 {
 	TTypedElement<BaseInterfaceType> Element;
 	GetElement(InRegistry, InElementHandle, Element);
@@ -79,7 +79,7 @@ public:
 	 * Get the element handle at the given index.
 	 * @note Use IsValidIndex to test for validity.
 	 */
-	FORCEINLINE HandleType operator[](const int32 InIndex) const
+	inline HandleType operator[](const int32 InIndex) const
 	{
 		return GetElementHandleAt(InIndex);
 	}
@@ -88,7 +88,7 @@ public:
 	 * Get the element handle at the given index.
 	 * @note Use IsValidIndex to test for validity.
 	 */
-	FORCEINLINE HandleType GetElementHandleAt(const int32 InIndex) const
+	inline HandleType GetElementHandleAt(const int32 InIndex) const
 	{
 		return ElementHandles[InIndex];
 	}
@@ -98,7 +98,7 @@ public:
 	 * @note Use IsValidIndex to test for validity.
 	 */
 	template <typename BaseInterfaceType>
-	FORCEINLINE TTypedElement<BaseInterfaceType> GetElementAt(const int32 InIndex) const
+	inline TTypedElement<BaseInterfaceType> GetElementAt(const int32 InIndex) const
 	{
 		return GetElement<BaseInterfaceType>(GetElementHandleAt(InIndex));
 	}
@@ -108,7 +108,7 @@ public:
 	 * @note Use IsValidIndex to test for validity.
 	 */
 	template <typename BaseInterfaceType>
-	FORCEINLINE void GetElementAt(const int32 InIndex, TTypedElement<BaseInterfaceType>& OutElement) const
+	inline void GetElementAt(const int32 InIndex, TTypedElement<BaseInterfaceType>& OutElement) const
 	{
 		GetElement(GetElementHandleAt(InIndex), OutElement);
 	}
@@ -117,7 +117,7 @@ public:
 	 * Get the element from the given handle.
 	 */
 	template <typename BaseInterfaceType>
-	FORCEINLINE TTypedElement<BaseInterfaceType> GetElement(const HandleType& InElementHandle) const
+	inline TTypedElement<BaseInterfaceType> GetElement(const HandleType& InElementHandle) const
 	{
 		return TypedElementList_Private::GetElement<BaseInterfaceType>(Registry.Get(), InElementHandle);
 	}
@@ -126,7 +126,7 @@ public:
 	 * Get the element from the given handle.
 	 */
 	template <typename BaseInterfaceType>
-	FORCEINLINE void GetElement(const HandleType& InElementHandle, TTypedElement<BaseInterfaceType>& OutElement) const
+	inline void GetElement(const HandleType& InElementHandle, TTypedElement<BaseInterfaceType>& OutElement) const
 	{
 		TypedElementList_Private::GetElement(Registry.Get(), InElementHandle, OutElement);
 	}
@@ -295,7 +295,7 @@ public:
 	/**
 	 * Is the given index a valid entry within this element list?
 	 */
-	FORCEINLINE bool IsValidIndex(const int32 InIndex) const
+	inline bool IsValidIndex(const int32 InIndex) const
 	{
 		return ElementHandles.IsValidIndex(InIndex);
 	}
@@ -303,7 +303,7 @@ public:
 	/**
 	 * Get the number of entries within this element list.
 	 */
-	FORCEINLINE int32 Num() const
+	inline int32 Num() const
 	{
 		return ElementHandles.Num();
 	}
@@ -311,7 +311,7 @@ public:
 	/**
 	 * Shrink this element list storage to avoid slack.
 	 */
-	FORCEINLINE void Shrink()
+	inline void Shrink()
 	{
 		ElementCombinedIds.Shrink();
 		ElementHandles.Shrink();
@@ -320,7 +320,7 @@ public:
 	/**
 	 * Pre-allocate enough memory in this element list to store the given number of entries.
 	 */
-	FORCEINLINE void Reserve(const int32 InSize)
+	inline void Reserve(const int32 InSize)
 	{
 		ElementCombinedIds.Reserve(InSize);
 		ElementHandles.Reserve(InSize);
@@ -329,7 +329,7 @@ public:
 	/**
 	 * Remove all entries from this element list, potentially leaving space allocated for the given number of entries.
 	 */
-	FORCEINLINE void Empty(const int32 InSlack = 0)
+	inline void Empty(const int32 InSlack = 0)
 	{
 		// Avoid creating unnecessary notifications
 		const bool bWasEmpty = ElementHandles.IsEmpty();
@@ -352,7 +352,7 @@ public:
 	/**
 	 * Remove all entries from this element list, preserving existing allocations.
 	 */
-	FORCEINLINE void Reset()
+	inline void Reset()
 	{
 		// Avoid creating unnecessary notifications
 		if (!ElementHandles.IsEmpty())
@@ -368,7 +368,7 @@ public:
 	/**
 	 * Does this element list contain an entry for the given element ID?
 	 */
-	FORCEINLINE bool Contains(const FTypedElementId& InElementId) const
+	inline bool Contains(const FTypedElementId& InElementId) const
 	{
 		return ContainsElementImpl(InElementId);
 	}
@@ -376,7 +376,7 @@ public:
 	/**
 	 * Does this element list contain an entry for the given element handle?
 	 */
-	FORCEINLINE bool Contains(const HandleType& InElementHandle) const
+	inline bool Contains(const HandleType& InElementHandle) const
 	{
 		return ContainsElementImpl(InElementHandle.GetId());
 	}
@@ -385,7 +385,7 @@ public:
 	 * Does this element list contain an entry for the given element owner?
 	 */
 	template <typename ElementDataType>
-	FORCEINLINE bool Contains(const TTypedElementOwner<ElementDataType>& InElementOwner)
+	inline bool Contains(const TTypedElementOwner<ElementDataType>& InElementOwner)
 	{
 		return ContainsElementImpl(InElementOwner.GetId());
 	}
@@ -394,7 +394,7 @@ public:
 	 * Add the given element handle to this element list, if it isn't already in the list.
 	 * @return True if the element handle was added, false if it is already in the list.
 	 */
-	FORCEINLINE bool Add(const HandleType& InElementHandle)
+	inline bool Add(const HandleType& InElementHandle)
 	{
 		return AddElementImpl(CopyTemp(InElementHandle));
 	}
@@ -403,7 +403,7 @@ public:
 	 * Add the given element handle to this element list, if it isn't already in the list.
 	 * @return True if the element handle was added, false if it is already in the list.
 	 */
-	FORCEINLINE bool Add(HandleType&& InElementHandle)
+	inline bool Add(HandleType&& InElementHandle)
 	{
 		return AddElementImpl(MoveTemp(InElementHandle));
 	}
@@ -413,7 +413,7 @@ public:
 	 * @return True if the element owner was added, false if it is already in the list.
 	 */
 	template <typename ElementDataType>
-	FORCEINLINE bool Add(const TTypedElementOwner<ElementDataType>& InElementOwner)
+	inline bool Add(const TTypedElementOwner<ElementDataType>& InElementOwner)
 	{
 		return AddElementImpl(InElementOwner.AcquireHandle());
 	}
@@ -454,7 +454,7 @@ public:
 	 * Append the given element owners to this element list.
 	 */
 	template <typename ElementDataType>
-	FORCEINLINE void Append(const TArray<TTypedElementOwner<ElementDataType>>& InElementOwners)
+	inline void Append(const TArray<TTypedElementOwner<ElementDataType>>& InElementOwners)
 	{
 		Append(MakeArrayView(InElementOwners));
 	}
@@ -478,7 +478,7 @@ public:
 	 * Remove the given element ID from this element list, if it is in the list.
 	 * @return True if the element ID was removed, false if it isn't in the list.
 	 */
-	FORCEINLINE bool Remove(const FTypedElementId& InElementId)
+	inline bool Remove(const FTypedElementId& InElementId)
 	{
 		return RemoveElementImpl(InElementId);
 	}
@@ -487,7 +487,7 @@ public:
 	 * Remove the given element handle from this element list, if it is in the list.
 	 * @return True if the element handle was removed, false if it isn't in the list.
 	 */
-	FORCEINLINE bool Remove(const HandleType& InElementHandle)
+	inline bool Remove(const HandleType& InElementHandle)
 	{
 		return RemoveElementImpl(InElementHandle.GetId());
 	}
@@ -497,7 +497,7 @@ public:
 	 * @return True if the element owner was removed, false if it isn't in the list.
 	 */
 	template <typename ElementDataType>
-	FORCEINLINE bool Remove(const TTypedElementOwner<ElementDataType>& InElementOwner)
+	inline bool Remove(const TTypedElementOwner<ElementDataType>& InElementOwner)
 	{
 		return RemoveElementImpl(InElementOwner.GetId());
 	}
@@ -506,7 +506,7 @@ public:
 	 * Remove any element handles that match the given predicate from this element list.
 	 * @return The number of element handles removed.
 	 */
-	FORCEINLINE int32 RemoveAll(TFunctionRef<bool(const HandleType&)> InPredicate)
+	inline int32 RemoveAll(TFunctionRef<bool(const HandleType&)> InPredicate)
 	{
 		return RemoveAllElementsImpl(InPredicate);
 	}
diff --git a/Engine/Source/Runtime/TypedElementFramework/Public/Elements/Framework/TypedElementOwnerStore.h b/Engine/Source/Runtime/TypedElementFramework/Public/Elements/Framework/TypedElementOwnerStore.h
--- a/Engine/Source/Runtime/TypedElementFramework/Public/Elements/Framework/TypedElementOwnerStore.h
+++ b/Engine/Source/Runtime/TypedElementFramework/Public/Elements/Framework/TypedElementOwnerStore.h
@@ -43,23 +43,23 @@ public:
 
 	TTypedElementOwnerScopedAccessImpl& operator=(TTypedElementOwnerScopedAccessImpl&&) = delete;
 
-	FORCEINLINE explicit operator bool() const
+	inline explicit operator bool() const
 	{
 		return IsSet();
 	}
 
-	FORCEINLINE bool IsSet() const
+	inline bool IsSet() const
 	{
 		return ElementOwner != nullptr;
 	}
 
-	FORCEINLINE ElementOwnerType& operator*() const
+	inline ElementOwnerType& operator*() const
 	{
 		check(ElementOwner);
 		return *ElementOwner;
 	}
 
-	FORCEINLINE ElementOwnerType* operator->() const
+	inline ElementOwnerType* operator->() const
 	{
 		check(ElementOwner);
 		return ElementOwner;
diff --git a/Engine/Source/Runtime/TypedElementFramework/Public/Elements/Framework/TypedElementRegistry.h b/Engine/Source/Runtime/TypedElementFramework/Public/Elements/Framework/TypedElementRegistry.h
--- a/Engine/Source/Runtime/TypedElementFramework/Public/Elements/Framework/TypedElementRegistry.h
+++ b/Engine/Source/Runtime/TypedElementFramework/Public/Elements/Framework/TypedElementRegistry.h
@@ -130,7 +130,7 @@ public:
 	 * Get the element type ID for the associated element type name, if any.
 	 * @return The element type ID, or 0 if the given name wasn't registered.
 	 */
-	FORCEINLINE FTypedHandleTypeId GetRegisteredElementTypeId(const FName InElementTypeName) const
+	inline FTypedHandleTypeId GetRegisteredElementTypeId(const FName InElementTypeName) const
 	{
 		const FRegisteredElementType* RegisteredElementType = GetRegisteredElementTypeFromName(InElementTypeName);
 		return RegisteredElementType ? RegisteredElementType->TypeId : 0;
@@ -140,7 +140,7 @@ public:
 	 * Get the element type name for the associated element type ID, if any.
 	 * @return The element type name, or None if the given ID wasn't registered.
 	 */
-	FORCEINLINE FName GetRegisteredElementTypeName(const FTypedHandleTypeId InElementTypeId) const
+	inline FName GetRegisteredElementTypeName(const FTypedHandleTypeId InElementTypeId) const
 	{
 		const FRegisteredElementType* RegisteredElementType = GetRegisteredElementTypeFromId(InElementTypeId);
 		return RegisteredElementType ? RegisteredElementType->TypeName : FName();
@@ -150,7 +150,7 @@ public:
 	 * Register an element type that doesn't require any additional payload data.
 	 * @param bSupportScriptHandles Does this element type support script handles. If true the elements of this type can only be destroyed from the game thread
 	 */
-	FORCEINLINE void RegisterElementType(const FName InElementTypeName, bool bSupportScriptHandles)
+	inline void RegisterElementType(const FName InElementTypeName, bool bSupportScriptHandles)
 	{
 		if (bSupportScriptHandles)
 		{
@@ -167,7 +167,7 @@ public:
 	 * Note if bSupportScriptHandles is true, elements of this type should only be destroyed from the game thread
 	 */
 	template <typename ElementDataType, bool bSupportScriptHandles>
-	FORCEINLINE void RegisterElementType(const FName InElementTypeName)
+	inline void RegisterElementType(const FName InElementTypeName)
 	{
 		RegisterElementTypeImpl(InElementTypeName, MakeUnique<TRegisteredElementType<ElementDataType, bSupportScriptHandles>>());
 	}
@@ -176,7 +176,7 @@ public:
 	 * Register that an element interface is supported for the given type, which must have previously been registered via RegisterElementType.
 	 */
 	template <typename BaseInterfaceType>
-	FORCEINLINE void RegisterElementInterface(const FName InElementTypeName, UObject* InElementInterface, const bool InAllowOverride = false)
+	inline void RegisterElementInterface(const FName InElementTypeName, UObject* InElementInterface, const bool InAllowOverride = false)
 	{
 		checkf(InElementInterface->GetClass()->ImplementsInterface(BaseInterfaceType::UClassType::StaticClass()), TEXT("The InElementInterface pass must implement the interface to register."));
 		RegisterElementInterfaceImpl(InElementTypeName, InElementInterface, BaseInterfaceType::UClassType::StaticClass(), InAllowOverride);
@@ -186,7 +186,7 @@ public:
 	 * Get the element interface supported by the given type, or null if there is no support for this interface.
 	 */
 	template <typename BaseInterfaceType>
-	FORCEINLINE BaseInterfaceType* GetElementInterface(const FTypedHandleTypeId InElementTypeId) const
+	inline BaseInterfaceType* GetElementInterface(const FTypedHandleTypeId InElementTypeId) const
 	{
 		return Cast<BaseInterfaceType>(GetElementInterfaceImpl(InElementTypeId, BaseInterfaceType::UClassType::StaticClass()));
 	}
@@ -195,7 +195,7 @@ public:
 	 * Get the element interface supported by the given handle, or null if there is no support for this interface.
 	 */
 	template <typename BaseInterfaceType>
-	FORCEINLINE BaseInterfaceType* GetElementInterface(const FTypedElementHandle& InElementHandle) const
+	inline BaseInterfaceType* GetElementInterface(const FTypedElementHandle& InElementHandle) const
 	{
 		return Cast<BaseInterfaceType>(GetElementInterfaceImpl(InElementHandle.GetId().GetTypeId(), BaseInterfaceType::UClassType::StaticClass()));
 	}
@@ -203,7 +203,7 @@ public:
 	/**
 	 * Get the element interface supported by the given handle, or null if there is no support for this interface.
 	 */
-	FORCEINLINE UObject* GetElementInterface(const FTypedElementHandle& InElementHandle, const TSubclassOf<UInterface> InBaseInterfaceType) const
+	inline UObject* GetElementInterface(const FTypedElementHandle& InElementHandle, const TSubclassOf<UInterface> InBaseInterfaceType) const
 	{
 		return GetElementInterfaceImpl(InElementHandle.GetId().GetTypeId(), InBaseInterfaceType);
 	}
@@ -212,7 +212,7 @@ public:
 	 * Get the element interface supported by the given handle, or null if there is no support for this interface or if the handle is invalid.
 	 */
 	UFUNCTION(BlueprintPure, Category="TypedElementFramework|Registry")
-	FORCEINLINE UObject* GetElementInterface(const FScriptTypedElementHandle& InElementHandle, const TSubclassOf<UInterface> InBaseInterfaceType) const
+	inline UObject* GetElementInterface(const FScriptTypedElementHandle& InElementHandle, const TSubclassOf<UInterface> InBaseInterfaceType) const
 	{
 		return GetElementInterfaceImpl(InElementHandle.GetId().GetTypeId(), InBaseInterfaceType);
 	}
@@ -221,7 +221,7 @@ public:
 	 * Create an element that doesn't require any additional payload data.
 	 * @note The associated handle ID should be something that can externally be used to uniquely identify this element, until DestroyElementHandle is called on this handle.
 	 */
-	FORCEINLINE FTypedElementOwner CreateElement(const FName InElementTypeName, const FTypedHandleElementId InElementId)
+	inline FTypedElementOwner CreateElement(const FName InElementTypeName, const FTypedHandleElementId InElementId)
 	{
 		return CreateElementImpl<void>(InElementTypeName, InElementId);
 	}
@@ -231,7 +231,7 @@ public:
 	 * @note Allocation of the payload data and the associated handle ID are managed internally, and the data will remain valid until DestroyElementHandle is called on this handle.
 	 */
 	template <typename ElementDataType>
-	FORCEINLINE TTypedElementOwner<ElementDataType> CreateElement(const FName InElementTypeName)
+	inline TTypedElementOwner<ElementDataType> CreateElement(const FName InElementTypeName)
 	{
 		return CreateElementImpl<ElementDataType>(InElementTypeName, INDEX_NONE);
 	}
@@ -240,7 +240,7 @@ public:
 	 * Destroy an element.
 	 * @note Destruction is deferred until the next call to ProcessDeferredElementsToDestroy.
 	 */
-	FORCEINLINE void DestroyElement(FTypedElementOwner& InOutElementOwner)
+	inline void DestroyElement(FTypedElementOwner& InOutElementOwner)
 	{
 		return DestroyElementImpl<void>(InOutElementOwner);
 	}
@@ -250,7 +250,7 @@ public:
 	 * @note Destruction is deferred until the next call to ProcessDeferredElementsToDestroy.
 	 */
 	template <typename ElementDataType>
-	FORCEINLINE void DestroyElement(TTypedElementOwner<ElementDataType>& InOutElementOwner)
+	inline void DestroyElement(TTypedElementOwner<ElementDataType>& InOutElementOwner)
 	{
 		return DestroyElementImpl<ElementDataType>(InOutElementOwner);
 	}
@@ -275,7 +275,7 @@ public:
 	/**
 	 * Get an element that implements the given interface from its minimal ID.
 	 */
-	FORCEINLINE FTypedElement GetElement(const FTypedElementId& InElementId, const TSubclassOf<UInterface>& InBaseInterfaceType) const
+	inline FTypedElement GetElement(const FTypedElementId& InElementId, const TSubclassOf<UInterface>& InBaseInterfaceType) const
 	{
 		FTypedElement Element;
 		GetElementImpl(InElementId, InBaseInterfaceType, Element);
@@ -286,7 +286,7 @@ public:
 	 * Get an element that implements the given interface from its minimal ID.
 	 */
 	template <typename BaseInterfaceType>
-	FORCEINLINE TTypedElement<BaseInterfaceType> GetElement(const FTypedElementId& InElementId) const
+	inline TTypedElement<BaseInterfaceType> GetElement(const FTypedElementId& InElementId) const
 	{
 		TTypedElement<BaseInterfaceType> Element;
 		GetElementImpl(InElementId, BaseInterfaceType::UClassType::StaticClass(), Element);
@@ -296,7 +296,7 @@ public:
 	/**
 	 * Get an element that implements the given interface from its handle.
 	 */
-	FORCEINLINE FTypedElement GetElement(const FTypedElementHandle& InElementHandle, const TSubclassOf<UInterface>& InBaseInterfaceType) const
+	inline FTypedElement GetElement(const FTypedElementHandle& InElementHandle, const TSubclassOf<UInterface>& InBaseInterfaceType) const
 	{
 		FTypedElement Element;
 		GetElementImpl(InElementHandle, InBaseInterfaceType, Element);
@@ -307,7 +307,7 @@ public:
 	 * Get an element that implements the given interface from its handle.
 	 */
 	template <typename BaseInterfaceType>
-	FORCEINLINE TTypedElement<BaseInterfaceType> GetElement(const FTypedElementHandle& InElementHandle) const
+	inline TTypedElement<BaseInterfaceType> GetElement(const FTypedElementHandle& InElementHandle) const
 	{
 		TTypedElement<BaseInterfaceType> Element;
 		GetElementImpl(InElementHandle, BaseInterfaceType::UClassType::StaticClass(), Element);
@@ -317,7 +317,7 @@ public:
 	/**
 	 * Create an empty list of elements associated with this registry.
 	 */
-	FORCEINLINE FTypedElementListRef CreateElementList()
+	inline FTypedElementListRef CreateElementList()
 	{
 		return FTypedElementList::Private_CreateElementList(this);
 	}
@@ -325,7 +325,7 @@ public:
 	/**
 	 * Create an empty list of script elements associated with this registry.
 	 */
-	FORCEINLINE FScriptTypedElementListRef CreateScriptElementList()
+	inline FScriptTypedElementListRef CreateScriptElementList()
 	{
 		return FScriptTypedElementList::Private_CreateElementList(this);
 	}
@@ -344,7 +344,7 @@ public:
 	 * Create an empty list of elements associated with this registry, populated from the given owners that are valid.
 	 */
 	template <typename ElementDataType>
-	FORCEINLINE FTypedElementListRef CreateElementList(const TArray<TTypedElementOwner<ElementDataType>>& InElementOwners)
+	inline FTypedElementListRef CreateElementList(const TArray<TTypedElementOwner<ElementDataType>>& InElementOwners)
 	{
 		return CreateElementList(MakeArrayView(InElementOwners));
 	}
@@ -369,7 +369,7 @@ public:
 	/**
 	 * Create a script handle. This should only be use for the script exposure apis.
 	 */
-	FORCEINLINE FScriptTypedElementHandle CreateScriptHandle(const FTypedElementId& InElementId)
+	inline FScriptTypedElementHandle CreateScriptHandle(const FTypedElementId& InElementId)
 	{
 		if (InElementId == FTypedElementId::Unset)
 		{
@@ -414,7 +414,7 @@ public:
 	}
 
 	// Note: Access for FTypedElementList
-	FORCEINLINE void Private_GetElementImpl(const FTypedElementHandle& InElementHandle, const TSubclassOf<UInterface>& InBaseInterfaceType, FTypedElement& OutElement) const
+	inline void Private_GetElementImpl(const FTypedElementHandle& InElementHandle, const TSubclassOf<UInterface>& InBaseInterfaceType, FTypedElement& OutElement) const
 	{
 		GetElementImpl(InElementHandle, InBaseInterfaceType, OutElement);
 	}
@@ -668,13 +668,13 @@ private:
 	TYPEDELEMENTFRAMEWORK_API void OnEndFrame();
 	TYPEDELEMENTFRAMEWORK_API void OnPostGarbageCollect();
 
-	FORCEINLINE void IncrementDisableElementDestructionOnGCCount()
+	inline void IncrementDisableElementDestructionOnGCCount()
 	{
 		checkf(DisableElementDestructionOnGCCount != MAX_uint8, TEXT("DisableElementDestructionOnGCCount overflow!"));
 		++DisableElementDestructionOnGCCount;
 	}
 
-	FORCEINLINE void DecrementDisableElementDestructionOnGCCount()
+	inline void DecrementDisableElementDestructionOnGCCount()
 	{
 		checkf(DisableElementDestructionOnGCCount != 0, TEXT("DisableElementDestructionOnGCCount underflow!"));
 		--DisableElementDestructionOnGCCount;
diff --git a/Engine/Source/Runtime/TypedElementRuntime/Public/Elements/Framework/TypedElementSelectionSet.h b/Engine/Source/Runtime/TypedElementRuntime/Public/Elements/Framework/TypedElementSelectionSet.h
--- a/Engine/Source/Runtime/TypedElementRuntime/Public/Elements/Framework/TypedElementSelectionSet.h
+++ b/Engine/Source/Runtime/TypedElementRuntime/Public/Elements/Framework/TypedElementSelectionSet.h
@@ -99,12 +99,12 @@ public:
 	FTypedElementSelectionSetElement(FTypedElementSelectionSetElement&&) = default;
 	FTypedElementSelectionSetElement& operator=(FTypedElementSelectionSetElement&&) = default;
 
-	FORCEINLINE explicit operator bool() const
+	inline explicit operator bool() const
 	{
 		return IsSet();
 	}
 
-	FORCEINLINE bool IsSet() const
+	inline bool IsSet() const
 	{
 		return ElementSelectionHandle.IsSet()
 			&& ElementList
