[Teds] Remove 5.5 deprecated code from Teds modules

diff --git a/Engine/Source/Runtime/TypedElementFramework/Private/Elements/Framework/TypedElementRegistry.cpp b/Engine/Source/Runtime/TypedElementFramework/Private/Elements/Framework/TypedElementRegistry.cpp
--- a/Engine/Source/Runtime/TypedElementFramework/Private/Elements/Framework/TypedElementRegistry.cpp
+++ b/Engine/Source/Runtime/TypedElementFramework/Private/Elements/Framework/TypedElementRegistry.cpp
@@ -65,52 +65,6 @@ UTypedElementRegistry* UTypedElementRegistry::GetInstance()
 	return Instance.Get();
 }
 
-UE::Editor::DataStorage::ICoreProvider* UTypedElementRegistry::GetMutableDataStorage()
-{
-	using namespace UE::Editor::DataStorage;
-	return GetMutableDataStorageFeature<ICoreProvider>(StorageFeatureName);
-}
-
-const UE::Editor::DataStorage::ICoreProvider* UTypedElementRegistry::GetDataStorage() const
-{
-	using namespace UE::Editor::DataStorage;
-	return GetDataStorageFeature<ICoreProvider>(StorageFeatureName);
-}
-
-UE::Editor::DataStorage::ICompatibilityProvider* UTypedElementRegistry::GetMutableDataStorageCompatibility()
-{
-	using namespace UE::Editor::DataStorage;
-	return GetMutableDataStorageFeature<ICompatibilityProvider>(CompatibilityFeatureName);
-}
-
-const UE::Editor::DataStorage::ICompatibilityProvider* UTypedElementRegistry::GetDataStorageCompatibility() const
-{
-	using namespace UE::Editor::DataStorage;
-	return GetDataStorageFeature<ICompatibilityProvider>(CompatibilityFeatureName);
-}
-
-UE::Editor::DataStorage::IUiProvider* UTypedElementRegistry::GetMutableDataStorageUi()
-{
-	using namespace UE::Editor::DataStorage;
-	return GetMutableDataStorageFeature<IUiProvider>(UiFeatureName);
-}
-
-const UE::Editor::DataStorage::IUiProvider* UTypedElementRegistry::GetDataStorageUi() const
-{
-	using namespace UE::Editor::DataStorage;
-	return GetDataStorageFeature<IUiProvider>(UiFeatureName);
-}
-
-bool UTypedElementRegistry::AreDataStorageInterfacesSet() const
-{
-	return UE::Editor::DataStorage::AreEditorDataStorageFeaturesEnabled();
-}
-
-FSimpleMulticastDelegate& UTypedElementRegistry::OnDataStorageInterfacesSet()
-{
-	return UE::Editor::DataStorage::OnEditorDataStorageFeaturesEnabled();
-}
-
 void UTypedElementRegistry::FinishDestroy()
 {
 	TStrongObjectPtr<UTypedElementRegistry>& Instance = GetTypedElementRegistryInstance();
diff --git a/Engine/Source/Runtime/TypedElementFramework/Public/Elements/Framework/TypedElementRegistry.h b/Engine/Source/Runtime/TypedElementFramework/Public/Elements/Framework/TypedElementRegistry.h
--- a/Engine/Source/Runtime/TypedElementFramework/Public/Elements/Framework/TypedElementRegistry.h
+++ b/Engine/Source/Runtime/TypedElementFramework/Public/Elements/Framework/TypedElementRegistry.h
@@ -77,27 +77,6 @@ public:
 	UFUNCTION(BlueprintPure, DisplayName="Get Default Typed Element Registry", Category = "TypedElementFramework|Registry", meta=(ScriptName="GetDefaultTypedElementRegistry"))
 	static TYPEDELEMENTFRAMEWORK_API UTypedElementRegistry* GetInstance();
 
-	UE_DEPRECATED(5.5, "Function has moved to modular features. Use UE::Editor::DataStorage::GetMutableDataStorageFeature instead.")
-	TYPEDELEMENTFRAMEWORK_API UE::Editor::DataStorage::ICoreProvider* GetMutableDataStorage();
-	UE_DEPRECATED(5.5, "Function has moved to modular features. Use UE::Editor::DataStorage::GetDataStorageFeature instead.")
-	TYPEDELEMENTFRAMEWORK_API const UE::Editor::DataStorage::ICoreProvider* GetDataStorage() const;
-
-	UE_DEPRECATED(5.5, "Function has moved to modular features. Use UE::Editor::DataStorage::GetMutableDataStorageFeature instead.")
-	TYPEDELEMENTFRAMEWORK_API UE::Editor::DataStorage::ICompatibilityProvider* GetMutableDataStorageCompatibility();
-	UE_DEPRECATED(5.5, "Function has moved to modular features. Use UE::Editor::DataStorage::GetDataStorageFeature instead.")
-	TYPEDELEMENTFRAMEWORK_API const UE::Editor::DataStorage::ICompatibilityProvider* GetDataStorageCompatibility() const;
-
-	UE_DEPRECATED(5.5, "Function has moved to modular features. Use UE::Editor::DataStorage::GetMutableDataStorageFeature instead.")
-	TYPEDELEMENTFRAMEWORK_API UE::Editor::DataStorage::IUiProvider* GetMutableDataStorageUi();
-	UE_DEPRECATED(5.5, "Function has moved to modular features. Use UE::Editor::DataStorage::GetDataStorageFeature instead.")
-	TYPEDELEMENTFRAMEWORK_API const UE::Editor::DataStorage::IUiProvider* GetDataStorageUi() const;
-
-	UE_DEPRECATED(5.5, "Use UE::Editor::DataStorage::AreDataStorageInterfacesSet instead.")
-	TYPEDELEMENTFRAMEWORK_API bool AreDataStorageInterfacesSet() const;
-
-	UE_DEPRECATED(5.5, "Use UE::Editor::DataStorage::OnAllStorageFeaturesEnabled instead")
-	TYPEDELEMENTFRAMEWORK_API FSimpleMulticastDelegate& OnDataStorageInterfacesSet();
-
 	/**
 	 * Event fired when references to one element should be replaced with a reference to a different element.
 	 */
