[TEDS] New Query Context foundation.

The query context used whenever a query callback needs access information about TEDS has been growing for a while and has become difficult to maintain. There are also several variants for different use cases, making it difficult to write query callbacks that can be reused in multiple ways, which limits the ability to make reusable utility functions. The API also has to cover both single row and batch versions causing confusion as functions sometimes behave differently in either case or batch exclusive functions are used in single row use case and vice versa. The new query context addresses this by having functionality grouped by capability and allowing users to pick what capabilities are needed. This allows for better ways to detect incorrect/unsupported used while also providing a uniform syntax for the query contexts so they can be reused in multiple situations.

This first implementation setups the foundation to:
- Setup capabilities. Two very basic ones are provided.
- Allow users to composite a context from capabilities.
- automatically generated shared contract between the context and the implementation.
- provide an implementation, with automatically generating default support for unsupported/unimplemented capabilities.
As an example a context implementation is provide for single row processing that forwards to calls to the Direct API and a new function that can filter rows based on a provide query syntax based filter function.

diff --git a/Engine/Plugins/Experimental/EditorDataStorage/Source/TedsCore/Private/Queries/TypedElementQueryContext_DirectSingleApi.cpp b/Engine/Plugins/Experimental/EditorDataStorage/Source/TedsCore/Private/Queries/TypedElementQueryContext_DirectSingleApi.cpp
new file mode 100644
--- /dev/null
+++ b/Engine/Plugins/Experimental/EditorDataStorage/Source/TedsCore/Private/Queries/TypedElementQueryContext_DirectSingleApi.cpp
@@ -0,0 +1,29 @@
+// Copyright Epic Games, Inc. All Rights Reserved.
+
+#include "Queries/TypedElementQueryContext_DirectSingleApi.h"
+#include "TypedElementDatabase.h"
+
+namespace UE::Editor::DataStorage::Queries
+{
+	const FName FQueryContext_DirectSingleApi::Name = "FQueryContext_DirectSingleApi";
+
+	FQueryContext_DirectSingleApi::FQueryContext_DirectSingleApi(UEditorDataStorage& DataStorage)
+		: DataStorage(DataStorage)
+	{
+	}
+
+	bool FQueryContext_DirectSingleApi::CurrentRowHasColumns(TConstArrayView<const UScriptStruct*> ColumnTypes) const
+	{
+		return DataStorage.HasColumns(CurrentRow, ColumnTypes);
+	}
+
+	RowHandle FQueryContext_DirectSingleApi::GetCurrentRow() const
+	{
+		return CurrentRow;
+	}
+
+	void FQueryContext_DirectSingleApi::SetCurrentRow(RowHandle Row)
+	{
+		CurrentRow = Row;
+	}
+} // namespace UE::Editor::DataStorage::Queries
\ No newline at end of file
diff --git a/Engine/Plugins/Experimental/EditorDataStorage/Source/TedsCore/Private/Queries/TypedElementQueryContext_DirectSingleApi.h b/Engine/Plugins/Experimental/EditorDataStorage/Source/TedsCore/Private/Queries/TypedElementQueryContext_DirectSingleApi.h
new file mode 100644
--- /dev/null
+++ b/Engine/Plugins/Experimental/EditorDataStorage/Source/TedsCore/Private/Queries/TypedElementQueryContext_DirectSingleApi.h
@@ -0,0 +1,30 @@
+// Copyright Epic Games, Inc. All Rights Reserved.
+
+#pragma once
+
+#include "Elements/Framework/TypedElementQueryContextImplementation.h"
+
+class UEditorDataStorage;
+
+namespace UE::Editor::DataStorage::Queries
+{
+	/** Query context implementation for single row queries that call into TEDS Direct API. */
+	class FQueryContext_DirectSingleApi final : public SingleRowInfo
+	{
+	public:
+		static const FName Name;
+
+		explicit FQueryContext_DirectSingleApi(UEditorDataStorage& DataStorage);
+
+		virtual bool CurrentRowHasColumns(TConstArrayView<const UScriptStruct*> ColumnTypes) const override;
+		virtual RowHandle GetCurrentRow() const override;
+
+		void SetCurrentRow(RowHandle Row);
+
+	private:
+		UEditorDataStorage& DataStorage;
+		RowHandle CurrentRow = InvalidRowHandle;
+	};
+	
+	using QueryContext_DirectSingleApi = TQueryContextImpl<FQueryContext_DirectSingleApi>;
+} // namespace UE::Editor::DataStorage::Queries
diff --git a/Engine/Plugins/Experimental/EditorDataStorage/Source/TedsCore/Private/TypedElementDatabase.cpp b/Engine/Plugins/Experimental/EditorDataStorage/Source/TedsCore/Private/TypedElementDatabase.cpp
--- a/Engine/Plugins/Experimental/EditorDataStorage/Source/TedsCore/Private/TypedElementDatabase.cpp
+++ b/Engine/Plugins/Experimental/EditorDataStorage/Source/TedsCore/Private/TypedElementDatabase.cpp
@@ -15,6 +15,7 @@
 #include "MassSubsystemAccess.h"
 #include "Processors/TypedElementProcessorAdaptors.h"
 #include "ProfilingDebugging/CpuProfilerTrace.h"
+#include "Queries/TypedElementQueryContext_DirectSingleApi.h"
 #include "Stats/Stats.h"
 #include "TickTaskManagerInterface.h"
 #include "TypedElementDatabaseEnvironment.h"
@@ -575,6 +576,28 @@ bool UEditorDataStorage::IsRowAssignedUnsafe(UE::Editor::DataStorage::RowHandle
 	return UE::Editor::DataStorage::Legacy::FCommandBuffer::Execute_IsRowAssigned(*ActiveEditorEntityManager, Row);
 }
 
+void UEditorDataStorage::FilterRowsBy(UE::Editor::DataStorage::FRowHandleArray& Result, 
+	UE::Editor::DataStorage::FRowHandleArrayView Input, UE::Editor::DataStorage::Queries::TQueryFunction<bool>& Filter)
+{
+	using namespace UE::Editor::DataStorage;
+	using namespace UE::Editor::DataStorage::Queries;
+
+	if (ActiveEditorEntityManager)
+	{
+		QueryContext_DirectSingleApi Context(*this);
+		Context.CheckCompatiblity(Filter);
+
+		for (RowHandle Row : Input)
+		{
+			Context.GetContextImplementation().SetCurrentRow(Row);
+			if (Filter.Function(Context))
+			{
+				Result.Add(Row);
+			}
+		}
+	}
+}
+
 void UEditorDataStorage::AddColumn(UE::Editor::DataStorage::RowHandle Row, const UScriptStruct* ColumnType)
 {
 	if (ColumnType && ActiveEditorEntityManager)
diff --git a/Engine/Plugins/Experimental/EditorDataStorage/Source/TedsCore/Private/TypedElementDatabase.h b/Engine/Plugins/Experimental/EditorDataStorage/Source/TedsCore/Private/TypedElementDatabase.h
--- a/Engine/Plugins/Experimental/EditorDataStorage/Source/TedsCore/Private/TypedElementDatabase.h
+++ b/Engine/Plugins/Experimental/EditorDataStorage/Source/TedsCore/Private/TypedElementDatabase.h
@@ -112,6 +112,8 @@ public:
 	UE_API bool IsRowAvailableUnsafe(UE::Editor::DataStorage::RowHandle Row) const;
 	/** Same as IsRowAssigned, but doesn't check if the data storage has been initialized. */
 	UE_API bool IsRowAssignedUnsafe(UE::Editor::DataStorage::RowHandle Row) const;
+	virtual void FilterRowsBy(UE::Editor::DataStorage::FRowHandleArray& Result, UE::Editor::DataStorage::FRowHandleArrayView Input, 
+		UE::Editor::DataStorage::Queries::TQueryFunction<bool>& Filter) override;
 
 	UE_API virtual void AddColumn(UE::Editor::DataStorage::RowHandle Row, const UScriptStruct* ColumnType) override;
 	UE_API virtual void AddColumn(UE::Editor::DataStorage::RowHandle Row, const UE::Editor::DataStorage::FValueTag& Tag, const FName& InValue) override;
@@ -150,7 +152,7 @@ public:
 	UE_API virtual UE::Editor::DataStorage::FQueryResult RunQuery(UE::Editor::DataStorage::QueryHandle Query, UE::Editor::DataStorage::EDirectQueryExecutionFlags Flags,
 		UE::Editor::DataStorage::DirectQueryCallbackRef Callback) override;
 	UE_API virtual void ActivateQueries(FName ActivationName) override;
-
+	
 	UE_API virtual UE::Editor::DataStorage::RowHandle LookupMappedRow(
 		const FName& Domain, const UE::Editor::DataStorage::FMapKeyView& Key) const override;
 	UE_API virtual void MapRow(
diff --git a/Engine/Source/Runtime/TypedElementFramework/Public/Elements/Framework/TypedElementQueryCapabilities.h b/Engine/Source/Runtime/TypedElementFramework/Public/Elements/Framework/TypedElementQueryCapabilities.h
new file mode 100644
--- /dev/null
+++ b/Engine/Source/Runtime/TypedElementFramework/Public/Elements/Framework/TypedElementQueryCapabilities.h
@@ -0,0 +1,67 @@
+// Copyright Epic Games, Inc. All Rights Reserved.
+
+#pragma once
+
+#include <type_traits>
+#include "Containers/ContainersFwd.h"
+#include "Elements/Common/TypedElementHandles.h"
+#include "UObject/NameTypes.h"
+
+class UScriptStruct;
+
+namespace UE::Editor::DataStorage::Queries
+{
+	/** 
+	 * Class fragments used to composite a context.
+	 * A query context is created by combining various context capabilities together. Users can request a query context
+	 * being created with one or more of these capabilities and based on the requested capabilities they get connected
+	 * with an implementation. Each capability represents a subset of functionality of TEDS such as access to column data,
+	 * creating new rows or access to subqueries.
+	 */
+	struct IContextCapability
+	{
+		virtual ~IContextCapability() = default;
+	};
+
+	
+#define ArgTypeName(Type, Name) Type Name
+
+#define Function0(Capability, Return, Function) virtual Return Function() = 0;
+#define Function1(Capability, Return, Function, Arg1) virtual Return Function(ArgTypeName Arg1 ) = 0;
+#define Function2(Capability, Return, Function, Arg1, Arg2) virtual Return Function(ArgTypeName Arg1 , ArgTypeName Arg2 ) = 0;
+#define Function3(Capability, Return, Function, Arg1, Arg2, Arg3) virtual Return Function(ArgTypeName Arg1 , ArgTypeName Arg2 , ArgTypeName Arg3 ) = 0;
+#define Function4(Capability, Return, Function, Arg1, Arg2, Arg3, Arg4) virtual Return Function(ArgTypeName Arg1 , ArgTypeName Arg2 , ArgTypeName Arg3 , ArgTypeName Arg4 ) = 0;
+
+#define ConstFunction0(Capability, Return, Function) virtual Return Function() const = 0;
+#define ConstFunction1(Capability, Return, Function, Arg1) virtual Return Function(ArgTypeName Arg1 ) const = 0;
+#define ConstFunction2(Capability, Return, Function, Arg1, Arg2) virtual Return Function(ArgTypeName Arg1 , ArgTypeName Arg2 ) const = 0;
+#define ConstFunction3(Capability, Return, Function, Arg1, Arg2, Arg3) virtual Return Function(ArgTypeName Arg1 , ArgTypeName Arg2 , ArgTypeName Arg3 ) const = 0;
+#define ConstFunction4(Capability, Return, Function, Arg1, Arg2, Arg3, Arg4) virtual Return Function(ArgTypeName Arg1 , ArgTypeName Arg2 , ArgTypeName Arg3 , ArgTypeName Arg4 ) const = 0;
+
+#define CapabilityStart(Capability) \
+	template<typename Base> \
+	struct I##Capability : Base \
+	{ \
+		inline static const FName Name = #Capability ; \
+		virtual ~I##Capability () override = default;
+
+#define CapabilityEnd(Capability) \
+	}; \
+	using Capability = I##Capability <IContextCapability>;
+
+#include "Elements/Framework/TypedElementQueryCapabilities.inl"
+
+#undef ArgTypeName
+#undef Function0
+#undef Function1
+#undef Function2
+#undef Function3
+#undef Function4
+#undef ConstFunction0
+#undef ConstFunction1
+#undef ConstFunction2
+#undef ConstFunction3
+#undef ConstFunction4
+#undef CapabilityStart
+#undef CapabilityEnd
+} // namespace UE::Editor::DataStorage::Queries
\ No newline at end of file
diff --git a/Engine/Source/Runtime/TypedElementFramework/Public/Elements/Framework/TypedElementQueryCapabilities.inl b/Engine/Source/Runtime/TypedElementFramework/Public/Elements/Framework/TypedElementQueryCapabilities.inl
new file mode 100644
--- /dev/null
+++ b/Engine/Source/Runtime/TypedElementFramework/Public/Elements/Framework/TypedElementQueryCapabilities.inl
@@ -0,0 +1,52 @@
+// Copyright Epic Games, Inc. All Rights Reserved.
+
+// Mini-DSL to declare TEDS capabilities.
+// Do not include #pragma once here or any other guards against repeated inclusion as this file is meant to be included multiple times.
+// This file contains the definitions of all available query capabilities. The macros get expended in various ways to generate interfaces,
+// function forwarders, contexts, etc.
+
+/**
+ * Capability to provide information about the currently active row.
+ */
+CapabilityStart(SingleRowInfo)
+	/*
+	 * Returns whether a column matches the requested type or not. This version only applies to current row. This version is faster
+	 * querying for an arbitrary row.
+	 */
+	ConstFunction1(SingleRowInfo, bool, CurrentRowHasColumns, (TConstArrayView<const UScriptStruct*>, Columns))
+	/** Returns the currently active row. */
+	ConstFunction0(SingleRowInfo, RowHandle, GetCurrentRow)
+
+#if defined(WithWrappers)
+	/*
+	 * Returns whether a column matches the requested type or not. This version only applies to current row. This version is faster
+	 * querying for an arbitrary row.
+	 */
+	template<TColumnType... ColumnTypes>
+	bool CurrentRowHasColumns() const
+	{
+		return this->CurrentRowHasColumns(TConstArrayView<const UScriptStruct*>({ ColumnTypes::StaticStruct()... }));
+	}
+#endif
+CapabilityEnd(SingleRowInfo)
+
+/**
+ * Capability to provide information about the currently active batch.
+ */
+CapabilityStart(RowBatchInfo)
+	/** Returns the number of rows in the current batch. */
+	ConstFunction0(RowBatchInfo, uint32, GetBatchRowCount)
+	/** Returns an view with the rows used by this batch. */
+	ConstFunction0(RowBatchInfo, TConstArrayView<RowHandle>, GetBatchRowHandles)
+	/** Checks if the rows in the current batch have the requested columns. */
+	ConstFunction1(RowBatchInfo, bool, CurrentBatchTableHasColumns, (TConstArrayView<const UScriptStruct*>, Columns))
+
+#if defined(WithWrappers)
+	/** Checks if the rows in the current batch have the requested columns. */
+	template<TColumnType... ColumnTypes>
+	bool CurrentBatchTableHasColumns() const
+	{
+		return this->CurrentBatchTableHasColumns(TConstArrayView<const UScriptStruct*>({ ColumnTypes::StaticStruct()... }));
+	}
+#endif
+CapabilityEnd(RowBatchInfo)
\ No newline at end of file
diff --git a/Engine/Source/Runtime/TypedElementFramework/Public/Elements/Framework/TypedElementQueryCapabilityForwarder.h b/Engine/Source/Runtime/TypedElementFramework/Public/Elements/Framework/TypedElementQueryCapabilityForwarder.h
new file mode 100644
--- /dev/null
+++ b/Engine/Source/Runtime/TypedElementFramework/Public/Elements/Framework/TypedElementQueryCapabilityForwarder.h
@@ -0,0 +1,98 @@
+// Copyright Epic Games, Inc. All Rights Reserved.
+
+#pragma once
+
+#include "Elements/Common/TypedElementCommonTypes.h"
+#include "Elements/Framework/TypedElementQueryCapabilities.h"
+#include "Elements/Framework/TypedElementQueryContract.h"
+
+namespace UE::Editor::DataStorage::Queries::Private
+{
+	// Class fragments use to composite a class that forwards its calls to the query context contract.
+	// If the forwarder contains any virtual calls, the forwarder won't be fully inlined and will add an additional pointer dereference
+	// to each call. To avoid this, the forwarder doesn't directly inherit from the capability interfaces. Static asserts are added
+	// to verify that the implementation is complete.
+	
+	struct FQueryContext {};
+
+	struct ForwardTesterBase
+	{
+		virtual ~ForwardTesterBase() = default;
+	};
+
+	template<typename... Capabilities> struct TForwarder
+	{
+		static_assert(sizeof...(Capabilities) < 0,
+			"Specialization for query context forwards missing. The most likely reason is the addition of a new query "
+			"context capability that wasn't specialized.");
+	};
+
+	template<typename Base, typename...Capabilities>
+	struct TForwarder<Base, Capabilities...> : Base, TForwarder<Capabilities...>{};
+
+	template<> 
+	struct TForwarder<> : FQueryContext
+	{
+		explicit TForwarder(IContextContract& Contract) : Contract(Contract) {}
+	
+	protected:
+		IContextContract& Contract;
+	};
+
+#define ArgTypeName(Type, Name) Type Name
+#define ArgName(Type, Name) Name
+
+#define Function0(Capability, Return, Function) Return Function() \
+	{ return this->Contract. Function (); }
+#define Function1(Capability, Return, Function, Arg1) Return Function(ArgTypeName Arg1 ) \
+	{ return this->Contract. Function (ArgName Arg1 ); }
+#define Function2(Capability, Return, Function, Arg1, Arg2) Return Function(ArgTypeName Arg1 , ArgTypeName Arg2 ) \
+	{ return this->Contract. Function (ArgName Arg1 , ArgName Arg2 ); }
+#define Function3(Capability, Return, Function, Arg1, Arg2, Arg3) Return Function(ArgTypeName Arg1 , ArgTypeName Arg2 , ArgTypeName Arg3 ) \
+	{ return this->Contract. Function (ArgName Arg1 , ArgName Arg2 , ArgName Arg3); }
+#define Function4(Capability, Return, Function, Arg1, Arg2, Arg3, Arg4) Return Function(ArgTypeName Arg1 , ArgTypeName Arg2 , ArgTypeName Arg3 , ArgTypeName Arg4 ) \
+	{ return this->Contract. Function (ArgName Arg1 , ArgName Arg2 , ArgName Arg3 , ArgName Arg4); }
+
+#define ConstFunction0(Capability, Return, Function) Return Function() const \
+	{ return this->Contract. Function (); }
+#define ConstFunction1(Capability, Return, Function, Arg1) Return Function(ArgTypeName Arg1 ) const \
+	{ return this->Contract. Function (ArgName Arg1 ); }
+#define ConstFunction2(Capability, Return, Function, Arg1, Arg2) Return Function(ArgTypeName Arg1 , ArgTypeName Arg2 ) const \
+	{ return this->Contract. Function (ArgName Arg1 , ArgName Arg2 ); }
+#define ConstFunction3(Capability, Return, Function, Arg1, Arg2, Arg3) Return Function(ArgTypeName Arg1 , ArgTypeName Arg2 , ArgTypeName Arg3 ) const \
+	{ return this->Contract. Function (ArgName Arg1 , ArgName Arg2 , ArgName Arg3); }
+#define ConstFunction4(Capability, Return, Function, Arg1, Arg2, Arg3, Arg4) Return Function(ArgTypeName Arg1 , ArgTypeName Arg2 , ArgTypeName Arg3 , ArgTypeName Arg4 ) const \
+	{ return this->Contract. Function (ArgName Arg1 , ArgName Arg2 , ArgName Arg3 , ArgName Arg4); }
+
+#define CapabilityStart(Capability) \
+	template<typename... Capabilities> \
+	struct TForwarder< Capability , Capabilities...> : TForwarder<Capabilities...> \
+	{ \
+		explicit TForwarder(IContextContract& Contract) : TForwarder<Capabilities...>(Contract) {}
+
+#define CapabilityEnd(Capability) \
+	}; \
+	struct Capability##Forwarder : TForwarder< Capability , I##Capability <ForwardTesterBase>> {}; \
+	static_assert(!std::is_abstract_v<Capability##Forwarder>, \
+		"One or more functions are not implemented in the " #Capability " forwarder.");
+
+#define WithWrappers 1
+	
+#include "Elements/Framework/TypedElementQueryCapabilities.inl"
+
+#undef ArgTypeName
+#undef ArgName
+#undef Function0
+#undef Function1
+#undef Function2
+#undef Function3
+#undef Function4
+#undef ConstFunction0
+#undef ConstFunction1
+#undef ConstFunction2
+#undef ConstFunction3
+#undef ConstFunction4
+#undef CapabilityStart
+#undef CapabilityEnd
+#undef WithWrappers
+} // namespace UE::Editor::DataStorage::Queries::Private
\ No newline at end of file
diff --git a/Engine/Source/Runtime/TypedElementFramework/Public/Elements/Framework/TypedElementQueryContext.h b/Engine/Source/Runtime/TypedElementFramework/Public/Elements/Framework/TypedElementQueryContext.h
new file mode 100644
--- /dev/null
+++ b/Engine/Source/Runtime/TypedElementFramework/Public/Elements/Framework/TypedElementQueryContext.h
@@ -0,0 +1,43 @@
+// Copyright Epic Games, Inc. All Rights Reserved.
+
+#pragma once
+
+#include <type_traits>
+#include "Elements/Framework/TypedElementQueryCapabilities.h"
+#include "Elements/Framework/TypedElementQueryCapabilityForwarder.h"
+
+namespace UE::Editor::DataStorage::Queries
+{
+	struct IContextContract;
+
+	template<typename T>
+	concept ContextCapability = std::is_base_of_v<IContextCapability, T> && !std::is_same_v<IContextCapability, T>;
+
+	template<ContextCapability... Capabilities>
+	struct TCapabilityStore : Capabilities... {};
+
+	/** 
+	 * Template to composite a query context using context capabilities.
+	 * Any query callback that requires interacting with the editor data storage requires a context to get access. Each context capability
+	 * provides access to a different kind of functionality. Some of the capabilities are mutually exclusive and not all places that accept
+	 * a query callback support the same capabilities. It's therefore recommended to only include the capabilities that are needed by the
+	 * query callback to improve re-usability or use a predefined query context if the query callback doesn't need to be reused.
+	 */
+	template<ContextCapability... CapabilityTypes>
+	struct TQueryContext final : Private::TForwarder<CapabilityTypes...>
+	{
+		using Capabilities = TCapabilityStore<CapabilityTypes...>;
+
+		TQueryContext(IContextContract& Contract);
+	};
+	
+	//
+	// Implementations.
+	//
+	
+	template<ContextCapability... CapabilityTypes>
+	TQueryContext<CapabilityTypes...>::TQueryContext(IContextContract& Contract)
+		: Private::TForwarder<CapabilityTypes...>(Contract)
+	{
+	}
+} // namespace UE::Editor::DataStorage::Queries
diff --git a/Engine/Source/Runtime/TypedElementFramework/Public/Elements/Framework/TypedElementQueryContextImplementation.h b/Engine/Source/Runtime/TypedElementFramework/Public/Elements/Framework/TypedElementQueryContextImplementation.h
new file mode 100644
--- /dev/null
+++ b/Engine/Source/Runtime/TypedElementFramework/Public/Elements/Framework/TypedElementQueryContextImplementation.h
@@ -0,0 +1,84 @@
+// Copyright Epic Games, Inc. All Rights Reserved.
+
+#pragma once
+
+#include "Elements/Framework/TypedElementQueryCapabilities.h"
+#include "Elements/Framework/TypedElementQueryContract.h"
+#include "Elements/Framework/TypedElementQueryContext.h"
+#include "Elements/Framework/TypedElementQueryFunctions.h"
+
+
+#define ArgTypeName(Type, Name) Type Name
+
+#define Function0(Capability, Return, Function) virtual Return Function() override;
+#define Function1(Capability, Return, Function, Arg1) virtual Return Function(ArgTypeName Arg1 ) override;
+#define Function2(Capability, Return, Function, Arg1, Arg2) virtual Return Function(ArgTypeName Arg1 , ArgTypeName Arg2 ) override;
+#define Function3(Capability, Return, Function, Arg1, Arg2, Arg3) virtual Return Function(ArgTypeName Arg1 , ArgTypeName Arg2 , ArgTypeName Arg3 ) override;
+#define Function4(Capability, Return, Function, Arg1, Arg2, Arg3, Arg4) virtual Return Function(ArgTypeName Arg1 , ArgTypeName Arg2 , ArgTypeName Arg3 , ArgTypeName Arg4 ) override;
+
+#define ConstFunction0(Capability, Return, Function) virtual Return Function() const override;
+#define ConstFunction1(Capability, Return, Function, Arg1) virtual Return Function(ArgTypeName Arg1 ) const override;
+#define ConstFunction2(Capability, Return, Function, Arg1, Arg2) virtual Return Function(ArgTypeName Arg1 , ArgTypeName Arg2 ) const override;
+#define ConstFunction3(Capability, Return, Function, Arg1, Arg2, Arg3) virtual Return Function(ArgTypeName Arg1 , ArgTypeName Arg2 , ArgTypeName Arg3 ) const override;
+#define ConstFunction4(Capability, Return, Function, Arg1, Arg2, Arg3, Arg4) virtual Return Function(ArgTypeName Arg1 , ArgTypeName Arg2 , ArgTypeName Arg3 , ArgTypeName Arg4 ) const override;
+
+#define CapabilityStart(Capability)
+#define CapabilityEnd(Capability)
+
+#define WithWrappers 1
+
+namespace UE::Editor::DataStorage::Queries
+{
+	/** 
+	 * Base used for context implementations. This will verify if a context is correctly implemented.
+	 * Classes that implement capabilities should use this template to automatically satisfy the requirements
+	 * of a query contract. Any missing functionality will be filled in by this template with a placeholder
+	 * function that asserts. The supported capabilities are automatically extracted from the implementation type
+	 * but can be overruled by providing a list of capabilities as extra template arguments. Note that any
+	 * capabilities provided this way do need to be implemented by the implementation type.
+	 */
+	template<typename ImplementationType, ContextCapability... SupportedCapabilities>
+	class TQueryContextImpl final : public IContextContract
+	{
+	public:
+		template<typename... TArgs>
+		explicit TQueryContextImpl(TArgs&&... Args);
+
+		virtual ~TQueryContextImpl() override = default;
+
+		static bool SupportsCapabilities(TConstArrayView<FName> Capabilities);
+
+#include "Elements/Framework/TypedElementQueryCapabilities.inl"
+
+		// Generic access
+		template<typename ReturnType>
+		bool CheckCompatiblity(const TQueryFunction<ReturnType>& Function) const;
+		const ImplementationType& GetContextImplementation() const;
+		ImplementationType& GetContextImplementation();
+	
+	private:
+		template<ContextCapability RequestedCapability>
+		static constexpr bool SupportsCapability();
+		static bool SupportsCapability(const FName& Capability);
+
+		ImplementationType Implementation;
+	};
+}
+// namespace UE::Editor::DataStorage::Queries
+
+#undef ArgTypeName
+#undef Function0
+#undef Function1
+#undef Function2
+#undef Function3
+#undef Function4
+#undef ConstFunction0
+#undef ConstFunction1
+#undef ConstFunction2
+#undef ConstFunction3
+#undef ConstFunction4
+#undef CapabilityStart
+#undef CapabilityEnd
+#undef WithWrappers
+
+#include "Elements/Framework/TypedElementQueryContextImplementation.inl"
diff --git a/Engine/Source/Runtime/TypedElementFramework/Public/Elements/Framework/TypedElementQueryContextImplementation.inl b/Engine/Source/Runtime/TypedElementFramework/Public/Elements/Framework/TypedElementQueryContextImplementation.inl
new file mode 100644
--- /dev/null
+++ b/Engine/Source/Runtime/TypedElementFramework/Public/Elements/Framework/TypedElementQueryContextImplementation.inl
@@ -0,0 +1,179 @@
+// Copyright Epic Games, Inc. All Rights Reserved.
+
+#include <type_traits>
+
+namespace UE::Editor::DataStorage::Queries
+{
+	template<typename QueryContext, ContextCapability... SupportedCapabilities>
+	template<typename... TArgs>
+	TQueryContextImpl<QueryContext, SupportedCapabilities...>::TQueryContextImpl(TArgs&&... Args)
+		: Implementation(Forward<TArgs>(Args)...)
+	{
+	}
+
+	template<typename QueryContext, ContextCapability... SupportedCapabilities>
+	template<ContextCapability RequestedCapability>
+	constexpr bool TQueryContextImpl<QueryContext, SupportedCapabilities...>::SupportsCapability()
+	{
+		return std::is_base_of_v<RequestedCapability, QueryContext> && (std::is_same_v<RequestedCapability, SupportedCapabilities> || ...);
+	}
+
+	template<typename QueryContext, ContextCapability... SupportedCapabilities>
+	bool TQueryContextImpl<QueryContext, SupportedCapabilities...>::SupportsCapability(const FName& Capability)
+	{
+		return IContextContract::SupportsCapability<QueryContext>(Capability) && ((SupportedCapabilities::Name == Capability) || ...);
+	}
+
+	template<typename QueryContext, ContextCapability... SupportedCapabilities>
+	bool TQueryContextImpl<QueryContext, SupportedCapabilities...>::SupportsCapabilities(TConstArrayView<FName> Capabilities)
+	{
+		for (const FName& Capability : Capabilities)
+		{
+			if (!SupportsCapability(Capability))
+			{
+				return false;
+			}
+		}
+		return true;
+	}
+
+	// Use a macro for this since it's impossible to do with a template since the function pointer that would be needed might not exist.
+	// The macro also has the added benefit that it can print the function name for better reporting.
+#define CallFunction(Capability, ReturnValue, FunctionName, ...) \
+	if constexpr (SupportsCapability< Capability >()) \
+	{ \
+		return Implementation.FunctionName( __VA_ARGS__ ); \
+	} \
+	else \
+	{ \
+		checkf(false, TEXT("Function '" #FunctionName "' in capability '" #Capability"' is not supported by the current query context implementation.")); \
+		return ReturnValue {} ; \
+	}
+
+#define ArgTypeName(Type, Name) Type Name
+#define ArgName(Type, Name) Name
+
+#define FunctionCommon(ReturnType) \
+	template<typename QueryContext, ContextCapability... SupportedCapabilities> \
+	ReturnType TQueryContextImpl<QueryContext, SupportedCapabilities...>::
+
+#define Function0(Capability, Return, Function) \
+	FunctionCommon(Return) Function() \
+	{ \
+		CallFunction(Capability, Return, Function ) \
+	}
+#define Function1(Capability, Return, Function, Arg1) \
+	FunctionCommon(Return) Function(ArgTypeName Arg1) \
+	{ \
+		CallFunction(Capability, Return, Function, ArgName Arg1 ) \
+	}
+#define Function2(Capability, Return, Function, Arg1, Arg2) \
+	FunctionCommon(Return) Function(ArgTypeName Arg1, ArgTypeName Arg2) \
+	{ \
+		CallFunction(Capability, Return, Function, ArgName Arg1, ArgName Arg2 ) \
+	}
+#define Function3(Capability, Return, Function, Arg1, Arg2, Arg3) \
+	FunctionCommon(Return) Function(ArgTypeName Arg1, ArgTypeName Arg2, ArgTypeName Arg3) \
+	{ \
+		CallFunction(Capability, Return, Function, ArgName Arg1, ArgName Arg2, ArgName Arg3 ) \
+	}
+#define Function4(Capability, Return, Function, Arg1, Arg2, Arg3, Arg4) \
+	FunctionCommon(Return) Function(ArgTypeName Arg1, ArgTypeName Arg2, ArgTypeName Arg3, ArgTypeName Arg4) \
+	{ \
+		CallFunction(Capability, Return, Function, ArgName Arg1, ArgName Arg2, ArgName Arg3, ArgName Arg4 ) \
+	}
+
+#define ConstFunction0(Capability, Return, Function) \
+	FunctionCommon(Return) Function() const \
+	{ \
+		CallFunction(Capability, Return, Function ) \
+	}
+#define ConstFunction1(Capability, Return, Function, Arg1) \
+	FunctionCommon(Return) Function(ArgTypeName Arg1) const \
+	{ \
+		CallFunction(Capability, Return, Function, ArgName Arg1 ) \
+	}
+#define ConstFunction2(Capability, Return, Function, Arg1, Arg2) \
+	FunctionCommon(Return) Function(ArgTypeName Arg1, ArgTypeName Arg2) const \
+	{ \
+		CallFunction(Capability, Return, Function, ArgName Arg1, ArgName Arg2 ) \
+	}
+#define ConstFunction3(Capability, Return, Function, Arg1, Arg2, Arg3) \
+	FunctionCommon(Return) Function(ArgTypeName Arg1, ArgTypeName Arg2, ArgTypeName Arg3) const \
+	{ \
+		CallFunction(Capability, Return, Function, ArgName Arg1, ArgName Arg2, ArgName Arg3 ) \
+	}
+#define ConstFunction4(Capability, Return, Function, Arg1, Arg2, Arg3, Arg4) \
+	FunctionCommon(Return) Function(ArgTypeName Arg1, ArgTypeName Arg2, ArgTypeName Arg3, ArgTypeName Arg4) const \
+	{ \
+		CallFunction(Capability, Return, Function, ArgName Arg1, ArgName Arg2, ArgName Arg3, ArgName Arg4 ) \
+	}
+
+#define CapabilityStart(Capability)
+#define CapabilityEnd(Capability)
+
+
+#define CallFunction(Capability, ReturnValue, FunctionName, ...) \
+	if constexpr (SupportsCapability< Capability >()) \
+	{ \
+		return Implementation.FunctionName( __VA_ARGS__ ); \
+	} \
+	else \
+	{ \
+		checkf(false, TEXT("Function '" #FunctionName "' in capability '" #Capability"' is not supported by the current query context implementation.")); \
+		return ReturnValue {} ; \
+	}
+
+#include "Elements/Framework/TypedElementQueryCapabilities.inl"
+
+#undef ArgTypeName
+#undef ArgName
+#undef FunctionCommon
+#undef Function0
+#undef Function1
+#undef Function2
+#undef Function3
+#undef Function4
+#undef ConstFunction0
+#undef ConstFunction1
+#undef ConstFunction2
+#undef ConstFunction3
+#undef ConstFunction4
+#undef CapabilityStart
+#undef CapabilityEnd
+#undef CallFunction
+
+	//
+	// Generic access
+	// 
+
+	template<typename QueryContext, ContextCapability... SupportedCapabilities>
+	template<typename ReturnType>
+	bool TQueryContextImpl<QueryContext, SupportedCapabilities...>::CheckCompatiblity(const TQueryFunction<ReturnType>& Function) const
+	{
+		bool Result = true;
+		for (const FName& Capability : Function.Capabilities)
+		{
+			if (!SupportsCapability(Capability))
+			{
+				Result = false;
+				checkf(false, TEXT("Requested query context '%s' is not support by query context implementation '%s'."), 
+					*Capability.ToString(), *(QueryContext::Name).ToString());
+			}
+		}
+		return false;
+	}
+
+	template<typename QueryContext, ContextCapability... SupportedCapabilities>
+	const QueryContext& TQueryContextImpl<QueryContext, SupportedCapabilities...>::GetContextImplementation() const
+	{
+		return Implementation;
+	}
+
+	template<typename QueryContext, ContextCapability... SupportedCapabilities>
+	QueryContext& TQueryContextImpl<QueryContext, SupportedCapabilities...>::GetContextImplementation()
+	{
+		return Implementation;
+	}
+}
+// namespace UE::Editor::DataStorage::Queries
diff --git a/Engine/Source/Runtime/TypedElementFramework/Public/Elements/Framework/TypedElementQueryContract.h b/Engine/Source/Runtime/TypedElementFramework/Public/Elements/Framework/TypedElementQueryContract.h
new file mode 100644
--- /dev/null
+++ b/Engine/Source/Runtime/TypedElementFramework/Public/Elements/Framework/TypedElementQueryContract.h
@@ -0,0 +1,129 @@
+// Copyright Epic Games, Inc. All Rights Reserved.
+
+#pragma once
+
+#include "Elements/Framework/TypedElementQueryCapabilities.h"
+
+namespace UE::Editor::DataStorage::Queries
+{
+	namespace Private
+	{
+		template<template<typename> typename Base, template<typename> typename... Bases>
+		struct IContextContractCombinerImpl : Base<IContextContractCombinerImpl<Bases>...>
+		{
+			virtual ~IContextContractCombinerImpl() override = default;
+		};
+
+		template<template<typename> typename Base>
+		struct IContextContractCombinerImpl<Base> : Base<IContextCapability>
+		{
+			virtual ~IContextContractCombinerImpl() override = default;
+		};
+
+		template<template<typename> typename... Bases>
+		struct IContextContractCombiner : IContextContractCombinerImpl<Bases>...
+		{
+			virtual ~IContextContractCombiner() override = default;
+
+		private:
+			// Work around for MSVC struggling to correctly expand fold expressions when providing the implementation for a template with 
+			// templated types.
+			template<typename Implementation, template<typename> typename Base>
+			constexpr static bool ImplementsCapability()
+			{
+				return std::is_base_of_v<Base<IContextCapability>, Implementation>;
+			}
+
+			template<typename Implementation, template<typename> typename Base>
+			static bool ImplementsCapability(const FName& Capability)
+			{
+				return ImplementsCapability<Implementation, Base>() && Base<IContextCapability>::Name == Capability;
+			}
+
+		public:
+			template<typename Implementation>
+			static bool SupportsCapability(const FName& Capability)
+			{
+				return (ImplementsCapability<Implementation, Bases>(Capability) || ...);
+			}
+
+			template<typename Implementation>
+			constexpr static int32 CountSupportedCapabilities()
+			{
+				return 0 + ((ImplementsCapability<Implementation, Bases>() ? 1 : 0) + ...);
+			}
+
+			template<typename Implementation>
+			static TConstArrayView<FName> SupportedCapabilitiesList()
+			{
+				static TConstArrayView<FName> Result = []()
+					{
+						constexpr int32 CapabilityCount = CountSupportedCapabilities<Implementation>();
+						static FName Capabilities[CapabilityCount] = {};
+
+						int32 Index = 0;
+						(
+							[&]
+							{
+								if constexpr (ImplementsCapability<Implementation, Bases>())
+								{
+									Capabilities[Index++] = Bases<IContextCapability>::Name;
+								}
+							}(), ...);
+
+						return TConstArrayView<FName>(Capabilities, CapabilityCount);
+					}();
+				return Result;
+			}
+		};
+
+		//Used to remove the first template argument as that's a placeholder to allow the X-macros to work with the template.
+		template<typename, template<typename> typename... Bases>
+		struct IPreContextContractCombiner : IContextContractCombiner<Bases...>
+		{
+			virtual ~IPreContextContractCombiner() override = default;
+		};
+} // namespace Private
+
+#define Function0(Capability, Return, Function)
+#define Function1(Capability, Return, Function, Arg1)
+#define Function2(Capability, Return, Function, Arg1, Arg2)
+#define Function3(Capability, Return, Function, Arg1, Arg2, Arg3)
+#define Function4(Capability, Return, Function, Arg1, Arg2, Arg3, Arg4)
+#define ConstFunction0(Capability, Return, Function)
+#define ConstFunction1(Capability, Return, Function, Arg1)
+#define ConstFunction2(Capability, Return, Function, Arg1, Arg2)
+#define ConstFunction3(Capability, Return, Function, Arg1, Arg2, Arg3)
+#define ConstFunction4(Capability, Return, Function, Arg1, Arg2, Arg3, Arg4)
+
+#define CapabilityStart(Capability) , I##Capability
+#define CapabilityEnd(Capability)
+
+	/** 
+	 * Contract between context and implementation to be able to communicate.
+	 * The contract contains all functions for all capabilities that are available. Based on the supported capabilities
+	 * an implementation may opt to only partially implement the contract, with the remaining functions asserting. On 
+	 * the opposite side a context may restrict what funtions on the contract can be called based on the requested
+	 * capabilities. Through the query function the capabilities on both sides are kept aligned, resulting in no
+	 * function on the contract being callable if they're not implemented.
+	 */
+	struct IContextContract : Private::IPreContextContractCombiner<int
+#include "Elements/Framework/TypedElementQueryCapabilities.inl"
+	>
+	{
+		virtual ~IContextContract() override = default;
+	};
+
+#undef Function0
+#undef Function1
+#undef Function2
+#undef Function3
+#undef Function4
+#undef ConstFunction0
+#undef ConstFunction1
+#undef ConstFunction2
+#undef ConstFunction3
+#undef ConstFunction4
+#undef CapabilityStart
+#undef CapabilityEnd
+} // namespace UE::Editor::DataStorage::Queries
\ No newline at end of file
diff --git a/Engine/Source/Runtime/TypedElementFramework/Public/Elements/Framework/TypedElementQueryFunctions.h b/Engine/Source/Runtime/TypedElementFramework/Public/Elements/Framework/TypedElementQueryFunctions.h
new file mode 100644
--- /dev/null
+++ b/Engine/Source/Runtime/TypedElementFramework/Public/Elements/Framework/TypedElementQueryFunctions.h
@@ -0,0 +1,26 @@
+// Copyright Epic Games, Inc. All Rights Reserved.
+
+#pragma once
+
+#include "Containers/ContainersFwd.h"
+#include "Elements/Framework/TypedElementQueryCapabilities.h"
+#include "Templates/Function.h"
+#include "UObject/NameTypes.h"
+
+namespace UE::Editor::DataStorage::Queries
+{
+	struct IContextContract;
+
+	/** Storage for a function that used as part of a query.  */
+	template<typename ReturnType>
+	struct TQueryFunction
+	{
+		TConstArrayView<FName> Capabilities;
+		TFunction<ReturnType(IContextContract& Contract)> Function;
+	};
+
+	template<typename ReturnType, typename FunctionType>
+	TQueryFunction<ReturnType> BuildQueryFunction(FunctionType&& Function);
+} // namespace UE::Editor::DataStorage::Queries
+
+#include "Elements/Framework/TypedElementQueryFunctions.inl"
diff --git a/Engine/Source/Runtime/TypedElementFramework/Public/Elements/Framework/TypedElementQueryFunctions.inl b/Engine/Source/Runtime/TypedElementFramework/Public/Elements/Framework/TypedElementQueryFunctions.inl
new file mode 100644
--- /dev/null
+++ b/Engine/Source/Runtime/TypedElementFramework/Public/Elements/Framework/TypedElementQueryFunctions.inl
@@ -0,0 +1,101 @@
+// Copyright Epic Games, Inc. All Rights Reserved.
+
+#include <type_traits>
+#include "Elements/Framework/TypedElementQueryCapabilities.h"
+#include "Elements/Framework/TypedElementQueryCapabilityForwarder.h"
+#include "Elements/Framework/TypedElementQueryContract.h"
+
+namespace UE::Editor::DataStorage::Queries
+{
+	namespace Private
+	{
+		template<typename T> concept FunctorType = requires{ &T::operator(); };
+		template<typename T> concept ContextType = std::is_base_of_v<FQueryContext, T>;
+	
+		template<typename T>
+		struct TFunctionInfoImpl
+		{
+			static constexpr bool bIsValidQueryFunction = false;
+			using ContextType = void;
+			using ReturnType = void;
+			static constexpr int32 CapabilityCount = 0;
+			inline static const TConstArrayView<FName> Capabilities;
+		};
+
+		template<typename Return, ContextType Context, typename... Args>
+		struct TFunctionInfoImpl<Return (*)(Context&, Args...)>
+		{
+			static constexpr bool bIsValidQueryFunction = true;
+			using ContextType = Context;
+			using ReturnType = Return;
+			static constexpr int32 CapabilityCount = IContextContract::CountSupportedCapabilities<typename Context::Capabilities>();
+			inline static const TConstArrayView<FName> Capabilities = IContextContract::SupportedCapabilitiesList<typename Context::Capabilities>();
+		};
+
+		template<typename Class, typename Return, ContextType Context, typename... Args>
+		struct TFunctionInfoImpl<Return(Class::*)(Context&, Args...)>
+		{
+			static constexpr bool bIsValidQueryFunction = true;
+			using ContextType = Context;
+			using ReturnType = Return;
+			static constexpr int32 CapabilityCount = IContextContract::CountSupportedCapabilities<typename Context::Capabilities>();
+			inline static const TConstArrayView<FName> Capabilities = IContextContract::SupportedCapabilitiesList<typename Context::Capabilities>();
+		};
+
+		template<typename Class, typename Return, ContextType Context, typename... Args>
+		struct TFunctionInfoImpl<Return(Class::*)(Context&, Args...) const>
+		{
+			static constexpr bool bIsValidQueryFunction = true;
+			using ContextType = Context;
+			using ReturnType = Return;
+			static constexpr int32 CapabilityCount = IContextContract::CountSupportedCapabilities<typename Context::Capabilities>();
+			inline static const TConstArrayView<FName> Capabilities = IContextContract::SupportedCapabilitiesList<typename Context::Capabilities>();
+		};
+
+		template<bool IsFunctor, typename FunctionType>
+		struct TFunctionInfoSelection {};
+
+		template<typename FunctionType>
+		struct TFunctionInfoSelection<true, FunctionType>
+		{
+			using Type = TFunctionInfoImpl<decltype(&FunctionType::operator())>;
+		};
+
+		template<typename FunctionType>
+		struct TFunctionInfoSelection<false, FunctionType>
+		{
+			using Type = TFunctionInfoImpl<FunctionType>;
+		};
+		
+		template<typename FunctionType>
+		using TFunctionInfo = TFunctionInfoSelection<FunctorType<FunctionType>, FunctionType>::Type;
+	}
+
+	template<typename ReturnType, typename FunctionType>
+	TQueryFunction<ReturnType> BuildQueryFunction(FunctionType&& Function)
+	{
+		using FunctionInfo = Private::TFunctionInfo<FunctionType>;
+
+		static_assert(FunctionInfo::bIsValidQueryFunction, "The provided function is not compatible with a query callback. "
+			"Likely errors are missing a query context or having incompatible arguments.");
+		static_assert(std::is_convertible_v<typename FunctionInfo::ReturnType, ReturnType>,
+			"The result of the provided function is not compatible with the query function that's being constructed.");
+
+		TQueryFunction<ReturnType> Result;
+		Result.Capabilities = FunctionInfo::Capabilities;
+		Result.Function =
+			[Callback = Forward<FunctionType>(Function)](IContextContract& Contract)
+			{
+				typename FunctionInfo::ContextType Context(Contract);
+				if constexpr (std::is_same_v<ReturnType, void>)
+				{
+					Callback(Context);
+				}
+				else
+				{
+					return Callback(Context);
+				}
+			};
+		return Result;
+	}
+} // namespace UE::Editor::DataStorage::Queries
diff --git a/Engine/Source/Runtime/TypedElementFramework/Public/Elements/Interfaces/TypedElementDataStorageInterface.h b/Engine/Source/Runtime/TypedElementFramework/Public/Elements/Interfaces/TypedElementDataStorageInterface.h
--- a/Engine/Source/Runtime/TypedElementFramework/Public/Elements/Interfaces/TypedElementDataStorageInterface.h
+++ b/Engine/Source/Runtime/TypedElementFramework/Public/Elements/Interfaces/TypedElementDataStorageInterface.h
@@ -13,7 +13,11 @@
 #include "Elements/Common/TypedElementQueryDescription.h"
 #include "Elements/Common/TypedElementQueryTypes.h"
 #include "Elements/Framework/TypedElementColumnUtils.h"
+#include "Elements/Framework/TypedElementQueryCapabilities.h"
+#include "Elements/Framework/TypedElementQueryFunctions.h"
 #include "Elements/Interfaces/TypedElementQueryStorageInterfaces.h"
+#include "Elements/Framework/TypedElementRowHandleArray.h"
+#include "Elements/Framework/TypedElementRowHandleArrayView.h"
 #include "Features/IModularFeature.h"
 #include "Math/NumericLimits.h"
 #include "Misc/EnumClassFlags.h"
@@ -175,6 +179,12 @@ public:
 	/** Checks whether or not a row has been reserved but not yet assigned to a table. */
 	virtual bool IsRowAssigned(RowHandle Row) const = 0;
 
+	/** Filters the provided rows using the query filter and stores the results in the provided list. */
+	virtual void FilterRowsBy(FRowHandleArray& Result, FRowHandleArrayView Input, Queries::TQueryFunction<bool>& Filter) = 0;
+	/** Filters the provided rows using the query filter and stores the results in the provided list. */
+	template<typename FilterFunction>
+	void FilterRowsBy(FRowHandleArray& Result, FRowHandleArrayView Input, FilterFunction&& Filter);
+
 	/**
 	 * @section Column management
 	 */
@@ -396,7 +406,7 @@ public:
 	 * if there are no columns to match against.
 	 */
 	virtual void ActivateQueries(FName ActivationName) = 0;
-	
+
 	/**
 	 * @section Mapping
 	 * @description
@@ -461,7 +471,7 @@ public:
 	/**
 	 * @section Deprecated
 	 */
-	/*UE_DEPRECATED(5.6, "Use 'LookUpMappedRow' instead and use the new FMapKey(View) instead of the explicit 'GenerateIndexHash'-functions.")
+	UE_DEPRECATED(5.6, "Use 'LookUpMappedRow' instead and use the new FMapKey(View) instead of the explicit 'GenerateIndexHash'-functions.")
 	inline RowHandle FindIndexedRow(IndexHash Index) const;
 	UE_DEPRECATED(5.6, "Use 'MapRow' instead and use the new FMapKey(View) instead of the explicit 'GenerateIndexHash'-functions.")
 	inline void IndexRow(IndexHash Index, RowHandle Row);
@@ -483,7 +493,7 @@ public:
 	UE_DEPRECATED(5.7, "Use the version of 'RemapRow' that uses a domain.")
 	inline void RemapRow(const FMapKeyView& OriginalKey, FMapKey NewKey);
 	UE_DEPRECATED(5.7, "Use the version of 'RemoveRowMapping' that uses a domain.")
-	inline void RemoveRowMapping(const FMapKeyView& Key);*/
+	inline void RemoveRowMapping(const FMapKeyView& Key);
 };
 
 // Implementations
@@ -513,6 +523,12 @@ TableHandle ICoreProvider::RegisterTable(
 	return RegisterTable(SourceTable, { Columns::StaticStruct()... }, Name);
 }
 
+template<typename FilterFunction>
+void ICoreProvider::FilterRowsBy(FRowHandleArray& Result, FRowHandleArrayView Input, FilterFunction&& Filter)
+{
+	FilterRowsBy(Result, Input, Queries::BuildQueryFunction<bool>(Forward<FilterFunction>(Filter)));
+}
+
 template<TColumnType Column>
 void ICoreProvider::AddColumn(RowHandle Row)
 {
@@ -723,7 +739,7 @@ SystemType* ICoreProvider::GetExternalSystem()
 	return reinterpret_cast<SystemType*>(GetExternalSystemAddress(SystemType::StaticClass()));
 }
 
-/*void ICoreProvider::ReindexRow(
+void ICoreProvider::ReindexRow(
 	IndexHash OriginalIndex, IndexHash NewIndex, RowHandle Row)
 {
 	RemapRow(NAME_None, FMapKeyView(OriginalIndex), FMapKey(NewIndex));
@@ -780,7 +796,7 @@ void ICoreProvider::RemapRow(const FMapKeyView& OriginalKey, FMapKey NewKey)
 void ICoreProvider::RemoveRowMapping(const FMapKeyView& Key)
 {
 	RemoveRowMapping(NAME_None, Key);
-}*/
+}
 
 template<TDynamicColumnTemplate DynamicColumnTemplateType>
 void ICoreProvider::ForEachDynamicColumn(TFunctionRef<void(const UScriptStruct& Type)> Callback) const
