[Teds] Rename Elements/Common/TypedElement...
to DataStorage/... aligning to namespace and class names

diff --git a/Engine/Source/Runtime/TypedElementFramework/Private/DataStorage/CommonTypes.cpp b/Engine/Source/Runtime/TypedElementFramework/Private/DataStorage/CommonTypes.cpp
new file mode 100644
--- /dev/null
+++ b/Engine/Source/Runtime/TypedElementFramework/Private/DataStorage/CommonTypes.cpp
@@ -0,0 +1,28 @@
+// Copyright Epic Games, Inc. All Rights Reserved.
+
+#include "DataStorage/CommonTypes.h"
+
+namespace UE
+{
+	namespace Editor::DataStorage
+	{
+		const FName& FValueTag::GetName() const
+		{
+			return Name;
+		}
+		
+		FValueTag::FValueTag(const FName& InTypeName)
+			: Name(InTypeName)
+		{}
+		
+		uint32 GetTypeHash(const FValueTag& InName)
+		{
+			return GetTypeHash(InName.Name);
+		}
+
+		uint32 GetTypeHash(const FDynamicColumnDescription& Descriptor)
+		{
+			return HashCombineFast(PointerHash(Descriptor.TemplateType), GetTypeHash(Descriptor.Identifier));
+		}
+	}
+}
diff --git a/Engine/Source/Runtime/TypedElementFramework/Private/DataStorage/Debug/Log.cpp b/Engine/Source/Runtime/TypedElementFramework/Private/DataStorage/Debug/Log.cpp
new file mode 100644
--- /dev/null
+++ b/Engine/Source/Runtime/TypedElementFramework/Private/DataStorage/Debug/Log.cpp
@@ -0,0 +1,5 @@
+// Copyright Epic Games, Inc. All Rights Reserved.
+
+#include "DataStorage/Debug/Log.h"
+
+DEFINE_LOG_CATEGORY(LogEditorDataStorage)
\ No newline at end of file
diff --git a/Engine/Source/Runtime/TypedElementFramework/Private/DataStorage/Features.cpp b/Engine/Source/Runtime/TypedElementFramework/Private/DataStorage/Features.cpp
new file mode 100644
--- /dev/null
+++ b/Engine/Source/Runtime/TypedElementFramework/Private/DataStorage/Features.cpp
@@ -0,0 +1,23 @@
+// Copyright Epic Games, Inc. All Rights Reserved.
+
+#include "DataStorage/Features.h"
+
+namespace UE::Editor::DataStorage
+{
+	const FName StorageFeatureName = TEXT("EditorDataStorage");
+	const FName CompatibilityFeatureName = TEXT("EditorDataStorageCompatibility");
+	const FName UiFeatureName = TEXT("EditorDataStorageUi");
+	
+	FSimpleMulticastDelegate& OnEditorDataStorageFeaturesEnabled()
+	{
+		static FSimpleMulticastDelegate OnTedsFeaturesEnabled;
+		return OnTedsFeaturesEnabled;
+	}
+
+	bool AreEditorDataStorageFeaturesEnabled()
+	{
+		return IModularFeatures::Get().IsModularFeatureAvailable(StorageFeatureName)
+			&& IModularFeatures::Get().IsModularFeatureAvailable(CompatibilityFeatureName)
+			&& IModularFeatures::Get().IsModularFeatureAvailable(UiFeatureName);
+	}
+} // namespace UE::Editor::DataStorage
diff --git a/Engine/Source/Runtime/TypedElementFramework/Private/DataStorage/MapKey.cpp b/Engine/Source/Runtime/TypedElementFramework/Private/DataStorage/MapKey.cpp
new file mode 100644
--- /dev/null
+++ b/Engine/Source/Runtime/TypedElementFramework/Private/DataStorage/MapKey.cpp
@@ -0,0 +1,443 @@
+// Copyright Epic Games, Inc. All Rights Reserved.
+
+#include "DataStorage/MapKey.h"
+
+namespace UE::Editor::DataStorage
+{
+	struct FHasher
+	{
+		uint64 operator()(FEmptyVariantState)	{ return 0; }
+		uint64 operator()(const void* Value)	{ return reinterpret_cast<uint64>(Value); }
+		uint64 operator()(const UObject* Value) { return reinterpret_cast<uint64>(Value); }
+		uint64 operator()(int64 Value)			{ return static_cast<uint64>(Value); }
+		uint64 operator()(uint64 Value)			{ return static_cast<uint64>(Value); }
+		uint64 operator()(float Value)			{ return static_cast<uint64>(Value); }
+		uint64 operator()(double Value)			{ return static_cast<uint64>(Value); }
+		uint64 operator()(const FString& Value)
+		{
+			return CityHash64(reinterpret_cast<const char*>(*Value), Value.Len() * sizeof(**Value));
+		}
+		uint64 operator()(const FStringView& Value)
+		{
+			return CityHash64(reinterpret_cast<const char*>(Value.GetData()), static_cast<uint32>(Value.NumBytes()));
+		}
+		uint64 operator()(const FName& Value)	{ return Value.ToUnstableInt(); }
+		uint64 operator()(const FName* Value)	{ return Value->ToUnstableInt(); }
+		uint64 operator()(const FSoftObjectPath* Value)
+		{
+			FTopLevelAssetPath TopLevelAssetPath = Value->GetAssetPath();
+			uint64 Hash = TopLevelAssetPath.GetPackageName().ToUnstableInt();
+			Hash = CityHash128to64({ Hash, TopLevelAssetPath.GetAssetName().ToUnstableInt() });
+			const FString& SubPath = Value->GetSubPathString();
+			return CityHash64WithSeed(reinterpret_cast<const char*>(*SubPath), SubPath.Len() * sizeof(**SubPath), Hash);
+		}
+		uint64 operator()(const TUniquePtr<FSoftObjectPath>& Value) { return operator()(Value.Get()); }
+	};
+
+	struct FKeyToKeyViewConverter
+	{
+		template<typename T>
+		FMapKeyView::KeyViewType operator()(const T& Value)
+		{
+			return FMapKeyView::KeyViewType(TInPlaceType<T>(), Value);
+		}
+
+		FMapKeyView::KeyViewType operator()(const FString& Value)
+		{
+			return FMapKeyView::KeyViewType(TInPlaceType<FStringView>(), Value);
+		}
+
+		FMapKeyView::KeyViewType operator()(const FName& Value)
+		{
+			return FMapKeyView::KeyViewType(TInPlaceType<const FName*>(), &Value);
+		}
+
+		FMapKeyView::KeyViewType operator()(const TUniquePtr<FSoftObjectPath>& Value)
+		{
+			return FMapKeyView::KeyViewType(TInPlaceType<const FSoftObjectPath*>(), Value.Get());
+		}
+	};
+
+	struct FKeyViewToKeyConverter
+	{
+		FMapKey operator()(const FEmptyVariantState& Value)
+		{
+			return FMapKey();
+		}
+
+		template<typename T>
+		FMapKey operator()(const T& Value)
+		{
+			return FMapKey(Value);
+		}
+
+		FMapKey operator()(const FStringView& Value)
+		{
+			return FMapKey(FString(Value));
+		}
+
+		FMapKey operator()(const FName* Value)
+		{
+			return FMapKey(*Value);
+		}
+
+		FMapKey operator()(const FSoftObjectPath* Value)
+		{
+			return FMapKey(*Value);
+		}
+	};
+
+	struct FKeyMove
+	{
+		FMapKey::KeyType operator()(FEmptyVariantState&)
+		{
+			return FMapKey::KeyType(TInPlaceType<FEmptyVariantState>());
+		}
+
+		template<typename T>
+		FMapKey::KeyType operator()(T&& Other)
+		{
+			return FMapKey::KeyType(TInPlaceType<T>(), Forward<T>(Other));
+		}
+	};
+
+	struct FKeyCopy
+	{
+		FMapKey::KeyType operator()(const FEmptyVariantState&)
+		{
+			return FMapKey::KeyType(TInPlaceType<FEmptyVariantState>());
+		}
+
+		template<typename T>
+		FMapKey::KeyType operator()(const T& Other)
+		{
+			return FMapKey::KeyType(TInPlaceType<T>(), Other);
+		}
+
+		FMapKey::KeyType operator()(const TUniquePtr<FSoftObjectPath>& Other)
+		{
+			return FMapKey::KeyType(TInPlaceType<TUniquePtr<FSoftObjectPath>>(), MakeUnique<FSoftObjectPath>(*Other));
+		}
+	};
+
+	struct FKeyComparer
+	{
+		const FMapKey::KeyType& Key;
+		explicit FKeyComparer(const FMapKey::KeyType& Key) : Key(Key) {}
+
+		bool operator()(FEmptyVariantState)
+		{
+			return Key.IsType<FEmptyVariantState>();
+		}
+
+		template<typename T>
+		bool operator()(const T& Other)
+		{
+			const auto* Local = Key.TryGet<T>();
+			return Local ? *Local == Other : false;
+		}
+
+		bool operator()(const TUniquePtr<FSoftObjectPath>& Other)
+		{
+			if (const TUniquePtr<FSoftObjectPath>* Local = Key.TryGet<TUniquePtr<FSoftObjectPath>>())
+			{
+				if (Local->IsValid() && Other.IsValid())
+				{
+					return *(Local->Get()) == *Other;
+				}
+			}
+			return false;
+		}
+	};
+
+	struct FKeyViewComparer
+	{
+		const FMapKeyView::KeyViewType& Key;
+		explicit FKeyViewComparer(const FMapKeyView::KeyViewType& Key) : Key(Key) {}
+
+		bool operator()(FEmptyVariantState)
+		{
+			return Key.IsType<FEmptyVariantState>();
+		}
+
+		template<typename T>
+		bool operator()(const T& Other)
+		{
+			const T* Local = Key.TryGet<T>();
+			return Local ? *Local == Other : false;
+		}
+
+		bool operator()(const FString& Other)
+		{
+			const FStringView* Local = Key.TryGet<FStringView>();
+			return Local ? *Local == Other : false;
+		}
+
+		bool operator()(const FStringView& Other)
+		{
+			const FStringView* Local = Key.TryGet<FStringView>();
+			return Local ? *Local == Other : false;
+		}
+
+		bool operator()(const FName& Other)
+		{
+			const FName* const* Local = Key.TryGet<const FName*>();
+			return Local ? **Local == Other : false;
+		}
+
+		bool operator()(const FName* Other)
+		{
+			const FName* const* Local = Key.TryGet<const FName*>();
+			return Local ? **Local == *Other : false;
+		}
+
+		bool operator()(const TUniquePtr<FSoftObjectPath>& Other)
+		{
+			if (Other.IsValid())
+			{
+				const FSoftObjectPath* const* Local = Key.TryGet<const FSoftObjectPath*>();
+				return Local ? **Local == *Other : false;
+			}
+			return false;
+		}
+
+		bool operator()(const FSoftObjectPath* Other)
+		{
+			const FSoftObjectPath* const* Local = Key.TryGet<const FSoftObjectPath*>();
+			return Local ? **Local == *Other : false;
+		}
+	};
+
+	struct FToString
+	{
+		FString operator()(FEmptyVariantState) { return TEXT("Empty"); }
+		FString operator()(const void* Value) { return FString::Printf(TEXT("Pointer(0x%p)"), Value); }
+		FString operator()(const UObject* Value) { return FString::Printf(TEXT("UObject(0x%p)"), Value); }
+		FString operator()(int64 Value) { return FString::Printf(TEXT("%lld"), Value); }
+		FString operator()(uint64 Value) { return FString::Printf(TEXT("%llu"), Value); }
+		FString operator()(float Value) { return FString::SanitizeFloat(Value);; }
+		FString operator()(double Value) { return FString::SanitizeFloat(Value);; }
+		FString operator()(const FString& Value) { return Value; }
+		FString operator()(const FStringView& Value) { return FString(Value); }
+		FString operator()(const FName& Value) { return Value.ToString(); }
+		FString operator()(const FName* Value) { return Value->ToString(); }
+		FString operator()(const FSoftObjectPath* Value) { return Value->ToString(); }
+		FString operator()(const TUniquePtr<FSoftObjectPath>& Value) { return Value ? Value->ToString() : TEXT("Empty path"); }
+	};
+
+	//
+	// FMapKey
+	//
+
+	FMapKey::FMapKey(FMapKey&& Rhs)
+		: Key(Visit(FKeyMove(), MoveTemp(Rhs.Key)))
+	{
+		Rhs.Clear();
+	}
+
+	FMapKey& FMapKey::operator=(FMapKey&& Rhs)
+	{
+		Key = Visit(FKeyMove(), MoveTemp(Rhs.Key));
+		Rhs.Clear();
+		return *this;
+	}
+
+	FMapKey::FMapKey(const FMapKey& Key)
+		: Key(Visit(FKeyCopy(), Key.Key))
+	{
+	}
+
+	FMapKey::FMapKey(const void* Key)
+		: Key(TInPlaceType<const void*>(), Key)
+	{
+	}
+
+	FMapKey::FMapKey(const UObject* Key)
+		: Key(TInPlaceType<const UObject*>(), Key)
+	{
+	}
+
+	FMapKey::FMapKey(int64 Key)
+		: Key(TInPlaceType<int64>(), Key)
+	{
+	}
+
+	FMapKey::FMapKey(uint64 Key)
+		: Key(TInPlaceType<uint64>(), Key)
+	{
+	}
+
+	FMapKey::FMapKey(float Key)
+		: Key(TInPlaceType<float>(), Key)
+	{
+	}
+
+	FMapKey::FMapKey(double Key)
+		: Key(TInPlaceType<double>(), Key)
+	{
+	}
+
+	FMapKey::FMapKey(FString Key)
+		: Key(TInPlaceType<FString>(), MoveTemp(Key))
+	{
+	}
+	
+	FMapKey::FMapKey(FName Key)
+		: Key(TInPlaceType<FName>(), MoveTemp(Key))
+	{
+	}
+
+	FMapKey::FMapKey(FSoftObjectPath Key)
+		: Key(TInPlaceType<TUniquePtr<FSoftObjectPath>>(), MakeUnique<FSoftObjectPath>(MoveTemp(Key)))
+	{
+	}
+
+	uint64 FMapKey::CalculateHash() const
+	{
+		return Visit(FHasher(), Key);
+	}
+	
+	bool FMapKey::IsSet() const
+	{
+		return !Key.IsType<FEmptyVariantState>();
+	}
+
+	void FMapKey::Clear()
+	{
+		Key.Emplace<FEmptyVariantState>();
+	}
+
+	FString FMapKey::ToString() const
+	{
+		return Visit(FToString(), Key);
+	}
+
+	FMapKey& FMapKey::operator=(const FMapKey& Rhs)
+	{
+		Key = Visit(FKeyCopy(), Rhs.Key);
+		return *this;
+	}
+
+	bool FMapKey::operator==(const FMapKey& Rhs) const
+	{
+		return Visit(FKeyComparer(Key), Rhs.Key);
+	}
+
+	bool FMapKey::operator==(const FMapKeyView& Rhs) const
+	{
+		return Rhs == *this;
+	}
+
+	bool FMapKey::operator!=(const FMapKey& Rhs) const
+	{
+		return !(*this == Rhs);
+	}
+
+	bool FMapKey::operator!=(const FMapKeyView& Rhs) const
+	{
+		return !(*this == Rhs);
+	}
+
+
+
+	//
+	// FMapKeyView
+	//
+
+	FMapKeyView::FMapKeyView(const FMapKey& InKey)
+		: Key(Visit(FKeyToKeyViewConverter(), InKey.Key))
+	{
+	}
+
+	FMapKeyView::FMapKeyView(const void* Key)
+		: Key(TInPlaceType<const void*>(), Key)
+	{
+	}
+
+	FMapKeyView::FMapKeyView(const UObject* Key)
+		: Key(TInPlaceType<const UObject*>(), Key)
+	{
+	}
+
+	FMapKeyView::FMapKeyView(int64 Key)
+		: Key(TInPlaceType<int64>(), Key)
+	{
+	}
+
+	FMapKeyView::FMapKeyView(uint64 Key)
+		: Key(TInPlaceType<uint64>(), Key)
+	{
+	}
+
+	FMapKeyView::FMapKeyView(float Key)
+		: Key(TInPlaceType<float>(), Key)
+	{
+	}
+
+	FMapKeyView::FMapKeyView(double Key)
+		: Key(TInPlaceType<double>(), Key)
+	{
+	}
+
+	FMapKeyView::FMapKeyView(const FString& Key)
+		: Key(TInPlaceType<FStringView>(), Key)
+	{
+	}
+	
+	FMapKeyView::FMapKeyView(FStringView Key)
+		: Key(TInPlaceType<FStringView>(), Key)
+	{
+	}
+
+	FMapKeyView::FMapKeyView(const FName& Key)
+		: Key(TInPlaceType<const FName*>(), &Key)
+	{
+	}
+
+	FMapKeyView::FMapKeyView(const FSoftObjectPath& Key)
+		: Key(TInPlaceType<const FSoftObjectPath*>(), &Key)
+	{
+	}
+
+	uint64 FMapKeyView::CalculateHash() const
+	{
+		return Visit(FHasher(), Key);
+	}
+
+	FString FMapKeyView::ToString() const
+	{
+		return Visit(FToString(), Key);
+	}
+
+	FMapKeyView& FMapKeyView::operator=(const FMapKey& InKey)
+	{
+		Key = Visit(FKeyToKeyViewConverter(), InKey.Key);
+		return *this;
+	}
+
+	FMapKey FMapKeyView::CreateKey() const
+	{
+		return Visit(FKeyViewToKeyConverter(), Key);
+	}
+
+	bool FMapKeyView::operator==(const FMapKey& Rhs) const
+	{
+		return Visit(FKeyViewComparer(Key), Rhs.Key);
+	}
+
+	bool FMapKeyView::operator==(const FMapKeyView& Rhs) const
+	{
+		return Visit(FKeyViewComparer(Key), Rhs.Key);
+	}
+
+	bool FMapKeyView::operator!=(const FMapKey& Rhs) const
+	{
+		return !(*this == Rhs);
+	}
+
+	bool FMapKeyView::operator!=(const FMapKeyView& Rhs) const
+	{
+		return !(*this == Rhs);
+	}
+}
+// namespace UE::Editor::DataStorage
diff --git a/Engine/Source/Runtime/TypedElementFramework/Private/DataStorage/Queries/QueryConditions.cpp b/Engine/Source/Runtime/TypedElementFramework/Private/DataStorage/Queries/QueryConditions.cpp
new file mode 100644
--- /dev/null
+++ b/Engine/Source/Runtime/TypedElementFramework/Private/DataStorage/Queries/QueryConditions.cpp
@@ -0,0 +1,728 @@
+// Copyright Epic Games, Inc. All Rights Reserved.
+
+#include "DataStorage/Queries/Conditions.h"
+
+#include "RHIStats.h"
+#include "Elements/Framework/TypedElementColumnUtils.h"
+#include "Elements/Interfaces/TypedElementDataStorageInterface.h"
+#include "UObject/Class.h"
+
+namespace UE::Editor::DataStorage::Queries
+{
+	FEditorStorageQueryConditionCompileContext::FEditorStorageQueryConditionCompileContext(ICoreProvider* InDataStorage)
+		: DataStorage(InDataStorage)
+	{
+	}
+
+	const UScriptStruct* FEditorStorageQueryConditionCompileContext::GenerateDynamicColumn(
+		const FDynamicColumnDescription& Description) const
+	{
+		return DataStorage->GenerateDynamicColumn(Description);
+	}
+
+	FConditions::FConditions()
+	{
+		Identifiers.Init(NAME_None, MaxColumnCount);
+		Columns.Init(nullptr, MaxColumnCount);
+		Tokens.Init(Token::None, MaxTokenCount);
+	}
+
+	FConditions::FConditions(FColumnBase Column)
+		: ColumnCount(1)
+	{
+		Identifiers.Init(NAME_None, MaxColumnCount);
+		Columns.Init(nullptr, MaxColumnCount);
+		Tokens.Init(Token::None, MaxTokenCount);
+		
+		Columns[0] = Column.TypeInfo;
+		Identifiers[0] = Column.Identifier;
+	}
+
+	void FConditions::AppendToString(FString& Output) const
+	{
+		checkf(bIsCompiled, TEXT("Query Conditions must call Compile() before you can use them"));
+		
+		if (TokenCount > 0)
+		{
+			Output += "{ ";
+			uint8_t ColumnIndex = 0;
+			if (Tokens[0] != Token::ScopeOpen)
+			{
+				ColumnIndex = 1;
+				AppendName(Output, Columns[0]);
+			}
+			for (uint8_t TokenIndex = 0; TokenIndex < TokenCount; ++TokenIndex)
+			{
+				switch (Tokens[TokenIndex])
+				{
+				case Token::And:
+					Output += " && ";
+					if (!EntersScopeNext(TokenIndex))
+					{
+						AppendName(Output, Columns[ColumnIndex++]);
+					}
+					break;
+				case Token::Or:
+					Output += " || ";
+					if (!EntersScopeNext(TokenIndex))
+					{
+						AppendName(Output, Columns[ColumnIndex++]);
+					}
+					break;
+				case Token::ScopeOpen:
+					Output += "( ";
+					if (!EntersScopeNext(TokenIndex))
+					{
+						AppendName(Output, Columns[ColumnIndex++]);
+					}
+					break;
+				case Token::ScopeClose:
+					Output += " )";
+					break;
+				default:
+					checkf(false, TEXT("Invalid query token"));
+					break;
+				}
+			}
+			Output += " }";
+		}
+	}
+
+	template <typename AvailableColumnType, typename ProjectionFunctionType>
+	bool FConditions::VerifyWithDynamicColumn(TConstArrayView<AvailableColumnType> AvailableColumns, uint64& Matches, ProjectionFunctionType Projection) const
+	{
+		// Upfront process AvailableColumns since AvailableColumns may be searched through more than once
+		TBitArray<TInlineAllocator<512>> DerivedFromDynamicTemplate;
+#if WITH_EDITORONLY_DATA
+		DerivedFromDynamicTemplate.SetNum(AvailableColumns.Num(), false);
+		for (int32 AvailableColumnIndex = 0; AvailableColumnIndex < AvailableColumns.Num(); ++AvailableColumnIndex)
+		{
+			const TWeakObjectPtr<const UScriptStruct>& AvailableColumn = Projection(AvailableColumns[AvailableColumnIndex]);
+			if (const UScriptStruct* AvailableColumnStruct = AvailableColumn.Get())
+			{
+				DerivedFromDynamicTemplate[AvailableColumnIndex] = ColumnUtils::IsDerivedFromDynamicTemplate(AvailableColumnStruct);
+			}
+		}
+#endif
+		
+		bool Result = VerifyBootstrap(
+			[this, &Matches, &AvailableColumns, &DerivedFromDynamicTemplate, &Projection](uint8_t ColumnIndex, TWeakObjectPtr<const UScriptStruct> Column)
+			{
+				bool bContains = false;
+				// Looping through each column in the condition, if the column is a template then determine if any of the available
+				// columns that are being matched against are derived from that template
+				if ((ColumnFlags[ColumnIndex] & EColumnFlags::DynamicColumnTemplate) == EColumnFlags::DynamicColumnTemplate)
+				{
+					for (int32 AvailableColumnIndex = 0; AvailableColumnIndex < AvailableColumns.Num(); ++AvailableColumnIndex)
+					{
+						// Only check if this available column is derived from our dynamic template if we know it derives from any DynamicTemplate
+						if (DerivedFromDynamicTemplate[AvailableColumnIndex])
+						{
+							if (Projection(AvailableColumns[AvailableColumnIndex])->IsChildOf(Column.Get()))
+							{
+								Matches |= (uint64)1 << ColumnIndex;
+								bContains = true;
+							}
+						}
+						else
+						{
+							if (Projection(AvailableColumns[AvailableColumnIndex]) == Column.Get())
+							{
+								Matches |= (uint64)1 << ColumnIndex;
+								bContains = true;
+							}
+						}
+					}
+				}
+				else
+				{
+					for (int32 AvailableColumnIndex = 0; AvailableColumnIndex < AvailableColumns.Num(); ++AvailableColumnIndex)
+					{
+						if (Projection(AvailableColumns[AvailableColumnIndex]) == Column.Get())
+						{
+							Matches |= (uint64)1 << ColumnIndex;
+							bContains = true;
+						}
+					}
+					//auto FindPredicate = ;
+					if ( AvailableColumns.FindByPredicate([&Column, &Projection](const AvailableColumnType& Entry) { return Projection(Entry) == Column.Get(); }) != nullptr)
+					{
+						Matches |= (uint64)1 << ColumnIndex;
+						bContains = true;
+					}
+				}
+				return bContains;
+			});
+
+		return Result;
+	}
+	
+	bool FConditions::Verify(TConstArrayView<FColumnBase> AvailableColumns) const
+	{
+		checkf(bIsCompiled, TEXT("Query Conditions must call Compile() before you can use them"));
+
+		bool Result;
+		if (!UsesDynamicTemplates())
+		{
+			Result = VerifyBootstrap(
+				[&AvailableColumns](uint8_t ColumnIndex, TWeakObjectPtr<const UScriptStruct> Column)
+				{
+					for (const FColumnBase& Target : AvailableColumns)
+					{
+						if (Target.TypeInfo == Column)
+						{
+							return true;
+						}
+					}
+					return false;
+				});
+		}
+		else
+		{
+			uint64 Matches = 0;
+			Result = VerifyWithDynamicColumn(AvailableColumns, Matches, [](const FColumnBase& Column) { return Column.TypeInfo; });
+		}
+		
+		return Result;
+	}
+
+	bool FConditions::Verify(TArray<TWeakObjectPtr<const UScriptStruct>>& MatchedColumns, TConstArrayView<FColumnBase> AvailableColumns,
+		bool AvailableColumnsAreSorted) const
+	{
+		static_assert(MaxColumnCount < 64, "Query conditions use a bit mask to locate matches. As a result MaxColumnCount can be larger than 64.");
+
+		checkf(bIsCompiled, TEXT("Query Conditions must call Compile() before you can use them"));
+		
+		uint64 Matches = 0;
+		bool Result;
+		if (!UsesDynamicTemplates())
+		{
+			Result = AvailableColumnsAreSorted
+				? VerifyBootstrap(
+					[&Matches, &AvailableColumns](uint8_t ColumnIndex, TWeakObjectPtr<const UScriptStruct> Column)
+					{
+						auto Projection = [](const FColumnBase& ColumnBase)
+						{
+							return ColumnBase.TypeInfo.Get();
+						};
+						if (Algo::BinarySearchBy(AvailableColumns, Column.Get(), Projection) != INDEX_NONE)
+						{
+							Matches |= (uint64)1 << ColumnIndex;
+							return true;
+						}
+						return false;
+					})
+				: VerifyBootstrap(
+					[&Matches, &AvailableColumns](uint8_t ColumnIndex, TWeakObjectPtr<const UScriptStruct> Column)
+					{
+						for (const FColumnBase& Target : AvailableColumns)
+						{
+							if (Target.TypeInfo == Column)
+							{
+								Matches |= (uint64)1 << ColumnIndex;
+								return true;
+							}
+						}
+						return false;
+					});
+		}
+		else
+		{
+			Result = VerifyWithDynamicColumn(AvailableColumns, Matches, [](const FColumnBase& ColumnBase) { return ColumnBase.TypeInfo; });
+		}
+		if (Result)
+		{
+			ConvertColumnBitToArray(MatchedColumns, Matches);
+		}
+		return Result;
+	}
+
+	bool FConditions::Verify(TConstArrayView<TWeakObjectPtr<const UScriptStruct>> AvailableColumns,
+		bool AvailableColumnsAreSorted) const
+	{
+		checkf(bIsCompiled, TEXT("Query Conditions must call Compile() before you can use them"));
+
+		bool Result;
+		if (!UsesDynamicTemplates())
+		{
+			Result = AvailableColumnsAreSorted
+				? VerifyBootstrap(
+					[&AvailableColumns](uint8_t ColumnIndex, TWeakObjectPtr<const UScriptStruct> Column)
+					{
+						auto Projection = [](TWeakObjectPtr<const UScriptStruct> ColumnBase)
+						{
+							return ColumnBase.Get();
+						};
+						return Algo::BinarySearchBy(AvailableColumns, Column.Get(), Projection) != INDEX_NONE;
+					})
+				: VerifyBootstrap(
+					[&AvailableColumns](uint8_t ColumnIndex, TWeakObjectPtr<const UScriptStruct> Column)
+					{
+						return AvailableColumns.Find(Column) != INDEX_NONE;
+					});
+		}
+		else
+		{
+			uint64 Matches = 0;
+			Result = VerifyWithDynamicColumn(AvailableColumns, Matches, [](TWeakObjectPtr<const UScriptStruct> ColumnType) { return ColumnType; });
+		}
+		return Result;
+	}
+
+	bool FConditions::Verify(TArray<TWeakObjectPtr<const UScriptStruct>>& MatchedColumns,
+		TConstArrayView<TWeakObjectPtr<const UScriptStruct>> AvailableColumns,
+		bool AvailableColumnsAreSorted) const
+	{
+		static_assert(MaxColumnCount < 64, "Query conditions use a bit mask to locate matches. As a result MaxColumnCount cannot be larger than 64.");
+		
+		checkf(bIsCompiled, TEXT("Query Conditions must call Compile() before you can use them"));
+		
+		uint64 Matches = 0;
+		bool Result;
+		if (!UsesDynamicTemplates())
+		{
+			Result = AvailableColumnsAreSorted
+				? VerifyBootstrap(
+					[&Matches, &AvailableColumns](uint8_t ColumnIndex, TWeakObjectPtr<const UScriptStruct> Column)
+					{
+						auto Projection = [](TWeakObjectPtr<const UScriptStruct> ColumnBase)
+						{
+							return ColumnBase.Get();
+						};
+						if (Algo::BinarySearchBy(AvailableColumns, Column.Get(), Projection) != INDEX_NONE)
+						{
+							Matches |= (uint64)1 << ColumnIndex;
+							return true;
+						}
+						return false;
+					})
+				: VerifyBootstrap(
+					[&Matches, &AvailableColumns](uint8_t ColumnIndex, TWeakObjectPtr<const UScriptStruct> Column)
+					{
+						if ( AvailableColumns.Find(Column) != INDEX_NONE)
+						{
+							Matches |= (uint64)1 << ColumnIndex;
+							return true;
+						}
+						return false;
+					});
+		}
+		else
+		{
+			Result = VerifyWithDynamicColumn(AvailableColumns, Matches, [](TWeakObjectPtr<const UScriptStruct> ColumnType) { return ColumnType; });
+		}
+		if (Result)
+		{
+			ConvertColumnBitToArray(MatchedColumns, Matches);
+		}
+		return Result;
+	}
+
+	bool FConditions::Verify(TSet<TWeakObjectPtr<const UScriptStruct>> AvailableColumns) const
+	{
+		checkf(bIsCompiled, TEXT("Query Conditions must call Compile() before you can use them"));
+		
+		return VerifyBootstrap(
+			[&AvailableColumns](uint8_t ColumnIndex, TWeakObjectPtr<const UScriptStruct> Column)
+			{
+				return AvailableColumns.Find(Column) != nullptr;
+			});
+	}
+
+	bool FConditions::Verify(ContainsCallback Callback) const
+	{
+		checkf(bIsCompiled, TEXT("Query Conditions must call Compile() before you can use them"));
+		
+		return VerifyBootstrap(Callback);
+	}
+
+	uint8_t FConditions::MinimumColumnMatchRequired() const
+	{
+		if (TokenCount > 0)
+		{
+			uint8_t Front = 0;
+			return MinimumColumnMatchRequiredRange(Front);
+		}
+		else if (ColumnCount == 1)
+		{
+			return 1;
+		}
+		else
+		{
+			return 0;
+		}
+	}
+
+	TConstArrayView<TWeakObjectPtr<const UScriptStruct>> FConditions::GetColumns() const
+	{
+		checkf(bIsCompiled, TEXT("Query Conditions must call Compile() before you can use them"));
+		
+		return TConstArrayView<TWeakObjectPtr<const UScriptStruct>>(Columns.GetData(), ColumnCount);
+	}
+
+	bool FConditions::IsEmpty() const
+	{
+		return ColumnCount == 0;
+	}
+
+	FConditions& FConditions::Compile(const IQueryConditionCompileContext& CompileContext)
+	{
+		if(bIsCompiled)
+		{
+			return *this;
+		}
+		
+		ColumnFlags.SetNumZeroed(ColumnCount);
+
+		// Resolve all dynamic columns by generating their UScriptStruct
+		for (uint8_t ColumnIndex = 0; ColumnIndex < ColumnCount; ++ColumnIndex)
+		{
+			if(Identifiers[ColumnIndex] != NAME_None)
+			{
+				const UScriptStruct* ResolvedDynamicColumn = CompileContext.GenerateDynamicColumn(
+					FDynamicColumnDescription
+						{
+							.TemplateType = Columns[ColumnIndex].Get(),
+							.Identifier = Identifiers[ColumnIndex]
+						});
+				
+				Columns[ColumnIndex] = ResolvedDynamicColumn;
+			}
+			else
+			{
+#if WITH_EDITORONLY_DATA
+				bool bIsDynamicTemplate = ColumnUtils::IsDynamicTemplate(Columns[ColumnIndex].Get());
+				if (bIsDynamicTemplate)
+				{
+					ColumnFlags[ColumnIndex] |= EColumnFlags::DynamicColumnTemplate;
+				}
+#endif
+			}
+		}
+
+		bIsCompiled = true;
+
+		return *this;
+	}
+
+	bool FConditions::IsCompiled() const
+	{
+		return bIsCompiled;
+	}
+
+	void FConditions::AppendName(FString& Output, TWeakObjectPtr<const UScriptStruct> TypeInfo) const
+	{
+#if WITH_EDITORONLY_DATA
+		static FName DisplayNameName(TEXT("DisplayName"));
+		if (const FString* Name = TypeInfo->FindMetaData(DisplayNameName))
+		{
+			Output += *Name;
+		}
+#else
+		Output += TEXT("<Unavailable>");
+#endif
+	}
+
+	bool FConditions::EntersScopeNext(uint8_t Index) const
+	{
+		return Index < (TokenCount - 1) && Tokens[Index + 1] == Token::ScopeOpen;
+	}
+
+	bool FConditions::EntersScope(uint8_t Index) const
+	{
+		return Tokens[Index] == Token::ScopeOpen;
+	}
+
+	bool FConditions::Contains(TWeakObjectPtr<const UScriptStruct> ColumnType, const TArray<FColumnBase>& AvailableColumns) const
+	{
+		for (const FColumnBase& Column : AvailableColumns)
+		{
+			if (ColumnType == Column.TypeInfo)
+			{
+				return true;
+			}
+		}
+		return false;
+	}
+
+	bool FConditions::VerifyBootstrap(ContainsCallback Contains) const
+	{
+		if (TokenCount > 0)
+		{
+			uint8_t TokenIndex = 0;
+			uint8_t ColumnIndex = 0;
+			return VerifyRange(TokenIndex, ColumnIndex, Forward<ContainsCallback>(Contains));
+		}
+		else if (ColumnCount == 1)
+		{
+			return Contains(0, Columns[0]);
+		}
+		// If there are no columns in the condition, everything passes it
+		else
+		{
+			return true;
+		}
+	}
+
+	bool FConditions::VerifyRange(uint8_t& TokenIndex, uint8_t& ColumnIndex, ContainsCallback Contains) const
+	{
+		auto Init = [&]()
+		{
+			if (EntersScope(TokenIndex))
+			{
+				return VerifyRange(++TokenIndex, ColumnIndex, Contains);
+			}
+			else
+			{
+				bool Result = Contains(ColumnIndex, Columns[ColumnIndex]);
+				ColumnIndex++;
+				return Result;
+			}
+		};
+		bool Result = Init();
+
+		while (TokenIndex < TokenCount)
+		{
+			Token T = Tokens[TokenIndex];
+			++TokenIndex;
+			switch (T)
+			{
+			case Token::And:
+				if (EntersScope(TokenIndex))
+				{
+					Result = VerifyRange(++TokenIndex, ColumnIndex, Contains) && Result;
+				}
+				else
+				{
+					Result = Contains(ColumnIndex, Columns[ColumnIndex]) && Result;
+					ColumnIndex++;
+				}
+				break;
+			case Token::Or:
+				if (EntersScope(TokenIndex))
+				{
+					Result = VerifyRange(++TokenIndex, ColumnIndex, Contains) || Result;
+				}
+				else
+				{
+					Result = Contains(ColumnIndex, Columns[ColumnIndex]) || Result;
+					ColumnIndex++;
+				}
+				break;
+			case Token::ScopeOpen:
+				checkf(false, TEXT("The scope open in a query should be called during processing as it should be captured by an earlier statement."));
+				break;
+			case Token::ScopeClose:
+				return Result;
+			default:
+				checkf(false, TEXT("Encountered an unknown query token."));
+				break;
+			}
+		}
+		return Result;
+	}
+
+	void FConditions::ConvertColumnBitToArray(TArray<TWeakObjectPtr<const UScriptStruct>>& MatchedColumns, uint64 ColumnBits) const
+	{
+		int Index = 0;
+		while (ColumnBits)
+		{
+			if (ColumnBits & 1)
+			{
+				MatchedColumns.Add(Columns[Index]);
+			}
+			++Index;
+			ColumnBits >>= 1;
+		}
+	}
+
+	uint8_t FConditions::MinimumColumnMatchRequiredRange(uint8_t& Front) const
+	{
+		uint8_t Result = EntersScope(Front) ? MinimumColumnMatchRequiredRange(++Front) : 1;
+		for (; Front < TokenCount; ++Front)
+		{
+			Token T = Tokens[Front];
+			switch (T)
+			{
+			case Token::And:
+				if (EntersScopeNext(Front))
+				{
+					Front += 2;
+					Result += MinimumColumnMatchRequiredRange(Front);
+				}
+				else
+				{
+					++Result;
+				}
+				break;
+			case Token::Or:
+				if (EntersScopeNext(Front))
+				{
+					Front += 2;
+					uint8_t Lhs = MinimumColumnMatchRequiredRange(Front);
+					Result = FMath::Min(Result, Lhs);
+				}
+				break;
+			case Token::ScopeOpen:
+				checkf(false, TEXT("The scope open in a query should be called during processing as it should be captured by an earlier statement."));
+				break;
+			case Token::ScopeClose:
+				++Front;
+				return Result;
+			default:
+				checkf(false, TEXT("Encountered an unknown query token."));
+				break;
+			}
+		}
+		return Result;
+	}
+
+	bool FConditions::UsesDynamicTemplates() const
+	{
+		for (EColumnFlags Flags : ColumnFlags)
+		{
+			if((Flags & EColumnFlags::DynamicColumnTemplate) == EColumnFlags::DynamicColumnTemplate)
+			{
+				return true;
+			}
+		}
+		return false;
+	}
+
+	void FConditions::AppendQuery(FConditions& Target, const FConditions& Source)
+	{
+		checkf(Target.ColumnCount + Source.ColumnCount < MaxColumnCount, TEXT("Too many columns in the query."));
+		for (uint8_t Index = 0; Index < Source.ColumnCount; ++Index)
+		{
+			Target.Columns[Target.ColumnCount + Index] = Source.Columns[Index];
+			Target.Identifiers[Target.ColumnCount + Index] = Source.Identifiers[Index];
+		}
+
+		checkf(Target.TokenCount + Source.TokenCount < MaxTokenCount, TEXT("Too many operations in the query. Try simplifying your query."));
+		for (uint8_t Index = 0; Index < Source.TokenCount; ++Index)
+		{
+			Target.Tokens[Target.TokenCount + Index] = Source.Tokens[Index];
+		}
+
+		Target.ColumnCount += Source.ColumnCount;
+		Target.TokenCount += Source.TokenCount;
+	}
+
+	FConditions operator&&(const FConditions& Lhs, FColumnBase Rhs)
+	{
+		FConditions Result = Lhs;
+		Result.Columns[Result.ColumnCount] = Rhs.TypeInfo;
+		Result.Identifiers[Result.ColumnCount++] = Rhs.Identifier;
+		Result.Tokens[Result.TokenCount++] = FConditions::Token::And;
+		return Result;
+	}
+
+	FConditions operator&&(const FConditions& Lhs, const FConditions& Rhs)
+	{
+		if (Lhs.IsEmpty())
+		{
+			return Rhs;
+		}
+		else if (Rhs.IsEmpty())
+		{
+			return Lhs;
+		}
+		
+		FConditions Result;
+		Result.Tokens[Result.TokenCount++] = FConditions::Token::ScopeOpen;
+		FConditions::AppendQuery(Result, Lhs);
+		Result.Tokens[Result.TokenCount++] = FConditions::Token::ScopeClose;
+
+		Result.Tokens[Result.TokenCount++] = FConditions::Token::And;
+
+		Result.Tokens[Result.TokenCount++] = FConditions::Token::ScopeOpen;
+		FConditions::AppendQuery(Result, Rhs);
+		Result.Tokens[Result.TokenCount++] = FConditions::Token::ScopeClose;
+
+		return Result;
+	}
+
+	FConditions operator&&(FColumnBase Lhs, FColumnBase Rhs)
+	{
+		FConditions Result(Lhs);
+		Result.Columns[Result.ColumnCount] = Rhs.TypeInfo;
+		Result.Identifiers[Result.ColumnCount++] = Rhs.Identifier;
+		Result.Tokens[Result.TokenCount++] = FConditions::Token::And;
+		return Result;
+	}
+
+	FConditions operator&&(FColumnBase Lhs, const FConditions& Rhs)
+	{
+		FConditions Result(Lhs);
+
+		if (!Rhs.IsEmpty())
+		{
+			Result.Tokens[Result.TokenCount++] = FConditions::Token::And;
+			Result.Tokens[Result.TokenCount++] = FConditions::Token::ScopeOpen;
+			FConditions::AppendQuery(Result, Rhs);
+			Result.Tokens[Result.TokenCount++] = FConditions::Token::ScopeClose;
+		}
+		
+		return Result;
+	}
+
+	FConditions operator||(const FConditions& Lhs, FColumnBase Rhs)
+	{
+		FConditions Result = Lhs;
+		Result.Columns[Result.ColumnCount] = Rhs.TypeInfo;
+		Result.Identifiers[Result.ColumnCount++] = Rhs.Identifier;
+		Result.Tokens[Result.TokenCount++] = FConditions::Token::Or;
+		return Result;
+	}
+
+	FConditions operator||(const FConditions& Lhs, const FConditions& Rhs)
+	{
+		if (Lhs.IsEmpty())
+		{
+			return Rhs;
+		}
+		else if (Rhs.IsEmpty())
+		{
+			return Lhs;
+		}
+		
+		FConditions Result;
+		Result.Tokens[Result.TokenCount++] = FConditions::Token::ScopeOpen;
+		FConditions::AppendQuery(Result, Lhs);
+		Result.Tokens[Result.TokenCount++] = FConditions::Token::ScopeClose;
+
+		Result.Tokens[Result.TokenCount++] = FConditions::Token::Or;
+
+		Result.Tokens[Result.TokenCount++] = FConditions::Token::ScopeOpen;
+		FConditions::AppendQuery(Result, Rhs);
+		Result.Tokens[Result.TokenCount++] = FConditions::Token::ScopeClose;
+
+		return Result;
+	}
+
+	FConditions operator||(FColumnBase Lhs, FColumnBase Rhs)
+	{
+		FConditions Result(Lhs);
+		Result.Columns[Result.ColumnCount] = Rhs.TypeInfo;
+		Result.Identifiers[Result.ColumnCount++] = Rhs.Identifier;
+		Result.Tokens[Result.TokenCount++] = FConditions::Token::Or;
+		return Result;
+	}
+
+	FConditions operator||(FColumnBase Lhs, const FConditions& Rhs)
+	{
+		FConditions Result(Lhs);
+		
+		if (!Rhs.IsEmpty())
+		{
+			Result.Tokens[Result.TokenCount++] = FConditions::Token::Or;
+			Result.Tokens[Result.TokenCount++] = FConditions::Token::ScopeOpen;
+			FConditions::AppendQuery(Result, Rhs);
+			Result.Tokens[Result.TokenCount++] = FConditions::Token::ScopeClose;
+
+		}
+		return Result;
+	}
+
+	ENUM_CLASS_FLAGS(FConditions::EColumnFlags)
+} // namespace UE::Editor::DataStorage::Queries
\ No newline at end of file
diff --git a/Engine/Source/Runtime/TypedElementFramework/Private/Elements/Common/EditorDataStorageFeatures.cpp b/Engine/Source/Runtime/TypedElementFramework/Private/Elements/Common/EditorDataStorageFeatures.cpp
deleted file mode 100644
diff --git a/Engine/Source/Runtime/TypedElementFramework/Private/Elements/Common/TypedElementCommonTypes.cpp b/Engine/Source/Runtime/TypedElementFramework/Private/Elements/Common/TypedElementCommonTypes.cpp
deleted file mode 100644
diff --git a/Engine/Source/Runtime/TypedElementFramework/Private/Elements/Common/TypedElementDataStorageLog.cpp b/Engine/Source/Runtime/TypedElementFramework/Private/Elements/Common/TypedElementDataStorageLog.cpp
deleted file mode 100644
diff --git a/Engine/Source/Runtime/TypedElementFramework/Private/Elements/Common/TypedElementMapKey.cpp b/Engine/Source/Runtime/TypedElementFramework/Private/Elements/Common/TypedElementMapKey.cpp
deleted file mode 100644
diff --git a/Engine/Source/Runtime/TypedElementFramework/Private/Elements/Common/TypedElementQueryConditions.cpp b/Engine/Source/Runtime/TypedElementFramework/Private/Elements/Common/TypedElementQueryConditions.cpp
deleted file mode 100644
diff --git a/Engine/Source/Runtime/TypedElementFramework/Public/DataStorage/CommonTypes.h b/Engine/Source/Runtime/TypedElementFramework/Public/DataStorage/CommonTypes.h
new file mode 100644
--- /dev/null
+++ b/Engine/Source/Runtime/TypedElementFramework/Public/DataStorage/CommonTypes.h
@@ -0,0 +1,106 @@
+// Copyright Epic Games, Inc. All Rights Reserved.
+
+#pragma once
+
+#include "DataStorage/Handles.h"
+
+#include "CommonTypes.generated.h"
+
+/**
+ * Base for the data structures for a column.
+ */
+USTRUCT()
+struct FEditorDataStorageColumn
+{
+	GENERATED_BODY()
+};
+
+/**
+ * Base for the data structures that act as tags to rows. Tags should not have any data.
+ */
+USTRUCT()
+struct FEditorDataStorageTag
+{
+	GENERATED_BODY()
+};
+
+namespace UE
+{
+	// Work around missing header/implementations on some platforms
+	namespace detail
+	{
+		template<typename T, typename U>
+		concept SameHelper = std::is_same_v<T, U>;
+	}
+	template<typename T, typename U>
+	concept same_as = detail::SameHelper<T, U> && detail::SameHelper<U, T>;
+
+	template<typename From, typename To>
+	concept convertible_to = std::is_convertible_v<From, To> && requires { static_cast<To>(std::declval<From>()); };
+
+	template<typename Derived, typename Base>
+	concept derived_from = std::is_base_of_v<Base, Derived> && std::is_convertible_v<const volatile Derived*, const volatile Base*>;
+
+	namespace Editor::DataStorage
+	{
+		using FColumn = FEditorDataStorageColumn;
+		using FTag = FEditorDataStorageTag;
+
+		/**
+		 * Defines a dynamic type for a value tag
+		 * Example:
+		 *   FValueTag ColorTagType(TEXT("Color"));
+		 *   FValueTag DirectionTagType(TEXT("Direction"));
+		 * A value tag can take on different values for each type.  This is set up when a tag is added to a row.
+		 */
+		class FValueTag
+		{
+		public:
+			TYPEDELEMENTFRAMEWORK_API explicit FValueTag(const FName& InName);
+			
+			TYPEDELEMENTFRAMEWORK_API const FName& GetName() const;
+			bool operator==(const FValueTag& Other) const = default;
+		private:
+			TYPEDELEMENTFRAMEWORK_API friend uint32 GetTypeHash(const FValueTag& InName);
+			FName Name;
+		};
+
+		template<typename T>
+		concept TValueTagType = std::is_same_v<T, FValueTag>;
+
+		struct FDynamicColumnDescription
+		{
+			const UScriptStruct* TemplateType;
+			FName Identifier;
+
+			TYPEDELEMENTFRAMEWORK_API friend uint32 GetTypeHash(const FDynamicColumnDescription& Descriptor);
+			bool operator==(const FDynamicColumnDescription&) const = default;
+		};
+		// Standard callbacks.
+
+		using RowCreationCallbackRef = TFunctionRef<void(RowHandle Row)>;
+		using ColumnCreationCallbackRef = TFunctionRef<void(void* Column, const UScriptStruct& ColumnType)>;
+		using ColumnListCallbackRef = TFunctionRef<void(const UScriptStruct& ColumnType)>;
+		using ColumnListWithDataCallbackRef = TFunctionRef<void(void* Column, const UScriptStruct& ColumnType)>;
+		using ColumnCopyOrMoveCallback = void (*)(const UScriptStruct& ColumnType, void* Destination, void* Source);
+
+		template<typename T>
+		concept THasDynamicColumnTemplateSpecifier = std::is_empty_v<typename T::EditorDataStorage_DynamicColumnTemplate>;
+		
+		template<typename T>
+		concept TDynamicColumnTemplate = (UE::derived_from<T, FColumn> || UE::derived_from<T, FTag>) && THasDynamicColumnTemplateSpecifier<T>;
+		
+		// Template concepts to enforce type correctness.
+		template<typename T>
+		concept TDataColumnType = UE::derived_from<T, FColumn> && !THasDynamicColumnTemplateSpecifier<T>;
+
+		template<typename T>
+		concept TTagColumnType = UE::derived_from<T, FTag> && !THasDynamicColumnTemplateSpecifier<T>;
+		
+		template<typename T>
+		concept TColumnType = TDataColumnType<T> || TTagColumnType<T>;
+
+		template<typename T>
+		concept TEnumType = std::is_enum_v<T>;
+	} // namespace Editor::DataStorage
+} // namespace UE
diff --git a/Engine/Source/Runtime/TypedElementFramework/Public/DataStorage/Debug/Log.h b/Engine/Source/Runtime/TypedElementFramework/Public/DataStorage/Debug/Log.h
new file mode 100644
--- /dev/null
+++ b/Engine/Source/Runtime/TypedElementFramework/Public/DataStorage/Debug/Log.h
@@ -0,0 +1,7 @@
+// Copyright Epic Games, Inc. All Rights Reserved.
+
+#pragma once
+
+#include "Logging/LogMacros.h"
+
+TYPEDELEMENTFRAMEWORK_API DECLARE_LOG_CATEGORY_EXTERN(LogEditorDataStorage, Log, All);
\ No newline at end of file
diff --git a/Engine/Source/Runtime/TypedElementFramework/Public/DataStorage/Features.h b/Engine/Source/Runtime/TypedElementFramework/Public/DataStorage/Features.h
new file mode 100644
--- /dev/null
+++ b/Engine/Source/Runtime/TypedElementFramework/Public/DataStorage/Features.h
@@ -0,0 +1,38 @@
+// Copyright Epic Games, Inc. All Rights Reserved.
+
+#pragma once
+
+#include "Delegates/Delegate.h"
+#include "Features/IModularFeatures.h"
+#include "UObject/NameTypes.h"
+#include "GenericPlatform/GenericPlatform.h"
+
+namespace UE::Editor::DataStorage
+{
+	TYPEDELEMENTFRAMEWORK_API const extern FName StorageFeatureName;
+	TYPEDELEMENTFRAMEWORK_API const extern FName CompatibilityFeatureName;
+	TYPEDELEMENTFRAMEWORK_API const extern FName UiFeatureName;
+
+	TYPEDELEMENTFRAMEWORK_API FSimpleMulticastDelegate& OnEditorDataStorageFeaturesEnabled();
+	TYPEDELEMENTFRAMEWORK_API bool AreEditorDataStorageFeaturesEnabled();
+
+	template<typename T>
+	const T* GetDataStorageFeature(const FName InName)
+	{
+		if (IModularFeatures::Get().IsModularFeatureAvailable(InName))
+		{
+			return &IModularFeatures::Get().GetModularFeature<T>(InName);
+		}
+		return nullptr;
+	}
+
+	template<typename T>
+	T* GetMutableDataStorageFeature(const FName InName)
+	{
+		if (IModularFeatures::Get().IsModularFeatureAvailable(InName))
+		{
+			return &IModularFeatures::Get().GetModularFeature<T>(InName);
+		}
+		return nullptr;
+	}
+} // namespace UE::Editor::DataStorage
diff --git a/Engine/Source/Runtime/TypedElementFramework/Public/DataStorage/Handles.h b/Engine/Source/Runtime/TypedElementFramework/Public/DataStorage/Handles.h
new file mode 100644
--- /dev/null
+++ b/Engine/Source/Runtime/TypedElementFramework/Public/DataStorage/Handles.h
@@ -0,0 +1,52 @@
+// Copyright Epic Games, Inc. All Rights Reserved.
+
+#pragma once
+
+#include "Math/NumericLimits.h"
+#include "Templates/TypeHash.h"
+
+#include "Handles.generated.h"
+
+namespace UE::Editor::DataStorage
+{
+	using TableHandle = uint64;
+	static constexpr TableHandle InvalidTableHandle = TNumericLimits<TableHandle>::Max();
+
+	using RowHandle = uint64;
+	static constexpr RowHandle InvalidRowHandle = 0;
+
+	using QueryHandle = uint64;
+	static constexpr QueryHandle InvalidQueryHandle = TNumericLimits<QueryHandle>::Max();
+} // namespace UE::Editor::DataStorage
+
+/*
+ * FTedsRowHandle is a strongly typed wrapper around UE::Editor::DataStorage::RowHandle and should only be used in cases where you need the extra info.
+ * E.g for reflection/UHT or for template specializing something that needs to know the semantics of the row handle.
+ * For all other cases, you should use the regular typedef UE::Editor::DataStorage::RowHandle
+ */
+USTRUCT()
+struct FTedsRowHandle
+{
+	GENERATED_BODY()
+
+	UPROPERTY()
+	uint64 RowHandle = UE::Editor::DataStorage::InvalidRowHandle;
+
+	operator UE::Editor::DataStorage::RowHandle () const
+	{
+		return RowHandle;
+	}
+
+	FTedsRowHandle& operator=(UE::Editor::DataStorage::RowHandle InRowHandle)
+	{
+		RowHandle = InRowHandle;
+		return *this;
+	}
+		
+	friend uint32 GetTypeHash(const FTedsRowHandle& Key)
+	{
+		return GetTypeHash(Key.RowHandle);
+	}
+};
+
+static_assert(sizeof(FTedsRowHandle::RowHandle) == sizeof(UE::Editor::DataStorage::RowHandle), "RowHandle and RowHandle wrapper sizes should match");
\ No newline at end of file
diff --git a/Engine/Source/Runtime/TypedElementFramework/Public/DataStorage/MapKey.h b/Engine/Source/Runtime/TypedElementFramework/Public/DataStorage/MapKey.h
new file mode 100644
--- /dev/null
+++ b/Engine/Source/Runtime/TypedElementFramework/Public/DataStorage/MapKey.h
@@ -0,0 +1,97 @@
+// Copyright Epic Games, Inc. All Rights Reserved.
+
+#pragma once
+
+#include "Containers/StringFwd.h"
+#include "Containers/UnrealString.h"
+#include "Misc/TVariant.h"
+#include "Templates/UniquePtr.h"
+#include "UObject/NameTypes.h"
+#include "UObject/SoftObjectPath.h"
+
+class UObject;
+
+namespace UE::Editor::DataStorage
+{
+	class FMapKeyView;
+
+	class FMapKey
+	{
+		friend class FMapKeyView;
+		friend struct FKeyComparer;
+		friend struct FKeyCopy;
+		friend struct FKeyMove;
+
+	public:
+		FMapKey() = default;
+		
+		TYPEDELEMENTFRAMEWORK_API FMapKey(FMapKey&& Rhs);
+		TYPEDELEMENTFRAMEWORK_API FMapKey& operator=(FMapKey&& Rhs);
+		
+		TYPEDELEMENTFRAMEWORK_API FMapKey(const FMapKey& Key);
+		TYPEDELEMENTFRAMEWORK_API explicit FMapKey(const void* Key);
+		TYPEDELEMENTFRAMEWORK_API explicit FMapKey(const UObject* Key);
+		TYPEDELEMENTFRAMEWORK_API explicit FMapKey(int64 Key);
+		TYPEDELEMENTFRAMEWORK_API explicit FMapKey(uint64 Key);
+		TYPEDELEMENTFRAMEWORK_API explicit FMapKey(float Key);
+		TYPEDELEMENTFRAMEWORK_API explicit FMapKey(double Key);
+		TYPEDELEMENTFRAMEWORK_API explicit FMapKey(FString Key);
+		TYPEDELEMENTFRAMEWORK_API explicit FMapKey(FName Key);
+		TYPEDELEMENTFRAMEWORK_API explicit FMapKey(FSoftObjectPath Key);
+
+		TYPEDELEMENTFRAMEWORK_API FMapKey& operator=(const FMapKey& Rhs);
+
+		TYPEDELEMENTFRAMEWORK_API bool operator==(const FMapKey& Rhs) const;
+		TYPEDELEMENTFRAMEWORK_API bool operator==(const FMapKeyView& Rhs) const;
+		TYPEDELEMENTFRAMEWORK_API bool operator!=(const FMapKey& Rhs) const;
+		TYPEDELEMENTFRAMEWORK_API bool operator!=(const FMapKeyView& Rhs) const;
+
+		TYPEDELEMENTFRAMEWORK_API uint64 CalculateHash() const;
+		TYPEDELEMENTFRAMEWORK_API bool IsSet() const;
+		TYPEDELEMENTFRAMEWORK_API void Clear();
+		TYPEDELEMENTFRAMEWORK_API FString ToString() const;
+
+	private:
+		using KeyType = TVariant<FEmptyVariantState, const void*, const UObject*, int64, uint64, float, double, FString, FName, TUniquePtr<FSoftObjectPath>>;
+
+		KeyType Key;
+	};
+
+	class FMapKeyView
+	{
+		friend struct FKeyToKeyViewConverter;
+		friend struct FKeyViewComparer;
+
+	public:
+		FMapKeyView() = default;
+		TYPEDELEMENTFRAMEWORK_API FMapKeyView(const FMapKey& InKey);
+		TYPEDELEMENTFRAMEWORK_API explicit FMapKeyView(const void* Key);
+		TYPEDELEMENTFRAMEWORK_API explicit FMapKeyView(const UObject* Key);
+		TYPEDELEMENTFRAMEWORK_API explicit FMapKeyView(int64 Key);
+		TYPEDELEMENTFRAMEWORK_API explicit FMapKeyView(uint64 Key);
+		TYPEDELEMENTFRAMEWORK_API explicit FMapKeyView(float Key);
+		TYPEDELEMENTFRAMEWORK_API explicit FMapKeyView(double Key);
+		TYPEDELEMENTFRAMEWORK_API explicit FMapKeyView(const FString& Key);
+		TYPEDELEMENTFRAMEWORK_API explicit FMapKeyView(FStringView Key);
+		TYPEDELEMENTFRAMEWORK_API explicit FMapKeyView(const FName& Key);
+		TYPEDELEMENTFRAMEWORK_API explicit FMapKeyView(const FSoftObjectPath& Key);
+
+		TYPEDELEMENTFRAMEWORK_API FMapKeyView& operator=(const FMapKey& InKey);
+
+		TYPEDELEMENTFRAMEWORK_API FMapKey CreateKey() const;
+
+		TYPEDELEMENTFRAMEWORK_API bool operator==(const FMapKey& Rhs) const;
+		TYPEDELEMENTFRAMEWORK_API bool operator==(const FMapKeyView& Rhs) const;
+		TYPEDELEMENTFRAMEWORK_API bool operator!=(const FMapKey& Rhs) const;
+		TYPEDELEMENTFRAMEWORK_API bool operator!=(const FMapKeyView& Rhs) const;
+
+		TYPEDELEMENTFRAMEWORK_API uint64 CalculateHash() const;
+		TYPEDELEMENTFRAMEWORK_API FString ToString() const;
+
+	private:
+		using KeyViewType = TVariant<FEmptyVariantState, const void*, const UObject*, int64, uint64, float, double, FStringView, const FName*, const FSoftObjectPath*>;
+
+		KeyViewType Key;
+	};
+}
+// namespace UE::Editor::DataStorage
diff --git a/Engine/Source/Runtime/TypedElementFramework/Public/DataStorage/Queries/Conditions.h b/Engine/Source/Runtime/TypedElementFramework/Public/DataStorage/Queries/Conditions.h
new file mode 100644
--- /dev/null
+++ b/Engine/Source/Runtime/TypedElementFramework/Public/DataStorage/Queries/Conditions.h
@@ -0,0 +1,197 @@
+// Copyright Epic Games, Inc. All Rights Reserved.
+
+#pragma once
+
+#include "DataStorage/CommonTypes.h"
+#include "Templates/SubclassOf.h"
+#include "UObject/WeakObjectPtrTemplates.h"
+
+class UScriptStruct;
+
+namespace UE::Editor::DataStorage
+{
+	class ICoreProvider;
+
+namespace Queries
+{
+	struct FColumnBase
+	{
+		TWeakObjectPtr<const UScriptStruct> TypeInfo = nullptr;
+		FName Identifier = NAME_None;
+
+	protected:
+		FColumnBase() = default;
+		constexpr explicit FColumnBase(TWeakObjectPtr<const UScriptStruct> ColumnTypeInfo, const FName& InIdentifier)
+			: TypeInfo(ColumnTypeInfo)
+			, Identifier(InIdentifier)
+		{}
+	};
+
+	template<typename T = void>
+	struct TColumn final : public FColumnBase
+	{
+		template <typename U = T> requires (!std::is_same_v<U, void>)
+			constexpr explicit TColumn(const FName& Identifier = NAME_None) : FColumnBase(T::StaticStruct(), Identifier) {}
+		
+		template <typename U = T> requires (std::is_same_v<U, void>)
+			constexpr explicit TColumn(TWeakObjectPtr<const UScriptStruct> ColumnTypeInfo, const FName& Identifier = NAME_None)
+				: FColumnBase(ColumnTypeInfo, Identifier)
+		{}
+	};
+
+	/*
+	 * Compile Context used to resolve dynamic columns when FConditions are compiled
+	 */
+	class IQueryConditionCompileContext 
+	{
+	public:
+		virtual ~IQueryConditionCompileContext() = default;
+		virtual const UScriptStruct* GenerateDynamicColumn(const FDynamicColumnDescription&) const = 0;
+	};
+
+	/*
+	 * Specialized compile context that accepts an ICoreProvider
+	 */
+	class FEditorStorageQueryConditionCompileContext : public IQueryConditionCompileContext 
+	{
+	public:
+		TYPEDELEMENTFRAMEWORK_API explicit FEditorStorageQueryConditionCompileContext(ICoreProvider* InDataStorage);
+		TYPEDELEMENTFRAMEWORK_API virtual const UScriptStruct* GenerateDynamicColumn(const FDynamicColumnDescription& Description) const override;
+		
+	private:
+		ICoreProvider* DataStorage;
+	};
+
+
+	/**
+	 * Product of boolean combination of multiple columns. This can be used to verify if a collection of columns match
+	 * the stored columns.
+	 * NOTE: You must call Compile() before you call any members accessing the conditions
+	 */
+	class FConditions final
+	{
+	public:
+		using ContainsCallback = TFunctionRef<bool(uint8_t ColumnIndex, TWeakObjectPtr<const UScriptStruct> Column)>;
+
+		TYPEDELEMENTFRAMEWORK_API friend FConditions operator&&(const FConditions& Lhs, FColumnBase Rhs);
+		TYPEDELEMENTFRAMEWORK_API friend FConditions operator&&(const FConditions& Lhs, const FConditions& Rhs);
+		TYPEDELEMENTFRAMEWORK_API friend FConditions operator&&(FColumnBase Lhs, FColumnBase Rhs);
+		TYPEDELEMENTFRAMEWORK_API friend FConditions operator&&(FColumnBase Lhs, const FConditions& Rhs);
+
+		TYPEDELEMENTFRAMEWORK_API friend FConditions operator||(const FConditions& Lhs, FColumnBase Rhs);
+		TYPEDELEMENTFRAMEWORK_API friend FConditions operator||(const FConditions& Lhs, const FConditions& Rhs);
+		TYPEDELEMENTFRAMEWORK_API friend FConditions operator||(FColumnBase Lhs, FColumnBase Rhs);
+		TYPEDELEMENTFRAMEWORK_API friend FConditions operator||(FColumnBase Lhs, const FConditions& Rhs);
+
+		TYPEDELEMENTFRAMEWORK_API FConditions();
+		// Not marked as "explicit" to allow conversion from a column. This means that conditions with a single
+		// argument can be written in the same way as ones that use combinations.
+		TYPEDELEMENTFRAMEWORK_API FConditions(FColumnBase Column);
+
+		// Compile must be called before using any functions that access the columns
+		TYPEDELEMENTFRAMEWORK_API FConditions& Compile(const IQueryConditionCompileContext& CompileContext);
+
+		// Check whether these query conditions have been compiled
+		TYPEDELEMENTFRAMEWORK_API bool IsCompiled() const;
+
+		/** Convert the conditions into a string and append them to the provided string. */
+		TYPEDELEMENTFRAMEWORK_API void AppendToString(FString& Output) const;
+
+		/** Runs the provided list of columns through the conditions and returns true if a valid combination of columns is found. */
+		TYPEDELEMENTFRAMEWORK_API bool Verify(TConstArrayView<FColumnBase> AvailableColumns) const;
+		/**
+		 * Runs the provided list of columns through the conditions and returns true if a valid combination of columns is found.
+		 * This version returns a list of the columns that were used to match the condition.
+		 */
+		TYPEDELEMENTFRAMEWORK_API bool Verify(
+			TArray<TWeakObjectPtr<const UScriptStruct>>& MatchedColumns,
+			TConstArrayView<FColumnBase> AvailableColumns,
+			bool AvailableColumnsAreSorted = false) const;
+		/** Runs the provided list of columns through the conditions and returns true if a valid combination of columns is found. */
+		TYPEDELEMENTFRAMEWORK_API bool Verify(TConstArrayView<TWeakObjectPtr<const UScriptStruct>> AvailableColumns,
+			bool AvailableColumnsAreSorted = false) const;
+		/**
+		 * Runs the provided list of columns through the conditions and returns true if a valid combination of columns is found.
+		 * This version returns a list of the columns that were used to match the condition.
+		 */
+		TYPEDELEMENTFRAMEWORK_API bool Verify(
+			TArray<TWeakObjectPtr<const UScriptStruct>>& MatchedColumns,
+			TConstArrayView<TWeakObjectPtr<const UScriptStruct>> AvailableColumns,
+			bool AvailableColumnsAreSorted = false) const;
+		/** Runs the provided list of columns through the conditions and returns true if a valid combination of columns is found. */
+		TYPEDELEMENTFRAMEWORK_API bool Verify(TSet<TWeakObjectPtr<const UScriptStruct>> AvailableColumns) const;
+		/** Runs through the list of query conditions and uses the callback to verify if a column is available. */
+		TYPEDELEMENTFRAMEWORK_API bool Verify(ContainsCallback Callback) const;
+
+		/** Returns the minimum number of columns needed for a successful match. */
+		TYPEDELEMENTFRAMEWORK_API uint8_t MinimumColumnMatchRequired() const;
+		/** Returns a list of all columns used. This can include duplicate columns. */
+		TYPEDELEMENTFRAMEWORK_API TConstArrayView<TWeakObjectPtr<const UScriptStruct>> GetColumns() const;
+		/** Whether or not there are any columns registered for operation.*/
+		TYPEDELEMENTFRAMEWORK_API bool IsEmpty() const;
+		/** Whether the conditions contain any columns that are dynamic templates. */
+		TYPEDELEMENTFRAMEWORK_API bool UsesDynamicTemplates() const;
+
+	private:
+		template<typename AvailableColumnType, typename ProjectionFunction>
+		bool VerifyWithDynamicColumn(TConstArrayView<AvailableColumnType> AvailableColumns, uint64& OutMatches, ProjectionFunction Projection) const;
+		
+		void AppendName(FString& Output, TWeakObjectPtr<const UScriptStruct> TypeInfo) const;
+
+		bool EntersScopeNext(uint8_t Index) const;
+		bool EntersScope(uint8_t Index) const;
+		bool Contains(TWeakObjectPtr<const UScriptStruct> ColumnType, const TArray<FColumnBase>& AvailableColumns) const;
+
+		bool VerifyBootstrap(ContainsCallback Contains) const;
+		bool VerifyRange(uint8_t& TokenIndex, uint8_t& ColumnIndex, ContainsCallback Contains) const;
+
+		void ConvertColumnBitToArray(TArray<TWeakObjectPtr<const UScriptStruct>>& MatchedColumns, uint64 ColumnBits) const;
+
+		uint8_t MinimumColumnMatchRequiredRange(uint8_t& Front) const;
+
+		static void AppendQuery(FConditions& Target, const FConditions& Source);
+
+		static constexpr SIZE_T MaxColumnCount = 32;
+		static constexpr SIZE_T MaxTokenCount = 64;
+
+		enum class Token : uint8_t
+		{
+			None,
+			And,
+			Or,
+			ScopeOpen,
+			ScopeClose
+		};
+
+		enum class EColumnFlags : uint8_t
+		{
+			DynamicColumnTemplate = 1 << 0
+		};
+		FRIEND_ENUM_CLASS_FLAGS(EColumnFlags);
+		
+		using ColumnArray = TArray<TWeakObjectPtr<const UScriptStruct>, TInlineAllocator<MaxColumnCount>>;
+		using TokenArray = TArray<Token, TInlineAllocator<MaxTokenCount>>;
+		using IdentifierArray = TArray<FName, TInlineAllocator<MaxColumnCount>>;
+		using ColumnFlagArray = TArray<EColumnFlags, TInlineAllocator<MaxColumnCount>>;
+
+		ColumnArray Columns;
+		ColumnFlagArray ColumnFlags;
+		TokenArray Tokens;
+		IdentifierArray Identifiers;
+		uint8_t ColumnCount = 0;
+		uint8_t TokenCount = 0;
+		bool bIsCompiled = false;
+	};
+
+	TYPEDELEMENTFRAMEWORK_API FConditions operator&&(const FConditions& Lhs, FColumnBase Rhs);
+	TYPEDELEMENTFRAMEWORK_API FConditions operator&&(const FConditions& Lhs, const FConditions& Rhs);
+	TYPEDELEMENTFRAMEWORK_API FConditions operator&&(FColumnBase Lhs, FColumnBase Rhs);
+	TYPEDELEMENTFRAMEWORK_API FConditions operator&&(FColumnBase Lhs, const FConditions& Rhs);
+
+	TYPEDELEMENTFRAMEWORK_API FConditions operator||(const FConditions& Lhs, FColumnBase Rhs);
+	TYPEDELEMENTFRAMEWORK_API FConditions operator||(const FConditions& Lhs, const FConditions& Rhs);
+	TYPEDELEMENTFRAMEWORK_API FConditions operator||(FColumnBase Lhs, FColumnBase Rhs);
+	TYPEDELEMENTFRAMEWORK_API FConditions operator||(FColumnBase Lhs, const FConditions& Rhs);
+} // namespace Queries
+} // namespace UE::Editor::DataStorage
+
diff --git a/Engine/Source/Runtime/TypedElementFramework/Public/DataStorage/Queries/Description.h b/Engine/Source/Runtime/TypedElementFramework/Public/DataStorage/Queries/Description.h
new file mode 100644
--- /dev/null
+++ b/Engine/Source/Runtime/TypedElementFramework/Public/DataStorage/Queries/Description.h
@@ -0,0 +1,124 @@
+// Copyright Epic Games, Inc. All Rights Reserved.
+
+#pragma once
+
+#include "Containers/Array.h"
+#include "Containers/ContainerAllocationPolicies.h"
+#include "DataStorage/CommonTypes.h"
+#include "DataStorage/Handles.h"
+#include "DataStorage/Queries/Types.h"
+#include "DataStorage/Queries/Conditions.h"
+#include "Elements/Interfaces/TypedElementQueryStorageInterfaces.h"
+#include "Elements/Framework/TypedElementMetaData.h"
+#include "Misc/Optional.h"
+#include "Templates/Function.h"
+#include "UObject/NameTypes.h"
+#include "UObject/WeakObjectPtr.h"
+#include "UObject/WeakObjectPtrTemplates.h"
+
+class UScriptStruct;
+
+namespace UE::Editor::DataStorage
+{
+	struct FQueryDescription;
+	struct IQueryContext;
+	
+	using QueryCallback = TFunction<void(const FQueryDescription&, IQueryContext&)>;
+	using QueryCallbackRef = TFunctionRef<void(const FQueryDescription&, IQueryContext&)>;
+	using DirectQueryCallback = TFunction<void(const FQueryDescription&, IDirectQueryContext&)>;
+	using DirectQueryCallbackRef = TFunctionRef<void(const FQueryDescription&, IDirectQueryContext&)>;
+	
+	struct FQueryDescription final
+	{
+		static constexpr int32 NumInlineSelections = 8;
+		static constexpr int32 NumInlineConditions = 8;
+		static constexpr int32 NumInlineDependencies = 2;
+		static constexpr int32 NumInlineGroups = 2;
+
+		enum class EActionType : uint8
+		{
+			None,	//< Do nothing.
+			Select,	//< Selects a set of columns for further processing.
+			Count,	//< Counts the number of entries that match the filter condition.
+
+			Max //< Value indicating the maximum value in this enum. Not to be used as an enum value.
+		};
+
+		enum class EOperatorType : uint16
+		{
+			SimpleAll,			//< Unary: Type
+			SimpleAny,			//< Unary: Type
+			SimpleNone,			//< Unary: Type
+			SimpleOptional,		//< Unary: Type
+		
+			Max //< Value indicating the maximum value in this enum. Not to be used as an enum value.
+		};
+
+		union FOperator
+		{
+			TWeakObjectPtr<const UScriptStruct> Type;
+		};
+
+		struct FValueTagData
+		{
+			// The Tag maps to a Mass ConstSharedFragment object
+			FValueTag Tag;
+
+			// The MatchValue specifies the value that the fragment must have to be matched
+			// If MatchValue is NAME_None, then TEDS will match all values
+			FName MatchValue;
+		};
+
+		struct FCallbackData
+		{
+			TArray<FName, TInlineAllocator<NumInlineGroups>> BeforeGroups;
+			TArray<FName, TInlineAllocator<NumInlineGroups>> AfterGroups;
+			QueryCallback Function;
+			FName Name;
+			FName Group;
+			/** If a name is set, it indicates the query callback will not be run unless the ActivationCount is greater than zero. */
+			FName ActivationName;
+			const UScriptStruct* MonitoredType{ nullptr };
+			EQueryCallbackType Type{ EQueryCallbackType::None };
+			EQueryTickPhase Phase{ EQueryTickPhase::FrameEnd };
+			/**
+			 * The number of remaining iterations for a activatable query callback. If this is higher than 0, the query callback will be 
+			 * called. If ActivationName is set, this value will be decremented by one at the end of the update cycle.
+			 */
+			uint8 ActivationCount = 255;
+			EExecutionMode ExecutionMode = EExecutionMode::Default;
+		};
+		FCallbackData Callback;
+
+		// The list of arrays below are required to remain in the same order as they're added as the function binding expects certain entries
+		// to be in a specific location.
+
+		TArray<TWeakObjectPtr<const UScriptStruct>, TInlineAllocator<NumInlineSelections>> SelectionTypes;
+		TArray<EQueryAccessType, TInlineAllocator<NumInlineSelections>> SelectionAccessTypes;
+		TArray<FColumnMetaData, TInlineAllocator<NumInlineSelections>> SelectionMetaData;
+
+		TArray<EOperatorType, TInlineAllocator<NumInlineConditions>> ConditionTypes;
+		TArray<FOperator, TInlineAllocator<NumInlineConditions>> ConditionOperators;
+
+		TArray<FDynamicColumnDescription, TInlineAllocator<NumInlineSelections>> DynamicSelectionTypes;
+		TArray<EQueryAccessType, TInlineAllocator<NumInlineSelections>> DynamicSelectionAccessTypes;
+		TArray<FColumnMetaData::EFlags, TInlineAllocator<NumInlineSelections>> DynamicSelectionMetaData;
+		
+		TArray<EOperatorType, TInlineAllocator<NumInlineConditions>> DynamicConditionOperations;
+		TArray<FDynamicColumnDescription, TInlineAllocator<NumInlineConditions>> DynamicConditionDescriptions;
+
+		TArray<FValueTagData> ValueTags;
+
+		TOptional<UE::Editor::DataStorage::Queries::FConditions> Conditions;
+
+		TArray<TWeakObjectPtr<const UClass>, TInlineAllocator<NumInlineDependencies>> DependencyTypes;
+		TArray<EQueryDependencyFlags, TInlineAllocator<NumInlineDependencies>> DependencyFlags;
+		/** Cached instances of the dependencies. This will always match the count of the other Dependency*Types, but may contain null pointers. */
+		TArray<TWeakObjectPtr<UObject>, TInlineAllocator<NumInlineDependencies>> CachedDependencies;
+		TArray<QueryHandle> Subqueries;
+		FMetaData MetaData;
+
+		EActionType Action;
+		bool bShouldBatchModifications = false;
+	};
+} // namespace UE::Editor::DataStorage
diff --git a/Engine/Source/Runtime/TypedElementFramework/Public/DataStorage/Queries/Types.h b/Engine/Source/Runtime/TypedElementFramework/Public/DataStorage/Queries/Types.h
new file mode 100644
--- /dev/null
+++ b/Engine/Source/Runtime/TypedElementFramework/Public/DataStorage/Queries/Types.h
@@ -0,0 +1,134 @@
+// Copyright Epic Games, Inc. All Rights Reserved.
+
+#pragma once
+
+#include "HAL/Platform.h"
+#include "Misc/EnumClassFlags.h"
+
+namespace UE::Editor::DataStorage
+{
+	// UE_DEPRECATED(5.6, "Use FMapKey(View) with the new mapping functions.") // Disabled as it's used by the deprecated functions.
+	using IndexHash = uint64;
+
+	enum class EQueryTickPhase : uint8
+	{
+		PrePhysics, //< Queries are executed before physics simulation starts.
+		DuringPhysics, //< Queries that can be run in parallel with physics simulation work.
+		PostPhysics, //< Queries that need rigid body and cloth simulation to be completed before being executed.
+		FrameEnd, //< Catchall for queries demoted to the last possible moment.
+
+		Max //< Value indicating the maximum value in this enum. Not to be used as an enum value.
+	};
+
+	enum class EQueryTickGroups : uint8
+	{
+		/** The standard group to run work in. */
+		Default,
+
+		/** General group that prepares for updating. */
+		PreUpdate,
+
+		/** General group that's used for updating data. */
+		Update,
+
+		/** General group that completes for updating. */
+		PostUpdate,
+
+		/**
+		 * The group for queries that need to sync data from external sources such as subsystems or the world into
+		 * the Data Storage. These typically run early in a phase.
+		 */
+		SyncExternalToDataStorage,
+		/**
+		 * The group for queries that need to sync data from the Data Storage to external sources such as subsystems
+		 * or the world into. These typically run late in a phase.
+		 */
+		SyncDataStorageToExternal,
+		/**
+		 * Queries grouped under this name will sync data to/from widgets.
+		 */
+		SyncWidgets,
+
+		Max //< Value indicating the maximum value in this enum. Not to be used as an enum value.
+	};
+
+	enum class EQueryCallbackType : uint8
+	{
+		/** No callback provided. */
+		None,
+		/** The query will be run every tick if at least one row matches. */
+		Processor,
+		/** The query will be run when a row is added that matches the query. The first recorded column will be actively monitored for changes. */
+		ObserveAdd,
+		/** The query will be run when a row is removed that matches the query. The first recorded column will be actively monitored for changes. */
+		ObserveRemove,
+		/**
+		 * At the start of the assigned phase this query will run if there are any matches. These queries will have any deferred operations such as
+		 * adding/removing rows/columns executed before the phase starts. This introduces sync points that hinder performance and are therefore
+		 * recommended only for queries that save on work later in the phase such as repeated checks for validity.
+		 */
+		PhasePreparation,
+		/**
+		 * At the end of the assigned phase this query will run if there are any matches. These queries will have any deferred operations such as
+		 * adding/removing rows/columns executed before the phase ends. This introduces sync points that hinder performance and are therefore
+		 * recommended only cases where delaying deferred operations is not possible e.g. when tables are known to be referenced outside the update
+		 * cycle.
+		 */
+		PhaseFinalization,
+
+		Max //< Value indicating the maximum value in this enum. Not to be used as an enum value.
+	};
+
+	enum class EExecutionMode : uint8
+	{
+		Default, //< Use the default for the query callback type or the group default.
+		GameThread, //< The query callback will be executed sequentially on the game thread.
+		Threaded, //< The entire query callback is called on a separate thread, while respecting dependencies.
+		ThreadedChunks //< The query callback is called on a separate thread for each chunk. This can not be set if there are sub-queries.
+	};
+
+	enum class EQueryAccessType : uint8
+	{
+		/** Provide read-only access. */
+		ReadOnly,
+		/** Provide read-only access, but do not include when finding matching tables. */
+		OptionalReadOnly,
+		/** Provide read and write access. */
+		ReadWrite
+	};
+
+	enum class EQueryDependencyFlags : uint8
+	{
+		None = 0,
+		/** If set the dependency is accessed as read-only. If not set the dependency requires Read/Write access. */
+		ReadOnly = 1 << 0,
+		/** If set the dependency can only be used from the game thread, otherwise it can be accessed from any thread. */
+		GameThreadBound = 1 << 1,
+		/** If set the dependency will be re-fetched every iteration, otherwise only if not fetched before. */
+		AlwaysRefresh = 1 << 2
+	};
+	ENUM_CLASS_FLAGS(EQueryDependencyFlags);
+
+	struct FQueryResult final
+	{
+		enum class ECompletion
+		{
+			/** Query could be fully executed. */
+			Fully,
+			/** Only portions of the query were executed. This is caused by a problem that was encountered partway through processing. */
+			Partially,
+			/**
+			 * The back-end doesn't support the particular query. This may be a limitation in how/where the query is run or because
+			 * the query contains actions and/or operations that are not supported.
+			 */
+			Unsupported,
+			/** The provided query is no longer available. */
+			Unavailable,
+			/** One or more dependencies declared on the query could not be retrieved. */
+			MissingDependency
+		};
+
+		uint32 Count{ 0 }; /** The number of rows were processed. */
+		ECompletion Completed{ ECompletion::Unavailable };
+	};
+} // namespace UE::Editor::DataStorage
diff --git a/Engine/Source/Runtime/TypedElementFramework/Public/Elements/Common/EditorDataStorageFeatures.h b/Engine/Source/Runtime/TypedElementFramework/Public/Elements/Common/EditorDataStorageFeatures.h
--- a/Engine/Source/Runtime/TypedElementFramework/Public/Elements/Common/EditorDataStorageFeatures.h
+++ b/Engine/Source/Runtime/TypedElementFramework/Public/Elements/Common/EditorDataStorageFeatures.h
@@ -2,37 +2,8 @@
 
 #pragma once
 
-#include "Delegates/Delegate.h"
-#include "Features/IModularFeatures.h"
-#include "UObject/NameTypes.h"
-#include "GenericPlatform/GenericPlatform.h"
+#include "Misc/CoreMiscDefines.h"
 
-namespace UE::Editor::DataStorage
-{
-	TYPEDELEMENTFRAMEWORK_API const extern FName StorageFeatureName;
-	TYPEDELEMENTFRAMEWORK_API const extern FName CompatibilityFeatureName;
-	TYPEDELEMENTFRAMEWORK_API const extern FName UiFeatureName;
+UE_DEPRECATED_HEADER(5.7, "Use DataStorage/Features.h instead of Elements/Common/EditorDataStorageFeatures.h")
 
-	TYPEDELEMENTFRAMEWORK_API FSimpleMulticastDelegate& OnEditorDataStorageFeaturesEnabled();
-	TYPEDELEMENTFRAMEWORK_API bool AreEditorDataStorageFeaturesEnabled();
-
-	template<typename T>
-	const T* GetDataStorageFeature(const FName InName)
-	{
-		if (IModularFeatures::Get().IsModularFeatureAvailable(InName))
-		{
-			return &IModularFeatures::Get().GetModularFeature<T>(InName);
-		}
-		return nullptr;
-	}
-
-	template<typename T>
-	T* GetMutableDataStorageFeature(const FName InName)
-	{
-		if (IModularFeatures::Get().IsModularFeatureAvailable(InName))
-		{
-			return &IModularFeatures::Get().GetModularFeature<T>(InName);
-		}
-		return nullptr;
-	}
-} // namespace UE::Editor::DataStorage
+#include "DataStorage/Features.h"
diff --git a/Engine/Source/Runtime/TypedElementFramework/Public/Elements/Common/TypedElementCommonTypes.h b/Engine/Source/Runtime/TypedElementFramework/Public/Elements/Common/TypedElementCommonTypes.h
--- a/Engine/Source/Runtime/TypedElementFramework/Public/Elements/Common/TypedElementCommonTypes.h
+++ b/Engine/Source/Runtime/TypedElementFramework/Public/Elements/Common/TypedElementCommonTypes.h
@@ -2,105 +2,8 @@
 
 #pragma once
 
-#include "Elements/Common/TypedElementHandles.h"
+#include "Misc/CoreMiscDefines.h"
 
-#include "TypedElementCommonTypes.generated.h"
+UE_DEPRECATED_HEADER(5.7, "Use DataStorage/CommonTypes.h instead of Elements/Common/TypedElementCommonTypes.h")
 
-/**
- * Base for the data structures for a column.
- */
-USTRUCT()
-struct FEditorDataStorageColumn
-{
-	GENERATED_BODY()
-};
-
-/**
- * Base for the data structures that act as tags to rows. Tags should not have any data.
- */
-USTRUCT()
-struct FEditorDataStorageTag
-{
-	GENERATED_BODY()
-};
-
-namespace UE
-{
-	// Work around missing header/implementations on some platforms
-	namespace detail
-	{
-		template<typename T, typename U>
-		concept SameHelper = std::is_same_v<T, U>;
-	}
-	template<typename T, typename U>
-	concept same_as = detail::SameHelper<T, U> && detail::SameHelper<U, T>;
-
-	template<typename From, typename To>
-	concept convertible_to = std::is_convertible_v<From, To> && requires { static_cast<To>(std::declval<From>()); };
-
-	template<typename Derived, typename Base>
-	concept derived_from = std::is_base_of_v<Base, Derived> && std::is_convertible_v<const volatile Derived*, const volatile Base*>;
-
-	namespace Editor::DataStorage
-	{
-		using FColumn = FEditorDataStorageColumn;
-		using FTag = FEditorDataStorageTag;
-
-		/**
-		 * Defines a dynamic type for a value tag
-		 * Example:
-		 *   FValueTag ColorTagType(TEXT("Color"));
-		 *   FValueTag DirectionTagType(TEXT("Direction"));
-		 * A value tag can take on different values for each type.  This is set up when a tag is added to a row.
-		 */
-		class FValueTag
-		{
-		public:
-			TYPEDELEMENTFRAMEWORK_API explicit FValueTag(const FName& InName);
-			
-			TYPEDELEMENTFRAMEWORK_API const FName& GetName() const;
-			bool operator==(const FValueTag& Other) const = default;
-		private:
-			TYPEDELEMENTFRAMEWORK_API friend uint32 GetTypeHash(const FValueTag& InName);
-			FName Name;
-		};
-
-		template<typename T>
-		concept TValueTagType = std::is_same_v<T, FValueTag>;
-
-		struct FDynamicColumnDescription
-		{
-			const UScriptStruct* TemplateType;
-			FName Identifier;
-
-			TYPEDELEMENTFRAMEWORK_API friend uint32 GetTypeHash(const FDynamicColumnDescription& Descriptor);
-			bool operator==(const FDynamicColumnDescription&) const = default;
-		};
-		// Standard callbacks.
-
-		using RowCreationCallbackRef = TFunctionRef<void(RowHandle Row)>;
-		using ColumnCreationCallbackRef = TFunctionRef<void(void* Column, const UScriptStruct& ColumnType)>;
-		using ColumnListCallbackRef = TFunctionRef<void(const UScriptStruct& ColumnType)>;
-		using ColumnListWithDataCallbackRef = TFunctionRef<void(void* Column, const UScriptStruct& ColumnType)>;
-		using ColumnCopyOrMoveCallback = void (*)(const UScriptStruct& ColumnType, void* Destination, void* Source);
-
-		template<typename T>
-		concept THasDynamicColumnTemplateSpecifier = std::is_empty_v<typename T::EditorDataStorage_DynamicColumnTemplate>;
-		
-		template<typename T>
-		concept TDynamicColumnTemplate = (UE::derived_from<T, FColumn> || UE::derived_from<T, FTag>) && THasDynamicColumnTemplateSpecifier<T>;
-		
-		// Template concepts to enforce type correctness.
-		template<typename T>
-		concept TDataColumnType = UE::derived_from<T, FColumn> && !THasDynamicColumnTemplateSpecifier<T>;
-
-		template<typename T>
-		concept TTagColumnType = UE::derived_from<T, FTag> && !THasDynamicColumnTemplateSpecifier<T>;
-		
-		template<typename T>
-		concept TColumnType = TDataColumnType<T> || TTagColumnType<T>;
-
-		template<typename T>
-		concept TEnumType = std::is_enum_v<T>;
-	} // namespace Editor::DataStorage
-} // namespace UE
+#include "DataStorage/CommonTypes.h"
diff --git a/Engine/Source/Runtime/TypedElementFramework/Public/Elements/Common/TypedElementHandles.h b/Engine/Source/Runtime/TypedElementFramework/Public/Elements/Common/TypedElementHandles.h
--- a/Engine/Source/Runtime/TypedElementFramework/Public/Elements/Common/TypedElementHandles.h
+++ b/Engine/Source/Runtime/TypedElementFramework/Public/Elements/Common/TypedElementHandles.h
@@ -2,51 +2,8 @@
 
 #pragma once
 
-#include "Math/NumericLimits.h"
-#include "Templates/TypeHash.h"
+#include "Misc/CoreMiscDefines.h"
 
-#include "TypedElementHandles.generated.h"
+UE_DEPRECATED_HEADER(5.7, "Use DataStorage/Handles.h instead of Elements/Common/TypedElementHandles.h")
 
-namespace UE::Editor::DataStorage
-{
-	using TableHandle = uint64;
-	static constexpr TableHandle InvalidTableHandle = TNumericLimits<TableHandle>::Max();
-
-	using RowHandle = uint64;
-	static constexpr RowHandle InvalidRowHandle = 0;
-
-	using QueryHandle = uint64;
-	static constexpr QueryHandle InvalidQueryHandle = TNumericLimits<QueryHandle>::Max();
-} // namespace UE::Editor::DataStorage
-
-/*
- * FTedsRowHandle is a strongly typed wrapper around UE::Editor::DataStorage::RowHandle and should only be used in cases where you need the extra info.
- * E.g for reflection/UHT or for template specializing something that needs to know the semantics of the row handle.
- * For all other cases, you should use the regular typedef UE::Editor::DataStorage::RowHandle
- */
-USTRUCT()
-struct FTedsRowHandle
-{
-	GENERATED_BODY()
-
-	UPROPERTY()
-	uint64 RowHandle = UE::Editor::DataStorage::InvalidRowHandle;
-
-	operator UE::Editor::DataStorage::RowHandle () const
-	{
-		return RowHandle;
-	}
-
-	FTedsRowHandle& operator=(UE::Editor::DataStorage::RowHandle InRowHandle)
-	{
-		RowHandle = InRowHandle;
-		return *this;
-	}
-		
-	friend uint32 GetTypeHash(const FTedsRowHandle& Key)
-	{
-		return GetTypeHash(Key.RowHandle);
-	}
-};
-
-static_assert(sizeof(FTedsRowHandle::RowHandle) == sizeof(UE::Editor::DataStorage::RowHandle), "RowHandle and RowHandle wrapper sizes should match");
\ No newline at end of file
+#include "DataStorage/Handles.h"
diff --git a/Engine/Source/Runtime/TypedElementFramework/Public/Elements/Common/TypedElementMapKey.h b/Engine/Source/Runtime/TypedElementFramework/Public/Elements/Common/TypedElementMapKey.h
--- a/Engine/Source/Runtime/TypedElementFramework/Public/Elements/Common/TypedElementMapKey.h
+++ b/Engine/Source/Runtime/TypedElementFramework/Public/Elements/Common/TypedElementMapKey.h
@@ -2,96 +2,8 @@
 
 #pragma once
 
-#include "Containers/StringFwd.h"
-#include "Containers/UnrealString.h"
-#include "Misc/TVariant.h"
-#include "Templates/UniquePtr.h"
-#include "UObject/NameTypes.h"
-#include "UObject/SoftObjectPath.h"
+#include "Misc/CoreMiscDefines.h"
 
-class UObject;
+UE_DEPRECATED_HEADER(5.7, "Use DataStorage/MapKey.h instead of Elements/Common/TypedElementMapKey.h")
 
-namespace UE::Editor::DataStorage
-{
-	class FMapKeyView;
-
-	class FMapKey
-	{
-		friend class FMapKeyView;
-		friend struct FKeyComparer;
-		friend struct FKeyCopy;
-		friend struct FKeyMove;
-
-	public:
-		FMapKey() = default;
-		
-		TYPEDELEMENTFRAMEWORK_API FMapKey(FMapKey&& Rhs);
-		TYPEDELEMENTFRAMEWORK_API FMapKey& operator=(FMapKey&& Rhs);
-		
-		TYPEDELEMENTFRAMEWORK_API FMapKey(const FMapKey& Key);
-		TYPEDELEMENTFRAMEWORK_API explicit FMapKey(const void* Key);
-		TYPEDELEMENTFRAMEWORK_API explicit FMapKey(const UObject* Key);
-		TYPEDELEMENTFRAMEWORK_API explicit FMapKey(int64 Key);
-		TYPEDELEMENTFRAMEWORK_API explicit FMapKey(uint64 Key);
-		TYPEDELEMENTFRAMEWORK_API explicit FMapKey(float Key);
-		TYPEDELEMENTFRAMEWORK_API explicit FMapKey(double Key);
-		TYPEDELEMENTFRAMEWORK_API explicit FMapKey(FString Key);
-		TYPEDELEMENTFRAMEWORK_API explicit FMapKey(FName Key);
-		TYPEDELEMENTFRAMEWORK_API explicit FMapKey(FSoftObjectPath Key);
-
-		TYPEDELEMENTFRAMEWORK_API FMapKey& operator=(const FMapKey& Rhs);
-
-		TYPEDELEMENTFRAMEWORK_API bool operator==(const FMapKey& Rhs) const;
-		TYPEDELEMENTFRAMEWORK_API bool operator==(const FMapKeyView& Rhs) const;
-		TYPEDELEMENTFRAMEWORK_API bool operator!=(const FMapKey& Rhs) const;
-		TYPEDELEMENTFRAMEWORK_API bool operator!=(const FMapKeyView& Rhs) const;
-
-		TYPEDELEMENTFRAMEWORK_API uint64 CalculateHash() const;
-		TYPEDELEMENTFRAMEWORK_API bool IsSet() const;
-		TYPEDELEMENTFRAMEWORK_API void Clear();
-		TYPEDELEMENTFRAMEWORK_API FString ToString() const;
-
-	private:
-		using KeyType = TVariant<FEmptyVariantState, const void*, const UObject*, int64, uint64, float, double, FString, FName, TUniquePtr<FSoftObjectPath>>;
-
-		KeyType Key;
-	};
-
-	class FMapKeyView
-	{
-		friend struct FKeyToKeyViewConverter;
-		friend struct FKeyViewComparer;
-
-	public:
-		FMapKeyView() = default;
-		TYPEDELEMENTFRAMEWORK_API FMapKeyView(const FMapKey& InKey);
-		TYPEDELEMENTFRAMEWORK_API explicit FMapKeyView(const void* Key);
-		TYPEDELEMENTFRAMEWORK_API explicit FMapKeyView(const UObject* Key);
-		TYPEDELEMENTFRAMEWORK_API explicit FMapKeyView(int64 Key);
-		TYPEDELEMENTFRAMEWORK_API explicit FMapKeyView(uint64 Key);
-		TYPEDELEMENTFRAMEWORK_API explicit FMapKeyView(float Key);
-		TYPEDELEMENTFRAMEWORK_API explicit FMapKeyView(double Key);
-		TYPEDELEMENTFRAMEWORK_API explicit FMapKeyView(const FString& Key);
-		TYPEDELEMENTFRAMEWORK_API explicit FMapKeyView(FStringView Key);
-		TYPEDELEMENTFRAMEWORK_API explicit FMapKeyView(const FName& Key);
-		TYPEDELEMENTFRAMEWORK_API explicit FMapKeyView(const FSoftObjectPath& Key);
-
-		TYPEDELEMENTFRAMEWORK_API FMapKeyView& operator=(const FMapKey& InKey);
-
-		TYPEDELEMENTFRAMEWORK_API FMapKey CreateKey() const;
-
-		TYPEDELEMENTFRAMEWORK_API bool operator==(const FMapKey& Rhs) const;
-		TYPEDELEMENTFRAMEWORK_API bool operator==(const FMapKeyView& Rhs) const;
-		TYPEDELEMENTFRAMEWORK_API bool operator!=(const FMapKey& Rhs) const;
-		TYPEDELEMENTFRAMEWORK_API bool operator!=(const FMapKeyView& Rhs) const;
-
-		TYPEDELEMENTFRAMEWORK_API uint64 CalculateHash() const;
-		TYPEDELEMENTFRAMEWORK_API FString ToString() const;
-
-	private:
-		using KeyViewType = TVariant<FEmptyVariantState, const void*, const UObject*, int64, uint64, float, double, FStringView, const FName*, const FSoftObjectPath*>;
-
-		KeyViewType Key;
-	};
-}
-// namespace UE::Editor::DataStorage
+#include "DataStorage/MapKey.h"
