[Teds] Rename Elements/Common/TypedElement...
to DataStorage/... aligning to namespace and class names

diff --git a/Engine/Plugins/ChaosVD/Source/ChaosVD/Private/ChaosVDScene.cpp b/Engine/Plugins/ChaosVD/Source/ChaosVD/Private/ChaosVDScene.cpp
--- a/Engine/Plugins/ChaosVD/Source/ChaosVD/Private/ChaosVDScene.cpp
+++ b/Engine/Plugins/ChaosVD/Source/ChaosVD/Private/ChaosVDScene.cpp
@@ -28,8 +28,8 @@
 #include "Settings/ChaosVDCoreSettings.h"
 #include "UObject/Package.h"
 #include "Actors/ChaosVDGeometryContainer.h"
+#include "DataStorage/Features.h"
 #include "Elements/Columns/TypedElementMiscColumns.h"
-#include "Elements/Common/EditorDataStorageFeatures.h"
 #include "Elements/Framework/TypedElementRegistry.h"
 #include "Engine/PostProcessVolume.h"
 #include "Engine/TextureCube.h"
diff --git a/Engine/Plugins/ChaosVD/Source/ChaosVD/Private/ChaosVDSceneParticle.cpp b/Engine/Plugins/ChaosVD/Source/ChaosVD/Private/ChaosVDSceneParticle.cpp
--- a/Engine/Plugins/ChaosVD/Source/ChaosVD/Private/ChaosVDSceneParticle.cpp
+++ b/Engine/Plugins/ChaosVD/Source/ChaosVD/Private/ChaosVDSceneParticle.cpp
@@ -14,10 +14,10 @@
 #include "Components/MeshComponent.h"
 #include "Components/SceneComponent.h"
 #include "Components/StaticMeshComponent.h"
+#include "DataStorage/Features.h"
 #include "DataWrappers/ChaosVDParticleDataWrapper.h"
 #include "Elements/Columns/TypedElementMiscColumns.h"
 #include "Elements/Columns/TypedElementVisibilityColumns.h"
-#include "Elements/Common/EditorDataStorageFeatures.h"
 #include "Elements/Interfaces/TypedElementDataStorageCompatibilityInterface.h"
 #include "Engine/StaticMesh.h"
 #include "TEDS/ChaosVDParticleEditorDataFactory.h"
diff --git a/Engine/Plugins/ChaosVD/Source/ChaosVD/Private/TEDS/ChaosVDDataStorageVisibilityQueries.cpp b/Engine/Plugins/ChaosVD/Source/ChaosVD/Private/TEDS/ChaosVDDataStorageVisibilityQueries.cpp
--- a/Engine/Plugins/ChaosVD/Source/ChaosVD/Private/TEDS/ChaosVDDataStorageVisibilityQueries.cpp
+++ b/Engine/Plugins/ChaosVD/Source/ChaosVD/Private/TEDS/ChaosVDDataStorageVisibilityQueries.cpp
@@ -4,11 +4,11 @@
 
 #include "ChaosVDParticleEditorDataFactory.h"
 #include "ChaosVDSceneParticle.h"
+#include "DataStorage/Queries/Types.h"
 #include "Elements/Columns/TypedElementCompatibilityColumns.h"
 #include "Elements/Columns/TypedElementMiscColumns.h"
 #include "Elements/Columns/TypedElementTypeInfoColumns.h"
 #include "Elements/Columns/TypedElementVisibilityColumns.h"
-#include "Elements/Common/TypedElementQueryTypes.h"
 #include "Elements/Framework/TypedElementQueryBuilder.h"
 
 
diff --git a/Engine/Plugins/ChaosVD/Source/ChaosVD/Private/TEDS/ChaosVDParentDataStorageFactory.h b/Engine/Plugins/ChaosVD/Source/ChaosVD/Private/TEDS/ChaosVDParentDataStorageFactory.h
--- a/Engine/Plugins/ChaosVD/Source/ChaosVD/Private/TEDS/ChaosVDParentDataStorageFactory.h
+++ b/Engine/Plugins/ChaosVD/Source/ChaosVD/Private/TEDS/ChaosVDParentDataStorageFactory.h
@@ -2,7 +2,7 @@
 
 #pragma once
 
-#include "Elements/Common/TypedElementCommonTypes.h"
+#include "DataStorage/CommonTypes.h"
 #include "Elements/Interfaces/TypedElementDataStorageFactory.h"
 #include "ChaosVDParentDataStorageFactory.generated.h"
 
diff --git a/Engine/Plugins/ChaosVD/Source/ChaosVD/Private/TEDS/ChaosVDParticleEditorDataFactory.cpp b/Engine/Plugins/ChaosVD/Source/ChaosVD/Private/TEDS/ChaosVDParticleEditorDataFactory.cpp
--- a/Engine/Plugins/ChaosVD/Source/ChaosVD/Private/TEDS/ChaosVDParticleEditorDataFactory.cpp
+++ b/Engine/Plugins/ChaosVD/Source/ChaosVD/Private/TEDS/ChaosVDParticleEditorDataFactory.cpp
@@ -7,12 +7,12 @@
 #include "TEDS/ChaosVDSelectionInterface.h"
 #include "TEDS/ChaosVDStructTypedElementData.h"
 #include "Actors/ChaosVDSolverInfoActor.h"
+#include "DataStorage/Handles.h"
 #include "Elements/Columns/TypedElementCompatibilityColumns.h"
 #include "Elements/Columns/TypedElementLabelColumns.h"
 #include "Elements/Columns/TypedElementMiscColumns.h"
 #include "Elements/Columns/TypedElementTypeInfoColumns.h"
 #include "Elements/Columns/TypedElementVisibilityColumns.h"
-#include "Elements/Common/TypedElementHandles.h"
 #include "Elements/Framework/TypedElementQueryBuilder.h"
 #include "Elements/Framework/TypedElementRegistry.h"
 #include "Elements/Interfaces/TypedElementDataStorageCompatibilityInterface.h"
diff --git a/Engine/Plugins/ChaosVD/Source/ChaosVD/Private/TEDS/ChaosVDParticleEditorDataFactory.h b/Engine/Plugins/ChaosVD/Source/ChaosVD/Private/TEDS/ChaosVDParticleEditorDataFactory.h
--- a/Engine/Plugins/ChaosVD/Source/ChaosVD/Private/TEDS/ChaosVDParticleEditorDataFactory.h
+++ b/Engine/Plugins/ChaosVD/Source/ChaosVD/Private/TEDS/ChaosVDParticleEditorDataFactory.h
@@ -4,7 +4,7 @@
 
 #pragma once
 
-#include "Elements/Common/TypedElementCommonTypes.h"
+#include "DataStorage/CommonTypes.h"
 #include "Elements/Interfaces/TypedElementDataStorageFactory.h"
 #include "UObject/ObjectMacros.h"
 
diff --git a/Engine/Plugins/ChaosVD/Source/ChaosVD/Private/TEDS/ChaosVDTedsUtils.h b/Engine/Plugins/ChaosVD/Source/ChaosVD/Private/TEDS/ChaosVDTedsUtils.h
--- a/Engine/Plugins/ChaosVD/Source/ChaosVD/Private/TEDS/ChaosVDTedsUtils.h
+++ b/Engine/Plugins/ChaosVD/Source/ChaosVD/Private/TEDS/ChaosVDTedsUtils.h
@@ -2,7 +2,7 @@
 
 #pragma once
 
-#include "Elements/Common/EditorDataStorageFeatures.h"
+#include "DataStorage/Features.h"
 #include "Elements/Interfaces/TypedElementDataStorageCompatibilityInterface.h"
 
 /** Set of utility methods to interact with TEDS using CVD types */
diff --git a/Engine/Plugins/ChaosVD/Source/ChaosVD/Public/ChaosVDBaseSceneObject.h b/Engine/Plugins/ChaosVD/Source/ChaosVD/Public/ChaosVDBaseSceneObject.h
--- a/Engine/Plugins/ChaosVD/Source/ChaosVD/Public/ChaosVDBaseSceneObject.h
+++ b/Engine/Plugins/ChaosVD/Source/ChaosVD/Public/ChaosVDBaseSceneObject.h
@@ -1,7 +1,7 @@
 // Copyright Epic Games, Inc. All Rights Reserved.
 
 #pragma once
-#include "Elements/Common/TypedElementHandles.h"
+#include "DataStorage/Handles.h"
 #include "UObject/NameTypes.h"
 #include "UObject/WeakObjectPtr.h"
 #include "Templates/SharedPointer.h"
diff --git a/Engine/Plugins/Experimental/Dataflow/Source/DataflowEditor/Private/Dataflow/DataflowOutlinerView.cpp b/Engine/Plugins/Experimental/Dataflow/Source/DataflowEditor/Private/Dataflow/DataflowOutlinerView.cpp
--- a/Engine/Plugins/Experimental/Dataflow/Source/DataflowEditor/Private/Dataflow/DataflowOutlinerView.cpp
+++ b/Engine/Plugins/Experimental/Dataflow/Source/DataflowEditor/Private/Dataflow/DataflowOutlinerView.cpp
@@ -6,10 +6,10 @@
 #include "TedsOutlinerImpl.h"
 #include "Compatibility/SceneOutlinerTedsBridge.h"
 #include "Dataflow/DataflowContent.h"
+#include "DataStorage/Features.h"
+#include "DataStorage/Queries/Description.h"
 #include "Elements/Columns/TypedElementTypeInfoColumns.h"
 #include "Elements/Columns/TypedElementVisibilityColumns.h"
-#include "Elements/Common/EditorDataStorageFeatures.h"
-#include "Elements/Common/TypedElementQueryDescription.h"
 #include "Elements/Framework/TypedElementQueryBuilder.h"
 
 FDataflowOutlinerView::FDataflowOutlinerView(FDataflowPreviewSceneBase* InConstructionScene, FDataflowPreviewSceneBase* InSimulationScene, TObjectPtr<UDataflowBaseContent> InContent)
diff --git a/Engine/Plugins/Experimental/Dataflow/Source/DataflowEditor/Public/Dataflow/DataflowEditorPreviewSceneBase.h b/Engine/Plugins/Experimental/Dataflow/Source/DataflowEditor/Public/Dataflow/DataflowEditorPreviewSceneBase.h
--- a/Engine/Plugins/Experimental/Dataflow/Source/DataflowEditor/Public/Dataflow/DataflowEditorPreviewSceneBase.h
+++ b/Engine/Plugins/Experimental/Dataflow/Source/DataflowEditor/Public/Dataflow/DataflowEditorPreviewSceneBase.h
@@ -3,9 +3,9 @@
 #pragma once
 
 #include "AdvancedPreviewScene.h"
+#include "DataStorage/CommonTypes.h"
 #include "Dataflow/DataflowDebugDraw.h"
 #include "Dataflow/DataflowDebugDrawObject.h"
-#include "Elements/Common/TypedElementCommonTypes.h"
 
 #include "DataflowEditorPreviewSceneBase.generated.h"
 
diff --git a/Engine/Plugins/Experimental/EditorDataStorage/Source/TedsCore/Private/Commands/EditorDataStorageCompatibilityCommands.h b/Engine/Plugins/Experimental/EditorDataStorage/Source/TedsCore/Private/Commands/EditorDataStorageCompatibilityCommands.h
--- a/Engine/Plugins/Experimental/EditorDataStorage/Source/TedsCore/Private/Commands/EditorDataStorageCompatibilityCommands.h
+++ b/Engine/Plugins/Experimental/EditorDataStorage/Source/TedsCore/Private/Commands/EditorDataStorageCompatibilityCommands.h
@@ -3,8 +3,8 @@
 #pragma once
 
 #include "Commands/EditorDataStorageCommandBuffer.h"
+#include "DataStorage/Handles.h"
 #include "Delegates/IDelegateInstance.h"
-#include "Elements/Common/TypedElementHandles.h"
 #include "Elements/Interfaces/TypedElementDataStorageCompatibilityInterface.h"
 #include "UObject/WeakObjectPtr.h"
 #include "UObject/WeakObjectPtrTemplates.h"
diff --git a/Engine/Plugins/Experimental/EditorDataStorage/Source/TedsCore/Private/Compatibility/TypedElementUObjectNameQueries.cpp b/Engine/Plugins/Experimental/EditorDataStorage/Source/TedsCore/Private/Compatibility/TypedElementUObjectNameQueries.cpp
--- a/Engine/Plugins/Experimental/EditorDataStorage/Source/TedsCore/Private/Compatibility/TypedElementUObjectNameQueries.cpp
+++ b/Engine/Plugins/Experimental/EditorDataStorage/Source/TedsCore/Private/Compatibility/TypedElementUObjectNameQueries.cpp
@@ -2,9 +2,9 @@
 
 #include "TypedElementUObjectNameQueries.h"
 
+#include "DataStorage/Queries/Types.h"
 #include "Elements/Columns/TypedElementCompatibilityColumns.h"
 #include "Elements/Columns/TypedElementMiscColumns.h"
-#include "Elements/Common/TypedElementQueryTypes.h"
 #include "Elements/Framework/TypedElementQueryBuilder.h"
 
 void UObjectNameDataStorageFactory::RegisterQueries(UE::Editor::DataStorage::ICoreProvider& DataStorage)
diff --git a/Engine/Plugins/Experimental/EditorDataStorage/Source/TedsCore/Private/Debug/TypedElementDatabaseDebugCommands.cpp b/Engine/Plugins/Experimental/EditorDataStorage/Source/TedsCore/Private/Debug/TypedElementDatabaseDebugCommands.cpp
--- a/Engine/Plugins/Experimental/EditorDataStorage/Source/TedsCore/Private/Debug/TypedElementDatabaseDebugCommands.cpp
+++ b/Engine/Plugins/Experimental/EditorDataStorage/Source/TedsCore/Private/Debug/TypedElementDatabaseDebugCommands.cpp
@@ -1,11 +1,11 @@
 // Copyright Epic Games, Inc. All Rights Reserved.
 
 #include "TypedElementDatabaseDebugTypes.h"
+#include "DataStorage/Features.h"
+#include "DataStorage/Debug/Log.h"
 #include "Elements/Columns/TypedElementCompatibilityColumns.h"
 #include "Elements/Columns/TypedElementLabelColumns.h"
 #include "Elements/Columns/TypedElementTypeInfoColumns.h"
-#include "Elements/Common/EditorDataStorageFeatures.h"
-#include "Elements/Common/TypedElementDataStorageLog.h"
 #include "Elements/Framework/TypedElementQueryBuilder.h"
 #include "Elements/Framework/TypedElementTestColumns.h"
 #include "Elements/Interfaces/TypedElementDataStorageInterface.h"
diff --git a/Engine/Plugins/Experimental/EditorDataStorage/Source/TedsCore/Private/Memento/TypedElementMementoSystem.cpp b/Engine/Plugins/Experimental/EditorDataStorage/Source/TedsCore/Private/Memento/TypedElementMementoSystem.cpp
--- a/Engine/Plugins/Experimental/EditorDataStorage/Source/TedsCore/Private/Memento/TypedElementMementoSystem.cpp
+++ b/Engine/Plugins/Experimental/EditorDataStorage/Source/TedsCore/Private/Memento/TypedElementMementoSystem.cpp
@@ -2,7 +2,7 @@
 
 #include "TypedElementMementoSystem.h"
 
-#include "Elements/Common/TypedElementDataStorageLog.h"
+#include "DataStorage/Debug/Log.h"
 #include "Elements/Framework/TypedElementQueryBuilder.h"
 #include "GlobalLock.h"
 #include "Memento/TypedElementMementoTranslators.h"
diff --git a/Engine/Plugins/Experimental/EditorDataStorage/Source/TedsCore/Private/Processors/TypedElementProcessorAdaptors.cpp b/Engine/Plugins/Experimental/EditorDataStorage/Source/TedsCore/Private/Processors/TypedElementProcessorAdaptors.cpp
--- a/Engine/Plugins/Experimental/EditorDataStorage/Source/TedsCore/Private/Processors/TypedElementProcessorAdaptors.cpp
+++ b/Engine/Plugins/Experimental/EditorDataStorage/Source/TedsCore/Private/Processors/TypedElementProcessorAdaptors.cpp
@@ -4,7 +4,7 @@
 
 #include <utility>
 #include "Algo/Copy.h"
-#include "Elements/Common/TypedElementQueryTypes.h"
+#include "DataStorage/Queries/Types.h"
 #include "MassEntityTypes.h"
 #include "MassEntityView.h"
 #include "MassExecutionContext.h"
diff --git a/Engine/Plugins/Experimental/EditorDataStorage/Source/TedsCore/Private/Queries/TypedElementExtendedQueryStore.cpp b/Engine/Plugins/Experimental/EditorDataStorage/Source/TedsCore/Private/Queries/TypedElementExtendedQueryStore.cpp
--- a/Engine/Plugins/Experimental/EditorDataStorage/Source/TedsCore/Private/Queries/TypedElementExtendedQueryStore.cpp
+++ b/Engine/Plugins/Experimental/EditorDataStorage/Source/TedsCore/Private/Queries/TypedElementExtendedQueryStore.cpp
@@ -6,7 +6,7 @@
 #include "MassProcessor.h"
 #include "TypedElementDatabaseEnvironment.h"
 #include "Commands/EditorDataStorageCommandBuffer.h"
-#include "Elements/Common/TypedElementDataStorageLog.h"
+#include "DataStorage/Debug/Log.h"
 #include "Misc/OutputDevice.h"
 #include "Processors/TypedElementProcessorAdaptors.h"
 
diff --git a/Engine/Plugins/Experimental/EditorDataStorage/Source/TedsCore/Private/TypedElementDatabaseCommandBuffer.h b/Engine/Plugins/Experimental/EditorDataStorage/Source/TedsCore/Private/TypedElementDatabaseCommandBuffer.h
--- a/Engine/Plugins/Experimental/EditorDataStorage/Source/TedsCore/Private/TypedElementDatabaseCommandBuffer.h
+++ b/Engine/Plugins/Experimental/EditorDataStorage/Source/TedsCore/Private/TypedElementDatabaseCommandBuffer.h
@@ -3,8 +3,8 @@
 #pragma once
 
 #include <type_traits>
-#include "Elements/Common/TypedElementCommonTypes.h"
-#include "Elements/Common/TypedElementHandles.h"
+#include "DataStorage/Handles.h"
+#include "DataStorage/CommonTypes.h"
 #include "MassArchetypeTypes.h"
 #include "Misc/TVariant.h"
 #include "UObject/Class.h"
diff --git a/Engine/Plugins/Experimental/EditorDataStorage/Source/TedsCore/Private/TypedElementDatabaseIndexTable.h b/Engine/Plugins/Experimental/EditorDataStorage/Source/TedsCore/Private/TypedElementDatabaseIndexTable.h
--- a/Engine/Plugins/Experimental/EditorDataStorage/Source/TedsCore/Private/TypedElementDatabaseIndexTable.h
+++ b/Engine/Plugins/Experimental/EditorDataStorage/Source/TedsCore/Private/TypedElementDatabaseIndexTable.h
@@ -7,8 +7,8 @@
 #include "Containers/Map.h"
 #include "Containers/Deque.h"
 #include "Containers/MpscQueue.h"
-#include "Elements/Common/TypedElementHandles.h"
-#include "Elements/Common/TypedElementMapKey.h"
+#include "DataStorage/Handles.h"
+#include "DataStorage/MapKey.h"
 #include "Elements/Interfaces/TypedElementQueryStorageInterfaces.h"
 #include "GlobalLock.h"
 #include "Misc/Timespan.h"
diff --git a/Engine/Plugins/Experimental/EditorDataStorage/Source/TedsCore/Private/TypedElementDatabaseUI.h b/Engine/Plugins/Experimental/EditorDataStorage/Source/TedsCore/Private/TypedElementDatabaseUI.h
--- a/Engine/Plugins/Experimental/EditorDataStorage/Source/TedsCore/Private/TypedElementDatabaseUI.h
+++ b/Engine/Plugins/Experimental/EditorDataStorage/Source/TedsCore/Private/TypedElementDatabaseUI.h
@@ -4,7 +4,7 @@
 
 #include <type_traits>
 #include "Containers/Map.h"
-#include "Elements/Common/TypedElementQueryConditions.h"
+#include "DataStorage/Queries/Conditions.h"
 #include "Elements/Interfaces/TypedElementDataStorageUiInterface.h"
 #include "Logging/LogMacros.h"
 #include "UObject/ObjectMacros.h"
diff --git a/Engine/Plugins/Experimental/EditorDataStorage/Source/TedsCore/Private/TypedElementsDataStorage.cpp b/Engine/Plugins/Experimental/EditorDataStorage/Source/TedsCore/Private/TypedElementsDataStorage.cpp
--- a/Engine/Plugins/Experimental/EditorDataStorage/Source/TedsCore/Private/TypedElementsDataStorage.cpp
+++ b/Engine/Plugins/Experimental/EditorDataStorage/Source/TedsCore/Private/TypedElementsDataStorage.cpp
@@ -2,9 +2,9 @@
 
 #include "TypedElementsDataStorage.h"
 
+#include "DataStorage/Features.h"
+#include "DataStorage/Debug/Log.h"
 #include "EditorDataStorageSettings.h"
-#include "Elements/Common/TypedElementDataStorageLog.h"
-#include "Elements/Common/EditorDataStorageFeatures.h"
 #include "Elements/Framework/TypedElementRegistry.h"
 #include "Elements/Interfaces/TypedElementDataStorageFactory.h"
 #include "Features/IModularFeatures.h"
diff --git a/Engine/Plugins/Experimental/EditorDataStorage/Source/TedsCore/Private/Widgets/STedsWidget.cpp b/Engine/Plugins/Experimental/EditorDataStorage/Source/TedsCore/Private/Widgets/STedsWidget.cpp
--- a/Engine/Plugins/Experimental/EditorDataStorage/Source/TedsCore/Private/Widgets/STedsWidget.cpp
+++ b/Engine/Plugins/Experimental/EditorDataStorage/Source/TedsCore/Private/Widgets/STedsWidget.cpp
@@ -2,8 +2,8 @@
 
 #include "STedsWidget.h"
 
+#include "DataStorage/Features.h"
 #include "Elements/Columns/TypedElementSlateWidgetColumns.h"
-#include "Elements/Common//EditorDataStorageFeatures.h"
 #include "Elements/Framework/TypedElementAttributeBinding.h"
 #include "Elements/Interfaces/TypedElementDataStorageInterface.h"
 #include "Elements/Interfaces/TypedElementDataStorageUiInterface.h"
diff --git a/Engine/Plugins/Experimental/EditorDataStorage/Source/TedsCore/Tests/DynamicColumnTests.cpp b/Engine/Plugins/Experimental/EditorDataStorage/Source/TedsCore/Tests/DynamicColumnTests.cpp
--- a/Engine/Plugins/Experimental/EditorDataStorage/Source/TedsCore/Tests/DynamicColumnTests.cpp
+++ b/Engine/Plugins/Experimental/EditorDataStorage/Source/TedsCore/Tests/DynamicColumnTests.cpp
@@ -6,7 +6,7 @@
 #include "Misc/CoreDelegates.h"
 #if WITH_TESTS
 
-#include "Elements/Common/EditorDataStorageFeatures.h"
+#include "DataStorage/Features.h"
 #include "Elements/Framework/TypedElementQueryBuilder.h"
 #include "Elements/Interfaces/TypedElementDataStorageInterface.h"
 
diff --git a/Engine/Plugins/Experimental/EditorDataStorage/Source/TedsUI/Public/Widgets/ExportedTextWidget.h b/Engine/Plugins/Experimental/EditorDataStorage/Source/TedsUI/Public/Widgets/ExportedTextWidget.h
--- a/Engine/Plugins/Experimental/EditorDataStorage/Source/TedsUI/Public/Widgets/ExportedTextWidget.h
+++ b/Engine/Plugins/Experimental/EditorDataStorage/Source/TedsUI/Public/Widgets/ExportedTextWidget.h
@@ -2,9 +2,9 @@
 
 #pragma once
 
+#include "DataStorage/Queries/Conditions.h"
 #include "Elements/Interfaces/TypedElementDataStorageFactory.h"
 #include "Elements/Interfaces/TypedElementDataStorageUiInterface.h"
-#include "Elements/Common/TypedElementQueryConditions.h"
 
 #include "ExportedTextWidget.generated.h"
 
diff --git a/Engine/Plugins/Experimental/EditorDataStorageFeatures/Source/TedsActorCompatibility/Private/Compatibility/TedsActorParentQueries.cpp b/Engine/Plugins/Experimental/EditorDataStorageFeatures/Source/TedsActorCompatibility/Private/Compatibility/TedsActorParentQueries.cpp
--- a/Engine/Plugins/Experimental/EditorDataStorageFeatures/Source/TedsActorCompatibility/Private/Compatibility/TedsActorParentQueries.cpp
+++ b/Engine/Plugins/Experimental/EditorDataStorageFeatures/Source/TedsActorCompatibility/Private/Compatibility/TedsActorParentQueries.cpp
@@ -2,10 +2,10 @@
 
 #include "Compatibility/TedsActorParentQueries.h"
 
+#include "DataStorage/MapKey.h"
 #include "Elements/Columns/TypedElementCompatibilityColumns.h"
 #include "Elements/Columns/TypedElementHiearchyColumns.h"
 #include "Elements/Columns/TypedElementMiscColumns.h"
-#include "Elements/Common/TypedElementMapKey.h"
 #include "Elements/Framework/TypedElementQueryBuilder.h"
 #include "GameFramework/Actor.h"
 
diff --git a/Engine/Plugins/Experimental/EditorDataStorageFeatures/Source/TedsAlerts/Private/TedsAlertWidget.h b/Engine/Plugins/Experimental/EditorDataStorageFeatures/Source/TedsAlerts/Private/TedsAlertWidget.h
--- a/Engine/Plugins/Experimental/EditorDataStorageFeatures/Source/TedsAlerts/Private/TedsAlertWidget.h
+++ b/Engine/Plugins/Experimental/EditorDataStorageFeatures/Source/TedsAlerts/Private/TedsAlertWidget.h
@@ -2,8 +2,8 @@
 
 #pragma once
 
+#include "DataStorage/CommonTypes.h"
 #include "Elements/Interfaces/TypedElementDataStorageFactory.h"
-#include "Elements/Common/TypedElementCommonTypes.h"
 #include "Elements/Interfaces/TypedElementDataStorageUiInterface.h"
 
 #include "TedsAlertWidget.generated.h"
diff --git a/Engine/Plugins/Experimental/EditorDataStorageFeatures/Source/TedsAlerts/Private/TedsAlertsFactory.h b/Engine/Plugins/Experimental/EditorDataStorageFeatures/Source/TedsAlerts/Private/TedsAlertsFactory.h
--- a/Engine/Plugins/Experimental/EditorDataStorageFeatures/Source/TedsAlerts/Private/TedsAlertsFactory.h
+++ b/Engine/Plugins/Experimental/EditorDataStorageFeatures/Source/TedsAlerts/Private/TedsAlertsFactory.h
@@ -2,7 +2,7 @@
 
 #pragma once
 
-#include "Elements/Common/TypedElementHandles.h"
+#include "DataStorage/Handles.h"
 #include "Elements/Interfaces/TypedElementDataStorageFactory.h"
 #include "TedsAlertColumns.h"
 #include "UObject/ObjectMacros.h"
diff --git a/Engine/Plugins/Experimental/EditorDataStorageFeatures/Source/TedsAssetData/Private/CB/TedsAssetDataCBDataSource.h b/Engine/Plugins/Experimental/EditorDataStorageFeatures/Source/TedsAssetData/Private/CB/TedsAssetDataCBDataSource.h
--- a/Engine/Plugins/Experimental/EditorDataStorageFeatures/Source/TedsAssetData/Private/CB/TedsAssetDataCBDataSource.h
+++ b/Engine/Plugins/Experimental/EditorDataStorageFeatures/Source/TedsAssetData/Private/CB/TedsAssetDataCBDataSource.h
@@ -3,10 +3,10 @@
 #pragma once
 
 #include "CB/TagsMetadataCache.h"
+#include "DataStorage/Handles.h"
+#include "DataStorage/CommonTypes.h"
 #include "Containers/StringFwd.h"
 #include "Containers/UnrealString.h"
-#include "Elements/Common/TypedElementCommonTypes.h"
-#include "Elements/Common/TypedElementHandles.h"
 #include "Interfaces/IPluginManager.h"
 
 class FAssetPackageData;
diff --git a/Engine/Plugins/Experimental/EditorDataStorageFeatures/Source/TedsAssetData/Private/Tests/TedsAssetDataTest.cpp b/Engine/Plugins/Experimental/EditorDataStorageFeatures/Source/TedsAssetData/Private/Tests/TedsAssetDataTest.cpp
--- a/Engine/Plugins/Experimental/EditorDataStorageFeatures/Source/TedsAssetData/Private/Tests/TedsAssetDataTest.cpp
+++ b/Engine/Plugins/Experimental/EditorDataStorageFeatures/Source/TedsAssetData/Private/Tests/TedsAssetDataTest.cpp
@@ -7,8 +7,8 @@
 #include "Containers/ChunkedArray.h"
 #include "Containers/Set.h"
 #include "Containers/UnrealString.h"
-#include "Elements/Common/EditorDataStorageFeatures.h"
-#include "Elements/Common/TypedElementHandles.h"
+#include "DataStorage/Features.h"
+#include "DataStorage/Handles.h"
 #include "Elements/Framework/TypedElementRegistry.h"
 #include "Elements/Interfaces/TypedElementDataStorageInterface.h"
 #include "Misc/ScopeExit.h"
diff --git a/Engine/Plugins/Experimental/EditorDataStorageFeatures/Source/TedsAssetData/Public/TedsAssetDataColumns.h b/Engine/Plugins/Experimental/EditorDataStorageFeatures/Source/TedsAssetData/Public/TedsAssetDataColumns.h
--- a/Engine/Plugins/Experimental/EditorDataStorageFeatures/Source/TedsAssetData/Public/TedsAssetDataColumns.h
+++ b/Engine/Plugins/Experimental/EditorDataStorageFeatures/Source/TedsAssetData/Public/TedsAssetDataColumns.h
@@ -8,8 +8,8 @@
 #include "AssetViewTypes.h" 
 #include "Containers/UnrealString.h"
 #include "Containers/VersePath.h"
-#include "Elements/Common/TypedElementCommonTypes.h"
-#include "Elements/Common/TypedElementMapKey.h"
+#include "DataStorage/CommonTypes.h"
+#include "DataStorage/MapKey.h"
 #include "Elements/Interfaces/TypedElementQueryStorageInterfaces.h"
 #include "Internationalization/Text.h"
 #include "Templates/SharedPointer.h"
diff --git a/Engine/Plugins/Experimental/EditorDataStorageFeatures/Source/TedsAssetData/Public/Widgets/AssetPreview/TedsAssetPreviewWidgetColumns.h b/Engine/Plugins/Experimental/EditorDataStorageFeatures/Source/TedsAssetData/Public/Widgets/AssetPreview/TedsAssetPreviewWidgetColumns.h
--- a/Engine/Plugins/Experimental/EditorDataStorageFeatures/Source/TedsAssetData/Public/Widgets/AssetPreview/TedsAssetPreviewWidgetColumns.h
+++ b/Engine/Plugins/Experimental/EditorDataStorageFeatures/Source/TedsAssetData/Public/Widgets/AssetPreview/TedsAssetPreviewWidgetColumns.h
@@ -2,7 +2,7 @@
 
 #pragma once
 
-#include <Elements/Common/TypedElementCommonTypes.h>
+#include "DataStorage/CommonTypes.h"
 #include "Framework/SlateDelegates.h"
 
 #include "TedsAssetPreviewWidgetColumns.generated.h"
diff --git a/Engine/Plugins/Experimental/EditorDataStorageFeatures/Source/TedsContentBrowser/Private/TedsContentBrowserModule.cpp b/Engine/Plugins/Experimental/EditorDataStorageFeatures/Source/TedsContentBrowser/Private/TedsContentBrowserModule.cpp
--- a/Engine/Plugins/Experimental/EditorDataStorageFeatures/Source/TedsContentBrowser/Private/TedsContentBrowserModule.cpp
+++ b/Engine/Plugins/Experimental/EditorDataStorageFeatures/Source/TedsContentBrowser/Private/TedsContentBrowserModule.cpp
@@ -6,10 +6,10 @@
 #include "ContentBrowserModule.h"
 #include "IContentBrowserSingleton.h"
 #include "Columns/SlateDelegateColumns.h"
+#include "DataStorage/Features.h"
+#include "DataStorage/Queries/Types.h"
 #include "Elements/Columns/TypedElementFolderColumns.h"
 #include "Elements/Columns/TypedElementMiscColumns.h"
-#include "Elements/Common/TypedElementQueryTypes.h"
-#include "Elements/Common/EditorDataStorageFeatures.h"
 #include "Elements/Interfaces/TypedElementDataStorageInterface.h"
 #include "Experimental/ContentBrowserViewExtender.h"
 #include "Modules/ModuleManager.h"
diff --git a/Engine/Plugins/Experimental/EditorDataStorageFeatures/Source/TedsDebugger/Private/QueryEditor/TedsQueryEditorModel.cpp b/Engine/Plugins/Experimental/EditorDataStorageFeatures/Source/TedsDebugger/Private/QueryEditor/TedsQueryEditorModel.cpp
--- a/Engine/Plugins/Experimental/EditorDataStorageFeatures/Source/TedsDebugger/Private/QueryEditor/TedsQueryEditorModel.cpp
+++ b/Engine/Plugins/Experimental/EditorDataStorageFeatures/Source/TedsDebugger/Private/QueryEditor/TedsQueryEditorModel.cpp
@@ -1,8 +1,8 @@
 // Copyright Epic Games, Inc. All Rights Reserved.
 #include "QueryEditor/TedsQueryEditorModel.h"
 
-#include "Elements/Common/TypedElementCommonTypes.h"
-#include "Elements/Common/TypedElementDataStorageLog.h"
+#include "DataStorage/CommonTypes.h"
+#include "DataStorage/Debug/Log.h"
 #include "Elements/Framework/TypedElementColumnUtils.h"
 #include "Elements/Interfaces/TypedElementDataStorageInterface.h"
 #include "UObject/UObjectIterator.h"
diff --git a/Engine/Plugins/Experimental/EditorDataStorageFeatures/Source/TedsDebugger/Private/QueryEditor/TedsQueryEditorModel.h b/Engine/Plugins/Experimental/EditorDataStorageFeatures/Source/TedsDebugger/Private/QueryEditor/TedsQueryEditorModel.h
--- a/Engine/Plugins/Experimental/EditorDataStorageFeatures/Source/TedsDebugger/Private/QueryEditor/TedsQueryEditorModel.h
+++ b/Engine/Plugins/Experimental/EditorDataStorageFeatures/Source/TedsDebugger/Private/QueryEditor/TedsQueryEditorModel.h
@@ -1,6 +1,7 @@
 // Copyright Epic Games, Inc. All Rights Reserved.
 #pragma once
-#include "Elements/Common/TypedElementHandles.h"
+
+#include "DataStorage/Handles.h"
 #include "Elements/Framework/TypedElementMetaData.h"
 #include "UObject/Class.h"
 
diff --git a/Engine/Plugins/Experimental/EditorDataStorageFeatures/Source/TedsDebugger/Private/Widgets/QueryEditor/TedsQueryEditorResultsView.cpp b/Engine/Plugins/Experimental/EditorDataStorageFeatures/Source/TedsDebugger/Private/Widgets/QueryEditor/TedsQueryEditorResultsView.cpp
--- a/Engine/Plugins/Experimental/EditorDataStorageFeatures/Source/TedsDebugger/Private/Widgets/QueryEditor/TedsQueryEditorResultsView.cpp
+++ b/Engine/Plugins/Experimental/EditorDataStorageFeatures/Source/TedsDebugger/Private/Widgets/QueryEditor/TedsQueryEditorResultsView.cpp
@@ -5,8 +5,8 @@
 #include "SWarningOrErrorBox.h"
 #include "TedsOutlinerModule.h"
 #include "TedsTableViewerColumn.h"
-#include "Elements/Common/EditorDataStorageFeatures.h"
-#include "Elements/Common/TypedElementDataStorageLog.h"
+#include "DataStorage/Features.h"
+#include "DataStorage/Debug/Log.h"
 #include "Elements/Framework/TypedElementQueryBuilder.h"
 #include "Elements/Interfaces/TypedElementDataStorageInterface.h"
 #include "Elements/Interfaces/TypedElementDataStorageUiInterface.h"
diff --git a/Engine/Plugins/Experimental/EditorDataStorageFeatures/Source/TedsOutliner/Private/TedsOutlinerItem.cpp b/Engine/Plugins/Experimental/EditorDataStorageFeatures/Source/TedsOutliner/Private/TedsOutlinerItem.cpp
--- a/Engine/Plugins/Experimental/EditorDataStorageFeatures/Source/TedsOutliner/Private/TedsOutlinerItem.cpp
+++ b/Engine/Plugins/Experimental/EditorDataStorageFeatures/Source/TedsOutliner/Private/TedsOutlinerItem.cpp
@@ -10,8 +10,8 @@
 #include "UObject/PropertyBagRepository.h"
 
 #include "Columns/TedsOutlinerColumns.h"
+#include "DataStorage/Features.h"
 #include "Elements/Columns/TypedElementCompatibilityColumns.h"
-#include "Elements/Common//EditorDataStorageFeatures.h"
 #include "Elements/Interfaces/TypedElementDataStorageInterface.h"
 
 
diff --git a/Engine/Plugins/Experimental/EditorDataStorageFeatures/Source/TedsOutliner/Public/TedsOutlinerFilter.h b/Engine/Plugins/Experimental/EditorDataStorageFeatures/Source/TedsOutliner/Public/TedsOutlinerFilter.h
--- a/Engine/Plugins/Experimental/EditorDataStorageFeatures/Source/TedsOutliner/Public/TedsOutlinerFilter.h
+++ b/Engine/Plugins/Experimental/EditorDataStorageFeatures/Source/TedsOutliner/Public/TedsOutlinerFilter.h
@@ -4,7 +4,7 @@
 
 #include "Filters/FilterBase.h"
 #include "SceneOutlinerPublicTypes.h"
-#include "Elements/Common/TypedElementQueryDescription.h"
+#include "DataStorage/Queries/Description.h"
 #include "Elements/Framework/TypedElementMetaData.h"
 
 namespace UE::Editor::Outliner
diff --git a/Engine/Plugins/Experimental/EditorDataStorageFeatures/Source/TedsOutliner/Public/TedsOutlinerModule.h b/Engine/Plugins/Experimental/EditorDataStorageFeatures/Source/TedsOutliner/Public/TedsOutlinerModule.h
--- a/Engine/Plugins/Experimental/EditorDataStorageFeatures/Source/TedsOutliner/Public/TedsOutlinerModule.h
+++ b/Engine/Plugins/Experimental/EditorDataStorageFeatures/Source/TedsOutliner/Public/TedsOutlinerModule.h
@@ -2,11 +2,11 @@
 
 #pragma once
 
-#include "Elements/Common/TypedElementHandles.h"
+#include "DataStorage/Handles.h"
 #include "Modules/ModuleInterface.h"
 #include "Templates/SharedPointer.h"
 #include "Delegates/IDelegateInstance.h"
-#include "Elements/Common/TypedElementQueryDescription.h"
+#include "DataStorage/Queries/Description.h"
 
 class ISceneOutliner;
 struct FSceneOutlinerInitializationOptions;
diff --git a/Engine/Plugins/Experimental/EditorDataStorageFeatures/Source/TedsOutliner/Public/Widgets/TransformHeadsUpWidget.h b/Engine/Plugins/Experimental/EditorDataStorageFeatures/Source/TedsOutliner/Public/Widgets/TransformHeadsUpWidget.h
--- a/Engine/Plugins/Experimental/EditorDataStorageFeatures/Source/TedsOutliner/Public/Widgets/TransformHeadsUpWidget.h
+++ b/Engine/Plugins/Experimental/EditorDataStorageFeatures/Source/TedsOutliner/Public/Widgets/TransformHeadsUpWidget.h
@@ -2,10 +2,10 @@
 
 #pragma once
 
+#include "DataStorage/Handles.h"
+#include "DataStorage/CommonTypes.h"
 #include "Elements/Interfaces/TypedElementDataStorageFactory.h"
 #include "Elements/Interfaces/TypedElementDataStorageUiInterface.h"
-#include "Elements/Common/TypedElementCommonTypes.h"
-#include "Elements/Common/TypedElementHandles.h"
 #include "Elements/Framework/TypedElementMetaData.h"
 #include "Internationalization/Text.h"
 #include "UObject/ObjectMacros.h"
diff --git a/Engine/Plugins/Experimental/EditorDataStorageFeatures/Source/TedsQueryStack/Public/TedsQueryMergeNode.h b/Engine/Plugins/Experimental/EditorDataStorageFeatures/Source/TedsQueryStack/Public/TedsQueryMergeNode.h
--- a/Engine/Plugins/Experimental/EditorDataStorageFeatures/Source/TedsQueryStack/Public/TedsQueryMergeNode.h
+++ b/Engine/Plugins/Experimental/EditorDataStorageFeatures/Source/TedsQueryStack/Public/TedsQueryMergeNode.h
@@ -2,7 +2,7 @@
 
 #pragma once
 
-#include "Elements/Common/TypedElementQueryDescription.h"
+#include "DataStorage/Queries/Description.h"
 #include "TedsQueryNode.h"
 
 namespace UE::Editor::DataStorage
diff --git a/Engine/Plugins/Experimental/EditorDataStorageFeatures/Source/TedsRevisionControl/Private/Processors/RevisionControlProcessors.cpp b/Engine/Plugins/Experimental/EditorDataStorageFeatures/Source/TedsRevisionControl/Private/Processors/RevisionControlProcessors.cpp
--- a/Engine/Plugins/Experimental/EditorDataStorageFeatures/Source/TedsRevisionControl/Private/Processors/RevisionControlProcessors.cpp
+++ b/Engine/Plugins/Experimental/EditorDataStorageFeatures/Source/TedsRevisionControl/Private/Processors/RevisionControlProcessors.cpp
@@ -7,13 +7,13 @@
 #include "GameFramework/Actor.h"
 #include "HAL/IConsoleManager.h"
 
+#include "DataStorage/Features.h"
+#include "DataStorage/Debug/Log.h"
 #include "Elements/Columns/TypedElementCompatibilityColumns.h"
 #include "Elements/Columns/TypedElementMiscColumns.h"
 #include "Elements/Columns/TypedElementRevisionControlColumns.h"
 #include "Elements/Columns/TypedElementSelectionColumns.h"
 #include "Elements/Columns/TypedElementViewportColumns.h"
-#include "Elements/Common/EditorDataStorageFeatures.h"
-#include "Elements/Common/TypedElementDataStorageLog.h"
 #include "Elements/Framework/TypedElementQueryBuilder.h"
 #include "Elements/Interfaces/TypedElementQueryStorageInterfaces.h"
 
diff --git a/Engine/Plugins/Experimental/EditorDataStorageFeatures/Source/TedsTableViewer/Private/TedsTableViewerUtils.cpp b/Engine/Plugins/Experimental/EditorDataStorageFeatures/Source/TedsTableViewer/Private/TedsTableViewerUtils.cpp
--- a/Engine/Plugins/Experimental/EditorDataStorageFeatures/Source/TedsTableViewer/Private/TedsTableViewerUtils.cpp
+++ b/Engine/Plugins/Experimental/EditorDataStorageFeatures/Source/TedsTableViewer/Private/TedsTableViewerUtils.cpp
@@ -3,11 +3,11 @@
 #include "TedsTableViewerUtils.h"
 
 #include "GameFramework/Actor.h"
+#include "DataStorage/Debug/Log.h"
 #include "Elements/Columns/TypedElementCompatibilityColumns.h"
 #include "Elements/Columns/TypedElementIconOverrideColumns.h"
 #include "Elements/Columns/TypedElementMiscColumns.h"
 #include "Elements/Columns/TypedElementTypeInfoColumns.h"
-#include "Elements/Common/TypedElementDataStorageLog.h"
 #include "Elements/Interfaces/TypedElementDataStorageUiInterface.h"
 #include "Styling/SlateIconFinder.h"
 
diff --git a/Engine/Plugins/Experimental/EditorDataStorageFeatures/Source/TedsTableViewer/Public/TedsTableViewerColumn.h b/Engine/Plugins/Experimental/EditorDataStorageFeatures/Source/TedsTableViewer/Public/TedsTableViewerColumn.h
--- a/Engine/Plugins/Experimental/EditorDataStorageFeatures/Source/TedsTableViewer/Public/TedsTableViewerColumn.h
+++ b/Engine/Plugins/Experimental/EditorDataStorageFeatures/Source/TedsTableViewer/Public/TedsTableViewerColumn.h
@@ -1,7 +1,8 @@
 // Copyright Epic Games, Inc. All Rights Reserved.
 
 #pragma once
-#include "Elements/Common/TypedElementHandles.h"
+
+#include "DataStorage/Handles.h"
 #include "Elements/Framework/TypedElementMetaData.h"
 #include "Elements/Framework/TypedElementQueryBuilder.h"
 #include "Elements/Interfaces/TypedElementDataStorageUiInterface.h"
diff --git a/Engine/Plugins/Experimental/EditorDataStorageFeatures/Source/TedsTableViewer/Public/TedsTableViewerWidgetColumns.h b/Engine/Plugins/Experimental/EditorDataStorageFeatures/Source/TedsTableViewer/Public/TedsTableViewerWidgetColumns.h
--- a/Engine/Plugins/Experimental/EditorDataStorageFeatures/Source/TedsTableViewer/Public/TedsTableViewerWidgetColumns.h
+++ b/Engine/Plugins/Experimental/EditorDataStorageFeatures/Source/TedsTableViewer/Public/TedsTableViewerWidgetColumns.h
@@ -1,8 +1,8 @@
-﻿// Copyright Epic Games, Inc. All Rights Reserved.
+// Copyright Epic Games, Inc. All Rights Reserved.
 
 #pragma once
 
-#include "Elements/Common/TypedElementCommonTypes.h"
+#include "DataStorage/CommonTypes.h"
 #include "Framework/SlateDelegates.h"
 
 #include "TedsTableViewerWidgetColumns.generated.h"
diff --git a/Engine/Plugins/Experimental/EditorDataStorageFeatures/Source/TedsTableViewer/Public/Widgets/RowHandleWidget.h b/Engine/Plugins/Experimental/EditorDataStorageFeatures/Source/TedsTableViewer/Public/Widgets/RowHandleWidget.h
--- a/Engine/Plugins/Experimental/EditorDataStorageFeatures/Source/TedsTableViewer/Public/Widgets/RowHandleWidget.h
+++ b/Engine/Plugins/Experimental/EditorDataStorageFeatures/Source/TedsTableViewer/Public/Widgets/RowHandleWidget.h
@@ -2,9 +2,9 @@
 
 #pragma once
 
+#include "DataStorage/Queries/Conditions.h"
 #include "Elements/Interfaces/TypedElementDataStorageFactory.h"
 #include "Elements/Interfaces/TypedElementDataStorageUiInterface.h"
-#include "Elements/Common/TypedElementQueryConditions.h"
 
 #include "RowHandleWidget.generated.h"
 
diff --git a/Engine/Source/Editor/UnrealEd/Private/EditorActor.cpp b/Engine/Source/Editor/UnrealEd/Private/EditorActor.cpp
--- a/Engine/Source/Editor/UnrealEd/Private/EditorActor.cpp
+++ b/Engine/Source/Editor/UnrealEd/Private/EditorActor.cpp
@@ -77,9 +77,8 @@
 #include "Subsystems/AssetEditorSubsystem.h"
 #include "SceneInterface.h"
 #include "SDeleteReferencedActorDialog.h"
+#include "DataStorage/Features.h"
 #include "Elements/Interfaces/TypedElementDataStorageCompatibilityInterface.h"
-#include "Elements/Common/EditorDataStorageFeatures.h"
-#include "Elements/Common/EditorDataStorageFeatures.h"
 #include "Elements/Columns/TypedElementVisibilityColumns.h"
 #include "Elements/Columns/TypedElementMiscColumns.h"
 
diff --git a/Engine/Source/Editor/UnrealEd/Private/Elements/Actor/ActorElementEditorSelectionInterface.cpp b/Engine/Source/Editor/UnrealEd/Private/Elements/Actor/ActorElementEditorSelectionInterface.cpp
--- a/Engine/Source/Editor/UnrealEd/Private/Elements/Actor/ActorElementEditorSelectionInterface.cpp
+++ b/Engine/Source/Editor/UnrealEd/Private/Elements/Actor/ActorElementEditorSelectionInterface.cpp
@@ -1,10 +1,10 @@
 // Copyright Epic Games, Inc. All Rights Reserved.
  
 #include "Elements/Actor/ActorElementEditorSelectionInterface.h"
- 
+
+#include "DataStorage/Features.h"
 #include "Elements/Actor/ActorElementData.h"
 #include "Elements/Columns/TypedElementSelectionColumns.h"
-#include "Elements/Common/EditorDataStorageFeatures.h"
 #include "Elements/Framework/EngineElementsLibrary.h"
 #include "Elements/Framework/TypedElementHandle.h"
 #include "Elements/Framework/TypedElementList.h"
diff --git a/Engine/Source/Runtime/TypedElementFramework/Public/Elements/Columns/TypedElementHiearchyColumns.h b/Engine/Source/Runtime/TypedElementFramework/Public/Elements/Columns/TypedElementHiearchyColumns.h
--- a/Engine/Source/Runtime/TypedElementFramework/Public/Elements/Columns/TypedElementHiearchyColumns.h
+++ b/Engine/Source/Runtime/TypedElementFramework/Public/Elements/Columns/TypedElementHiearchyColumns.h
@@ -2,8 +2,8 @@
 
 #pragma once
 
-#include "Elements/Common/TypedElementHandles.h"
-#include "Elements/Common/TypedElementMapKey.h"
+#include "DataStorage/Handles.h"
+#include "DataStorage/MapKey.h"
 #include "Elements/Interfaces/TypedElementDataStorageInterface.h"
 #include "UObject/ObjectMacros.h"
 
diff --git a/Engine/Source/Runtime/TypedElementFramework/Public/Elements/Columns/TypedElementPackageColumns.h b/Engine/Source/Runtime/TypedElementFramework/Public/Elements/Columns/TypedElementPackageColumns.h
--- a/Engine/Source/Runtime/TypedElementFramework/Public/Elements/Columns/TypedElementPackageColumns.h
+++ b/Engine/Source/Runtime/TypedElementFramework/Public/Elements/Columns/TypedElementPackageColumns.h
@@ -2,11 +2,11 @@
 
 #pragma once
 
+#include "DataStorage/CommonTypes.h"
+#include "DataStorage/Handles.h"
+#include "DataStorage/MapKey.h"
+#include "DataStorage/Queries/Types.h"
 #include "Containers/UnrealString.h"
-#include "Elements/Common/TypedElementCommonTypes.h"
-#include "Elements/Common/TypedElementMapKey.h"
-#include "Elements/Common/TypedElementHandles.h"
-#include "Elements/Common/TypedElementQueryTypes.h"
 #include "HAL/Platform.h"
 #include "IO/PackageId.h"
 #include "Misc/PackagePath.h"
diff --git a/Engine/Source/Runtime/TypedElementFramework/Public/Elements/Common/TypedElementDataStorageLog.h b/Engine/Source/Runtime/TypedElementFramework/Public/Elements/Common/TypedElementDataStorageLog.h
--- a/Engine/Source/Runtime/TypedElementFramework/Public/Elements/Common/TypedElementDataStorageLog.h
+++ b/Engine/Source/Runtime/TypedElementFramework/Public/Elements/Common/TypedElementDataStorageLog.h
@@ -2,6 +2,8 @@
 
 #pragma once
 
-#include <Logging/LogMacros.h>
+#include "Misc/CoreMiscDefines.h"
 
-TYPEDELEMENTFRAMEWORK_API DECLARE_LOG_CATEGORY_EXTERN(LogEditorDataStorage, Log, All);
\ No newline at end of file
+UE_DEPRECATED_HEADER(5.7, "Use DataStorage/Debug/Log.h instead of Elements/Common/TypedElementDataStorageLog.h")
+
+#include "DataStorage/Debug/Log.h"
diff --git a/Engine/Source/Runtime/TypedElementFramework/Public/Elements/Common/TypedElementQueryConditions.h b/Engine/Source/Runtime/TypedElementFramework/Public/Elements/Common/TypedElementQueryConditions.h
--- a/Engine/Source/Runtime/TypedElementFramework/Public/Elements/Common/TypedElementQueryConditions.h
+++ b/Engine/Source/Runtime/TypedElementFramework/Public/Elements/Common/TypedElementQueryConditions.h
@@ -2,196 +2,8 @@
 
 #pragma once
 
-#include "Elements/Common/TypedElementCommonTypes.h"
-#include "Templates/SubclassOf.h"
-#include "UObject/WeakObjectPtrTemplates.h"
+#include "Misc/CoreMiscDefines.h"
 
-class UScriptStruct;
-
-namespace UE::Editor::DataStorage
-{
-	class ICoreProvider;
-
-namespace Queries
-{
-	struct FColumnBase
-	{
-		TWeakObjectPtr<const UScriptStruct> TypeInfo = nullptr;
-		FName Identifier = NAME_None;
-
-	protected:
-		FColumnBase() = default;
-		constexpr explicit FColumnBase(TWeakObjectPtr<const UScriptStruct> ColumnTypeInfo, const FName& InIdentifier)
-			: TypeInfo(ColumnTypeInfo)
-			, Identifier(InIdentifier)
-		{}
-	};
-
-	template<typename T = void>
-	struct TColumn final : public FColumnBase
-	{
-		template <typename U = T> requires (!std::is_same_v<U, void>)
-			constexpr explicit TColumn(const FName& Identifier = NAME_None) : FColumnBase(T::StaticStruct(), Identifier) {}
-		
-		template <typename U = T> requires (std::is_same_v<U, void>)
-			constexpr explicit TColumn(TWeakObjectPtr<const UScriptStruct> ColumnTypeInfo, const FName& Identifier = NAME_None)
-				: FColumnBase(ColumnTypeInfo, Identifier)
-		{}
-	};
-
-	/*
-	 * Compile Context used to resolve dynamic columns when FConditions are compiled
-	 */
-	class IQueryConditionCompileContext 
-	{
-	public:
-		virtual ~IQueryConditionCompileContext() = default;
-		virtual const UScriptStruct* GenerateDynamicColumn(const FDynamicColumnDescription&) const = 0;
-	};
-
-	/*
-	 * Specialized compile context that accepts an ICoreProvider
-	 */
-	class FEditorStorageQueryConditionCompileContext : public IQueryConditionCompileContext 
-	{
-	public:
-		TYPEDELEMENTFRAMEWORK_API explicit FEditorStorageQueryConditionCompileContext(ICoreProvider* InDataStorage);
-		TYPEDELEMENTFRAMEWORK_API virtual const UScriptStruct* GenerateDynamicColumn(const FDynamicColumnDescription& Description) const override;
-		
-	private:
-		ICoreProvider* DataStorage;
-	};
-
-
-	/**
-	 * Product of boolean combination of multiple columns. This can be used to verify if a collection of columns match
-	 * the stored columns.
-	 * NOTE: You must call Compile() before you call any members accessing the conditions
-	 */
-	class FConditions final
-	{
-	public:
-		using ContainsCallback = TFunctionRef<bool(uint8_t ColumnIndex, TWeakObjectPtr<const UScriptStruct> Column)>;
-
-		TYPEDELEMENTFRAMEWORK_API friend FConditions operator&&(const FConditions& Lhs, FColumnBase Rhs);
-		TYPEDELEMENTFRAMEWORK_API friend FConditions operator&&(const FConditions& Lhs, const FConditions& Rhs);
-		TYPEDELEMENTFRAMEWORK_API friend FConditions operator&&(FColumnBase Lhs, FColumnBase Rhs);
-		TYPEDELEMENTFRAMEWORK_API friend FConditions operator&&(FColumnBase Lhs, const FConditions& Rhs);
-
-		TYPEDELEMENTFRAMEWORK_API friend FConditions operator||(const FConditions& Lhs, FColumnBase Rhs);
-		TYPEDELEMENTFRAMEWORK_API friend FConditions operator||(const FConditions& Lhs, const FConditions& Rhs);
-		TYPEDELEMENTFRAMEWORK_API friend FConditions operator||(FColumnBase Lhs, FColumnBase Rhs);
-		TYPEDELEMENTFRAMEWORK_API friend FConditions operator||(FColumnBase Lhs, const FConditions& Rhs);
-
-		TYPEDELEMENTFRAMEWORK_API FConditions();
-		// Not marked as "explicit" to allow conversion from a column. This means that conditions with a single
-		// argument can be written in the same way as ones that use combinations.
-		TYPEDELEMENTFRAMEWORK_API FConditions(FColumnBase Column);
-
-		// Compile must be called before using any functions that access the columns
-		TYPEDELEMENTFRAMEWORK_API FConditions& Compile(const IQueryConditionCompileContext& CompileContext);
-
-		// Check whether these query conditions have been compiled
-		TYPEDELEMENTFRAMEWORK_API bool IsCompiled() const;
-
-		/** Convert the conditions into a string and append them to the provided string. */
-		TYPEDELEMENTFRAMEWORK_API void AppendToString(FString& Output) const;
-
-		/** Runs the provided list of columns through the conditions and returns true if a valid combination of columns is found. */
-		TYPEDELEMENTFRAMEWORK_API bool Verify(TConstArrayView<FColumnBase> AvailableColumns) const;
-		/**
-		 * Runs the provided list of columns through the conditions and returns true if a valid combination of columns is found.
-		 * This version returns a list of the columns that were used to match the condition.
-		 */
-		TYPEDELEMENTFRAMEWORK_API bool Verify(
-			TArray<TWeakObjectPtr<const UScriptStruct>>& MatchedColumns,
-			TConstArrayView<FColumnBase> AvailableColumns,
-			bool AvailableColumnsAreSorted = false) const;
-		/** Runs the provided list of columns through the conditions and returns true if a valid combination of columns is found. */
-		TYPEDELEMENTFRAMEWORK_API bool Verify(TConstArrayView<TWeakObjectPtr<const UScriptStruct>> AvailableColumns,
-			bool AvailableColumnsAreSorted = false) const;
-		/**
-		 * Runs the provided list of columns through the conditions and returns true if a valid combination of columns is found.
-		 * This version returns a list of the columns that were used to match the condition.
-		 */
-		TYPEDELEMENTFRAMEWORK_API bool Verify(
-			TArray<TWeakObjectPtr<const UScriptStruct>>& MatchedColumns,
-			TConstArrayView<TWeakObjectPtr<const UScriptStruct>> AvailableColumns,
-			bool AvailableColumnsAreSorted = false) const;
-		/** Runs the provided list of columns through the conditions and returns true if a valid combination of columns is found. */
-		TYPEDELEMENTFRAMEWORK_API bool Verify(TSet<TWeakObjectPtr<const UScriptStruct>> AvailableColumns) const;
-		/** Runs through the list of query conditions and uses the callback to verify if a column is available. */
-		TYPEDELEMENTFRAMEWORK_API bool Verify(ContainsCallback Callback) const;
-
-		/** Returns the minimum number of columns needed for a successful match. */
-		TYPEDELEMENTFRAMEWORK_API uint8_t MinimumColumnMatchRequired() const;
-		/** Returns a list of all columns used. This can include duplicate columns. */
-		TYPEDELEMENTFRAMEWORK_API TConstArrayView<TWeakObjectPtr<const UScriptStruct>> GetColumns() const;
-		/** Whether or not there are any columns registered for operation.*/
-		TYPEDELEMENTFRAMEWORK_API bool IsEmpty() const;
-		/** Whether the conditions contain any columns that are dynamic templates. */
-		TYPEDELEMENTFRAMEWORK_API bool UsesDynamicTemplates() const;
-
-	private:
-		template<typename AvailableColumnType, typename ProjectionFunction>
-		bool VerifyWithDynamicColumn(TConstArrayView<AvailableColumnType> AvailableColumns, uint64& OutMatches, ProjectionFunction Projection) const;
-		
-		void AppendName(FString& Output, TWeakObjectPtr<const UScriptStruct> TypeInfo) const;
-
-		bool EntersScopeNext(uint8_t Index) const;
-		bool EntersScope(uint8_t Index) const;
-		bool Contains(TWeakObjectPtr<const UScriptStruct> ColumnType, const TArray<FColumnBase>& AvailableColumns) const;
-
-		bool VerifyBootstrap(ContainsCallback Contains) const;
-		bool VerifyRange(uint8_t& TokenIndex, uint8_t& ColumnIndex, ContainsCallback Contains) const;
-
-		void ConvertColumnBitToArray(TArray<TWeakObjectPtr<const UScriptStruct>>& MatchedColumns, uint64 ColumnBits) const;
-
-		uint8_t MinimumColumnMatchRequiredRange(uint8_t& Front) const;
-
-		static void AppendQuery(FConditions& Target, const FConditions& Source);
-
-		static constexpr SIZE_T MaxColumnCount = 32;
-		static constexpr SIZE_T MaxTokenCount = 64;
-
-		enum class Token : uint8_t
-		{
-			None,
-			And,
-			Or,
-			ScopeOpen,
-			ScopeClose
-		};
-
-		enum class EColumnFlags : uint8_t
-		{
-			DynamicColumnTemplate = 1 << 0
-		};
-		FRIEND_ENUM_CLASS_FLAGS(EColumnFlags);
-		
-		using ColumnArray = TArray<TWeakObjectPtr<const UScriptStruct>, TInlineAllocator<MaxColumnCount>>;
-		using TokenArray = TArray<Token, TInlineAllocator<MaxTokenCount>>;
-		using IdentifierArray = TArray<FName, TInlineAllocator<MaxColumnCount>>;
-		using ColumnFlagArray = TArray<EColumnFlags, TInlineAllocator<MaxColumnCount>>;
-
-		ColumnArray Columns;
-		ColumnFlagArray ColumnFlags;
-		TokenArray Tokens;
-		IdentifierArray Identifiers;
-		uint8_t ColumnCount = 0;
-		uint8_t TokenCount = 0;
-		bool bIsCompiled = false;
-	};
-
-	TYPEDELEMENTFRAMEWORK_API FConditions operator&&(const FConditions& Lhs, FColumnBase Rhs);
-	TYPEDELEMENTFRAMEWORK_API FConditions operator&&(const FConditions& Lhs, const FConditions& Rhs);
-	TYPEDELEMENTFRAMEWORK_API FConditions operator&&(FColumnBase Lhs, FColumnBase Rhs);
-	TYPEDELEMENTFRAMEWORK_API FConditions operator&&(FColumnBase Lhs, const FConditions& Rhs);
-
-	TYPEDELEMENTFRAMEWORK_API FConditions operator||(const FConditions& Lhs, FColumnBase Rhs);
-	TYPEDELEMENTFRAMEWORK_API FConditions operator||(const FConditions& Lhs, const FConditions& Rhs);
-	TYPEDELEMENTFRAMEWORK_API FConditions operator||(FColumnBase Lhs, FColumnBase Rhs);
-	TYPEDELEMENTFRAMEWORK_API FConditions operator||(FColumnBase Lhs, const FConditions& Rhs);
-} // namespace Queries
-} // namespace UE::Editor::DataStorage
+UE_DEPRECATED_HEADER(5.7, "Use DataStorage/Queries/Conditions.h instead of Elements/Common/TypedElementQueryConditions.h")
 
+#include "DataStorage/Queries/Conditions.h"
diff --git a/Engine/Source/Runtime/TypedElementFramework/Public/Elements/Common/TypedElementQueryDescription.h b/Engine/Source/Runtime/TypedElementFramework/Public/Elements/Common/TypedElementQueryDescription.h
--- a/Engine/Source/Runtime/TypedElementFramework/Public/Elements/Common/TypedElementQueryDescription.h
+++ b/Engine/Source/Runtime/TypedElementFramework/Public/Elements/Common/TypedElementQueryDescription.h
@@ -2,123 +2,8 @@
 
 #pragma once
 
-#include "Containers/Array.h"
-#include "Containers/ContainerAllocationPolicies.h"
-#include "Elements/Common/TypedElementHandles.h"
-#include "Elements/Common/TypedElementCommonTypes.h"
-#include "Elements/Common/TypedElementQueryConditions.h"
-#include "Elements/Common/TypedElementQueryTypes.h"
-#include "Elements/Interfaces/TypedElementQueryStorageInterfaces.h"
-#include "Elements/Framework/TypedElementMetaData.h"
-#include "Misc/Optional.h"
-#include "Templates/Function.h"
-#include "UObject/NameTypes.h"
-#include "UObject/WeakObjectPtr.h"
-#include "UObject/WeakObjectPtrTemplates.h"
+#include "Misc/CoreMiscDefines.h"
 
-class UScriptStruct;
+UE_DEPRECATED_HEADER(5.7, "Use DataStorage/Queries/Description.h instead of Elements/Common/TypedElementQueryDescription.h")
 
-namespace UE::Editor::DataStorage
-{
-	struct FQueryDescription;
-	struct IQueryContext;
-	
-	using QueryCallback = TFunction<void(const FQueryDescription&, IQueryContext&)>;
-	using QueryCallbackRef = TFunctionRef<void(const FQueryDescription&, IQueryContext&)>;
-	using DirectQueryCallback = TFunction<void(const FQueryDescription&, IDirectQueryContext&)>;
-	using DirectQueryCallbackRef = TFunctionRef<void(const FQueryDescription&, IDirectQueryContext&)>;
-	
-	struct FQueryDescription final
-	{
-		static constexpr int32 NumInlineSelections = 8;
-		static constexpr int32 NumInlineConditions = 8;
-		static constexpr int32 NumInlineDependencies = 2;
-		static constexpr int32 NumInlineGroups = 2;
-
-		enum class EActionType : uint8
-		{
-			None,	//< Do nothing.
-			Select,	//< Selects a set of columns for further processing.
-			Count,	//< Counts the number of entries that match the filter condition.
-
-			Max //< Value indicating the maximum value in this enum. Not to be used as an enum value.
-		};
-
-		enum class EOperatorType : uint16
-		{
-			SimpleAll,			//< Unary: Type
-			SimpleAny,			//< Unary: Type
-			SimpleNone,			//< Unary: Type
-			SimpleOptional,		//< Unary: Type
-		
-			Max //< Value indicating the maximum value in this enum. Not to be used as an enum value.
-		};
-
-		union FOperator
-		{
-			TWeakObjectPtr<const UScriptStruct> Type;
-		};
-
-		struct FValueTagData
-		{
-			// The Tag maps to a Mass ConstSharedFragment object
-			FValueTag Tag;
-
-			// The MatchValue specifies the value that the fragment must have to be matched
-			// If MatchValue is NAME_None, then TEDS will match all values
-			FName MatchValue;
-		};
-
-		struct FCallbackData
-		{
-			TArray<FName, TInlineAllocator<NumInlineGroups>> BeforeGroups;
-			TArray<FName, TInlineAllocator<NumInlineGroups>> AfterGroups;
-			QueryCallback Function;
-			FName Name;
-			FName Group;
-			/** If a name is set, it indicates the query callback will not be run unless the ActivationCount is greater than zero. */
-			FName ActivationName;
-			const UScriptStruct* MonitoredType{ nullptr };
-			EQueryCallbackType Type{ EQueryCallbackType::None };
-			EQueryTickPhase Phase{ EQueryTickPhase::FrameEnd };
-			/**
-			 * The number of remaining iterations for a activatable query callback. If this is higher than 0, the query callback will be 
-			 * called. If ActivationName is set, this value will be decremented by one at the end of the update cycle.
-			 */
-			uint8 ActivationCount = 255;
-			EExecutionMode ExecutionMode = EExecutionMode::Default;
-		};
-		FCallbackData Callback;
-
-		// The list of arrays below are required to remain in the same order as they're added as the function binding expects certain entries
-		// to be in a specific location.
-
-		TArray<TWeakObjectPtr<const UScriptStruct>, TInlineAllocator<NumInlineSelections>> SelectionTypes;
-		TArray<EQueryAccessType, TInlineAllocator<NumInlineSelections>> SelectionAccessTypes;
-		TArray<FColumnMetaData, TInlineAllocator<NumInlineSelections>> SelectionMetaData;
-
-		TArray<EOperatorType, TInlineAllocator<NumInlineConditions>> ConditionTypes;
-		TArray<FOperator, TInlineAllocator<NumInlineConditions>> ConditionOperators;
-
-		TArray<FDynamicColumnDescription, TInlineAllocator<NumInlineSelections>> DynamicSelectionTypes;
-		TArray<EQueryAccessType, TInlineAllocator<NumInlineSelections>> DynamicSelectionAccessTypes;
-		TArray<FColumnMetaData::EFlags, TInlineAllocator<NumInlineSelections>> DynamicSelectionMetaData;
-		
-		TArray<EOperatorType, TInlineAllocator<NumInlineConditions>> DynamicConditionOperations;
-		TArray<FDynamicColumnDescription, TInlineAllocator<NumInlineConditions>> DynamicConditionDescriptions;
-
-		TArray<FValueTagData> ValueTags;
-
-		TOptional<UE::Editor::DataStorage::Queries::FConditions> Conditions;
-
-		TArray<TWeakObjectPtr<const UClass>, TInlineAllocator<NumInlineDependencies>> DependencyTypes;
-		TArray<EQueryDependencyFlags, TInlineAllocator<NumInlineDependencies>> DependencyFlags;
-		/** Cached instances of the dependencies. This will always match the count of the other Dependency*Types, but may contain null pointers. */
-		TArray<TWeakObjectPtr<UObject>, TInlineAllocator<NumInlineDependencies>> CachedDependencies;
-		TArray<QueryHandle> Subqueries;
-		FMetaData MetaData;
-
-		EActionType Action;
-		bool bShouldBatchModifications = false;
-	};
-} // namespace UE::Editor::DataStorage
+#include "DataStorage/Queries/Description.h"
diff --git a/Engine/Source/Runtime/TypedElementFramework/Public/Elements/Common/TypedElementQueryTypes.h b/Engine/Source/Runtime/TypedElementFramework/Public/Elements/Common/TypedElementQueryTypes.h
--- a/Engine/Source/Runtime/TypedElementFramework/Public/Elements/Common/TypedElementQueryTypes.h
+++ b/Engine/Source/Runtime/TypedElementFramework/Public/Elements/Common/TypedElementQueryTypes.h
@@ -2,133 +2,8 @@
 
 #pragma once
 
-#include "HAL/Platform.h"
-#include "Misc/EnumClassFlags.h"
+#include "Misc/CoreMiscDefines.h"
 
-namespace UE::Editor::DataStorage
-{
-	// UE_DEPRECATED(5.6, "Use FMapKey(View) with the new mapping functions.") // Disabled as it's used by the deprecated functions.
-	using IndexHash = uint64;
+UE_DEPRECATED_HEADER(5.7, "Use DataStorage/Queries/Types.h instead of Elements/Common/TypedElementQueryTypes.h")
 
-	enum class EQueryTickPhase : uint8
-	{
-		PrePhysics, //< Queries are executed before physics simulation starts.
-		DuringPhysics, //< Queries that can be run in parallel with physics simulation work.
-		PostPhysics, //< Queries that need rigid body and cloth simulation to be completed before being executed.
-		FrameEnd, //< Catchall for queries demoted to the last possible moment.
-
-		Max //< Value indicating the maximum value in this enum. Not to be used as an enum value.
-	};
-
-	enum class EQueryTickGroups : uint8
-	{
-		/** The standard group to run work in. */
-		Default,
-
-		/** General group that prepares for updating. */
-		PreUpdate,
-
-		/** General group that's used for updating data. */
-		Update,
-
-		/** General group that completes for updating. */
-		PostUpdate,
-
-		/**
-		 * The group for queries that need to sync data from external sources such as subsystems or the world into
-		 * the Data Storage. These typically run early in a phase.
-		 */
-		SyncExternalToDataStorage,
-		/**
-		 * The group for queries that need to sync data from the Data Storage to external sources such as subsystems
-		 * or the world into. These typically run late in a phase.
-		 */
-		SyncDataStorageToExternal,
-		/**
-		 * Queries grouped under this name will sync data to/from widgets.
-		 */
-		SyncWidgets,
-
-		Max //< Value indicating the maximum value in this enum. Not to be used as an enum value.
-	};
-
-	enum class EQueryCallbackType : uint8
-	{
-		/** No callback provided. */
-		None,
-		/** The query will be run every tick if at least one row matches. */
-		Processor,
-		/** The query will be run when a row is added that matches the query. The first recorded column will be actively monitored for changes. */
-		ObserveAdd,
-		/** The query will be run when a row is removed that matches the query. The first recorded column will be actively monitored for changes. */
-		ObserveRemove,
-		/**
-		 * At the start of the assigned phase this query will run if there are any matches. These queries will have any deferred operations such as
-		 * adding/removing rows/columns executed before the phase starts. This introduces sync points that hinder performance and are therefore
-		 * recommended only for queries that save on work later in the phase such as repeated checks for validity.
-		 */
-		PhasePreparation,
-		/**
-		 * At the end of the assigned phase this query will run if there are any matches. These queries will have any deferred operations such as
-		 * adding/removing rows/columns executed before the phase ends. This introduces sync points that hinder performance and are therefore
-		 * recommended only cases where delaying deferred operations is not possible e.g. when tables are known to be referenced outside the update
-		 * cycle.
-		 */
-		PhaseFinalization,
-
-		Max //< Value indicating the maximum value in this enum. Not to be used as an enum value.
-	};
-
-	enum class EExecutionMode : uint8
-	{
-		Default, //< Use the default for the query callback type or the group default.
-		GameThread, //< The query callback will be executed sequentially on the game thread.
-		Threaded, //< The entire query callback is called on a separate thread, while respecting dependencies.
-		ThreadedChunks //< The query callback is called on a separate thread for each chunk. This can not be set if there are sub-queries.
-	};
-
-	enum class EQueryAccessType : uint8
-	{
-		/** Provide read-only access. */
-		ReadOnly,
-		/** Provide read-only access, but do not include when finding matching tables. */
-		OptionalReadOnly,
-		/** Provide read and write access. */
-		ReadWrite
-	};
-
-	enum class EQueryDependencyFlags : uint8
-	{
-		None = 0,
-		/** If set the dependency is accessed as read-only. If not set the dependency requires Read/Write access. */
-		ReadOnly = 1 << 0,
-		/** If set the dependency can only be used from the game thread, otherwise it can be accessed from any thread. */
-		GameThreadBound = 1 << 1,
-		/** If set the dependency will be re-fetched every iteration, otherwise only if not fetched before. */
-		AlwaysRefresh = 1 << 2
-	};
-	ENUM_CLASS_FLAGS(EQueryDependencyFlags);
-
-	struct FQueryResult final
-	{
-		enum class ECompletion
-		{
-			/** Query could be fully executed. */
-			Fully,
-			/** Only portions of the query were executed. This is caused by a problem that was encountered partway through processing. */
-			Partially,
-			/**
-			 * The back-end doesn't support the particular query. This may be a limitation in how/where the query is run or because
-			 * the query contains actions and/or operations that are not supported.
-			 */
-			Unsupported,
-			/** The provided query is no longer available. */
-			Unavailable,
-			/** One or more dependencies declared on the query could not be retrieved. */
-			MissingDependency
-		};
-
-		uint32 Count{ 0 }; /** The number of rows were processed. */
-		ECompletion Completed{ ECompletion::Unavailable };
-	};
-} // namespace UE::Editor::DataStorage
+#include "DataStorage/Queries/Types.h"
diff --git a/Engine/Source/Runtime/TypedElementFramework/Public/Elements/Framework/TypedElementQueryBuilder.inl b/Engine/Source/Runtime/TypedElementFramework/Public/Elements/Framework/TypedElementQueryBuilder.inl
--- a/Engine/Source/Runtime/TypedElementFramework/Public/Elements/Framework/TypedElementQueryBuilder.inl
+++ b/Engine/Source/Runtime/TypedElementFramework/Public/Elements/Framework/TypedElementQueryBuilder.inl
@@ -5,10 +5,10 @@
 #include <utility>
 #include "Templates/IsConst.h"
 #include "Templates/UnrealTypeTraits.h"
+#include "DataStorage/Queries/Description.h"
+#include "DataStorage/Queries/Types.h"
 #include "Elements/Interfaces/TypedElementDataStorageCompatibilityInterface.h"
 #include "Elements/Interfaces/TypedElementQueryStorageInterfaces.h"
-#include "Elements/Common/TypedElementQueryDescription.h"
-#include "Elements/Common/TypedElementQueryTypes.h"
 
 namespace UE::Editor::DataStorage::Queries
 {
diff --git a/Engine/Source/Runtime/TypedElementFramework/Public/Elements/Interfaces/TypedElementDataStorageInterface.h b/Engine/Source/Runtime/TypedElementFramework/Public/Elements/Interfaces/TypedElementDataStorageInterface.h
--- a/Engine/Source/Runtime/TypedElementFramework/Public/Elements/Interfaces/TypedElementDataStorageInterface.h
+++ b/Engine/Source/Runtime/TypedElementFramework/Public/Elements/Interfaces/TypedElementDataStorageInterface.h
@@ -5,13 +5,13 @@
 #include <type_traits>
 #include "Containers/Array.h"
 #include "Containers/ContainersFwd.h"
+#include "DataStorage/CommonTypes.h"
+#include "DataStorage/Handles.h"
+#include "DataStorage/MapKey.h"
+#include "DataStorage/Queries/Description.h"
+#include "DataStorage/Queries/Types.h"
+#include "DataStorage/Queries/Conditions.h"
 #include "Delegates/Delegate.h"
-#include "Elements/Common/TypedElementCommonTypes.h"
-#include "Elements/Common/TypedElementHandles.h"
-#include "Elements/Common/TypedElementMapKey.h"
-#include "Elements/Common/TypedElementQueryConditions.h"
-#include "Elements/Common/TypedElementQueryDescription.h"
-#include "Elements/Common/TypedElementQueryTypes.h"
 #include "Elements/Framework/TypedElementColumnUtils.h"
 #include "Elements/Framework/TypedElementQueryCapabilities.h"
 #include "Elements/Framework/TypedElementQueryFunctions.h"
diff --git a/Engine/Source/Runtime/TypedElementFramework/Public/Elements/Interfaces/TypedElementDataStorageUiInterface.h b/Engine/Source/Runtime/TypedElementFramework/Public/Elements/Interfaces/TypedElementDataStorageUiInterface.h
--- a/Engine/Source/Runtime/TypedElementFramework/Public/Elements/Interfaces/TypedElementDataStorageUiInterface.h
+++ b/Engine/Source/Runtime/TypedElementFramework/Public/Elements/Interfaces/TypedElementDataStorageUiInterface.h
@@ -6,7 +6,7 @@
 #include "Containers/ArrayView.h"
 #include "Containers/ContainersFwd.h"
 #include "Containers/StringFwd.h"
-#include "Elements/Common/TypedElementQueryConditions.h"
+#include "DataStorage/Queries/Conditions.h"
 #include "Elements/Interfaces/TypedElementDataStorageInterface.h"
 #include "Elements/Framework/TypedElementMetaData.h"
 #include "Elements/Framework/TypedElementQueryBuilder.h"
diff --git a/Engine/Source/Runtime/TypedElementFramework/Public/Elements/Interfaces/TypedElementQueryStorageInterfaces.h b/Engine/Source/Runtime/TypedElementFramework/Public/Elements/Interfaces/TypedElementQueryStorageInterfaces.h
--- a/Engine/Source/Runtime/TypedElementFramework/Public/Elements/Interfaces/TypedElementQueryStorageInterfaces.h
+++ b/Engine/Source/Runtime/TypedElementFramework/Public/Elements/Interfaces/TypedElementQueryStorageInterfaces.h
@@ -3,10 +3,10 @@
 #pragma once
 
 #include "Containers/ContainersFwd.h"
-#include "Elements/Common/TypedElementCommonTypes.h"
-#include "Elements/Common/TypedElementHandles.h"
-#include "Elements/Common/TypedElementMapKey.h"
-#include "Elements/Common/TypedElementQueryTypes.h"
+#include "DataStorage/CommonTypes.h"
+#include "DataStorage/Handles.h"
+#include "DataStorage/MapKey.h"
+#include "DataStorage/Queries/Types.h"
 #include "UObject/Class.h"
 #include "UObject/WeakObjectPtrTemplates.h"
 #include "Templates/Function.h"
diff --git a/Engine/Source/Runtime/TypedElementFramework/Tests/Elements/Common/TypedElementMapKeyTests.cpp b/Engine/Source/Runtime/TypedElementFramework/Tests/Elements/Common/TypedElementMapKeyTests.cpp
--- a/Engine/Source/Runtime/TypedElementFramework/Tests/Elements/Common/TypedElementMapKeyTests.cpp
+++ b/Engine/Source/Runtime/TypedElementFramework/Tests/Elements/Common/TypedElementMapKeyTests.cpp
@@ -2,7 +2,7 @@
 
 #if WITH_TESTS
 
-#include "Elements/Common/TypedElementMapKey.h"
+#include "DataStorage/MapKey.h"
 #include "Tests/TestHarnessAdapter.h"
 
 namespace UE::Editor::DataStorage::Tests::Private
diff --git a/Engine/Source/Runtime/TypedElementFramework/Tests/Elements/Framework/TypedElementDataStoragePerformanceTestCommands.h b/Engine/Source/Runtime/TypedElementFramework/Tests/Elements/Framework/TypedElementDataStoragePerformanceTestCommands.h
--- a/Engine/Source/Runtime/TypedElementFramework/Tests/Elements/Framework/TypedElementDataStoragePerformanceTestCommands.h
+++ b/Engine/Source/Runtime/TypedElementFramework/Tests/Elements/Framework/TypedElementDataStoragePerformanceTestCommands.h
@@ -2,9 +2,8 @@
 
 #pragma once
 
-#include "CoreMinimal.h"
+#include "DataStorage/CommonTypes.h"
 #include "Elements/Interfaces/TypedElementDataStorageFactory.h"
-#include "Elements/Common/TypedElementCommonTypes.h"
 #include "Modules/ModuleManager.h"
 
 #include "TypedElementDataStoragePerformanceTestCommands.generated.h"
diff --git a/Engine/Source/Runtime/TypedElementFramework/Tests/Elements/Framework/TypedElementQueryBuilderTests.cpp b/Engine/Source/Runtime/TypedElementFramework/Tests/Elements/Framework/TypedElementQueryBuilderTests.cpp
--- a/Engine/Source/Runtime/TypedElementFramework/Tests/Elements/Framework/TypedElementQueryBuilderTests.cpp
+++ b/Engine/Source/Runtime/TypedElementFramework/Tests/Elements/Framework/TypedElementQueryBuilderTests.cpp
@@ -7,8 +7,8 @@
 #include "CoreTypes.h"
 #include "Containers/Array.h"
 #include "Containers/UnrealString.h"
-#include "Elements/Common/EditorDataStorageFeatures.h"
-#include "Elements/Common/TypedElementQueryConditions.h"
+#include "DataStorage/Features.h"
+#include "DataStorage/Queries/Conditions.h"
 #include "Elements/Interfaces/TypedElementDataStorageInterface.h"
 #include "Tests/TestHarnessAdapter.h"
 #include "UObject/Class.h"
diff --git a/Engine/Source/Runtime/TypedElementFramework/Tests/Elements/Framework/TypedElementTestColumns.h b/Engine/Source/Runtime/TypedElementFramework/Tests/Elements/Framework/TypedElementTestColumns.h
--- a/Engine/Source/Runtime/TypedElementFramework/Tests/Elements/Framework/TypedElementTestColumns.h
+++ b/Engine/Source/Runtime/TypedElementFramework/Tests/Elements/Framework/TypedElementTestColumns.h
@@ -2,8 +2,8 @@
 
 #pragma once
 
-#include "Elements/Common/TypedElementCommonTypes.h"
-#include "Elements/Common/TypedElementHandles.h"
+#include "DataStorage/Handles.h"
+#include "DataStorage/CommonTypes.h"
 
 #include "TypedElementTestColumns.generated.h"
 
