[FNameColumn] Deprecate the generic TEDS FNameColumn

diff --git a/Engine/Source/Runtime/TypedElementFramework/Public/Elements/Columns/TypedElementMiscColumns.h b/Engine/Source/Runtime/TypedElementFramework/Public/Elements/Columns/TypedElementMiscColumns.h
--- a/Engine/Source/Runtime/TypedElementFramework/Public/Elements/Columns/TypedElementMiscColumns.h
+++ b/Engine/Source/Runtime/TypedElementFramework/Public/Elements/Columns/TypedElementMiscColumns.h
@@ -60,12 +60,21 @@ struct FTypedElementRowReferenceColumn final : public FEditorDataStorageColumn
  * A name for this row.
  */
 USTRUCT(meta = (DisplayName = "Name"))
-struct FNameColumn final : public FEditorDataStorageColumn
+struct UE_DEPRECATED(5.7, "Use of the FNameColumn has been deprecated, please create a new column for your purpose.") FNameColumn final : public FEditorDataStorageColumn
 {
 	GENERATED_BODY()
 
-	UPROPERTY(meta = (Searchable))
+	UPROPERTY(meta = (Searchable, DeprecatedProperty))
 	FName Name;
+
+	PRAGMA_DISABLE_DEPRECATION_WARNINGS
+	FNameColumn() = default;
+	~FNameColumn() = default;
+	FNameColumn(const FNameColumn&) = default;
+	FNameColumn(FNameColumn&&) = default;
+	FNameColumn& operator=(const FNameColumn&) = default;
+	FNameColumn& operator=(FNameColumn&&) = default;
+	PRAGMA_ENABLE_DEPRECATION_WARNINGS
 };
 
 /**
