[TEDS] Functions using the new query callback can now use const and mutable columns.

The new query callbacks can now use single reference columns. The setup is the same as previous versions, but the now users can place the context anywhere among the arguments and the context no longer requires being the first argument.
Additional changes:
- Query callbacks that take no arguments are now supported as well.
- Unit tests were added.
- A mock context is now automatically created. Each function will get a TFunction generated for it that will default to an assert. Test functions can be assigned to TFunctions that are expected to be called on the context

diff --git a/Engine/Plugins/Experimental/EditorDataStorage/Source/TedsCore/Private/Queries/TypedElementQueryContext_DirectSingleApi.cpp b/Engine/Plugins/Experimental/EditorDataStorage/Source/TedsCore/Private/Queries/TypedElementQueryContext_DirectSingleApi.cpp
--- a/Engine/Plugins/Experimental/EditorDataStorage/Source/TedsCore/Private/Queries/TypedElementQueryContext_DirectSingleApi.cpp
+++ b/Engine/Plugins/Experimental/EditorDataStorage/Source/TedsCore/Private/Queries/TypedElementQueryContext_DirectSingleApi.cpp
@@ -26,4 +26,24 @@ namespace UE::Editor::DataStorage::Queries
 	{
 		CurrentRow = Row;
 	}
+
+	void FQueryContext_DirectSingleApi::GetConstColumns(TArrayView<const void*> ColumnsData, TConstArrayView<const UScriptStruct*> ColumnTypes)
+	{
+		const void** ColumnsDataIt = ColumnsData.GetData();
+		const UEditorDataStorage& ConstDataStorage = DataStorage;
+		for (const UScriptStruct* ColumnType : ColumnTypes)
+		{
+			*ColumnsDataIt++ = DataStorage.GetColumnData(CurrentRow, ColumnType);
+		}
+	}
+
+	void FQueryContext_DirectSingleApi::GetMutableColumns(TArrayView<void*> ColumnsData, TConstArrayView<const UScriptStruct*> ColumnTypes)
+	{
+		void** ColumnsDataIt = ColumnsData.GetData();
+		for (const UScriptStruct* ColumnType : ColumnTypes)
+		{
+			*ColumnsDataIt++ = DataStorage.GetColumnData(CurrentRow, ColumnType);
+		}
+	}
+
 } // namespace UE::Editor::DataStorage::Queries
\ No newline at end of file
diff --git a/Engine/Plugins/Experimental/EditorDataStorage/Source/TedsCore/Private/Queries/TypedElementQueryContext_DirectSingleApi.h b/Engine/Plugins/Experimental/EditorDataStorage/Source/TedsCore/Private/Queries/TypedElementQueryContext_DirectSingleApi.h
--- a/Engine/Plugins/Experimental/EditorDataStorage/Source/TedsCore/Private/Queries/TypedElementQueryContext_DirectSingleApi.h
+++ b/Engine/Plugins/Experimental/EditorDataStorage/Source/TedsCore/Private/Queries/TypedElementQueryContext_DirectSingleApi.h
@@ -3,22 +3,33 @@
 #pragma once
 
 #include "Elements/Framework/TypedElementQueryContextImplementation.h"
+#include "Elements/Framework/TypedElementQueryFunctions.h"
 
 class UEditorDataStorage;
 
 namespace UE::Editor::DataStorage::Queries
 {
 	/** Query context implementation for single row queries that call into TEDS Direct API. */
-	class FQueryContext_DirectSingleApi final : public SingleRowInfo
+	class FQueryContext_DirectSingleApi final : public ISingleRowInfo<IQueryFunctionResponse>
 	{
 	public:
 		static const FName Name;
 
 		explicit FQueryContext_DirectSingleApi(UEditorDataStorage& DataStorage);
 
+		/**
+		 * @section ISingleRowInfo
+		 */
+
 		virtual bool CurrentRowHasColumns(TConstArrayView<const UScriptStruct*> ColumnTypes) const override;
 		virtual RowHandle GetCurrentRow() const override;
 
+		/**
+		 * @section IQueryFunctionResponse
+		 */
+		virtual void GetConstColumns(TArrayView<const void*> ColumnsData, TConstArrayView<const UScriptStruct*> ColumnTypes) override;
+		virtual void GetMutableColumns(TArrayView<void*> ColumnsData, TConstArrayView<const UScriptStruct*> ColumnTypes) override;
+
 		void SetCurrentRow(RowHandle Row);
 
 	private:
diff --git a/Engine/Plugins/Experimental/EditorDataStorage/Source/TedsCore/Private/TypedElementDatabase.cpp b/Engine/Plugins/Experimental/EditorDataStorage/Source/TedsCore/Private/TypedElementDatabase.cpp
--- a/Engine/Plugins/Experimental/EditorDataStorage/Source/TedsCore/Private/TypedElementDatabase.cpp
+++ b/Engine/Plugins/Experimental/EditorDataStorage/Source/TedsCore/Private/TypedElementDatabase.cpp
@@ -590,7 +590,8 @@ void UEditorDataStorage::FilterRowsBy(UE::Editor::DataStorage::FRowHandleArray&
 		for (RowHandle Row : Input)
 		{
 			Context.GetContextImplementation().SetCurrentRow(Row);
-			if (Filter.Function(Context))
+			
+			if (Filter.Call(Context, Context.GetContextImplementation()))
 			{
 				Result.Add(Row);
 			}
diff --git a/Engine/Source/Runtime/TypedElementFramework/Public/Elements/Framework/TypedElementQueryCapabilityForwarder.h b/Engine/Source/Runtime/TypedElementFramework/Public/Elements/Framework/TypedElementQueryCapabilityForwarder.h
--- a/Engine/Source/Runtime/TypedElementFramework/Public/Elements/Framework/TypedElementQueryCapabilityForwarder.h
+++ b/Engine/Source/Runtime/TypedElementFramework/Public/Elements/Framework/TypedElementQueryCapabilityForwarder.h
@@ -33,41 +33,43 @@ namespace UE::Editor::DataStorage::Queries::Private
 	template<> 
 	struct TForwarder<> : FQueryContext
 	{
-		explicit TForwarder(IContextContract& Contract) : Contract(Contract) {}
+		TForwarder() = default;
+		explicit TForwarder(IContextContract& Contract) : Contract(&Contract) {}
 	
 	protected:
-		IContextContract& Contract;
+		IContextContract* Contract = nullptr;
 	};
 
 #define ArgTypeName(Type, Name) Type Name
 #define ArgName(Type, Name) Name
 
 #define Function0(Capability, Return, Function) Return Function() \
-	{ return this->Contract. Function (); }
+	{ return this->Contract-> Function (); }
 #define Function1(Capability, Return, Function, Arg1) Return Function(ArgTypeName Arg1 ) \
-	{ return this->Contract. Function (ArgName Arg1 ); }
+	{ return this->Contract-> Function (ArgName Arg1 ); }
 #define Function2(Capability, Return, Function, Arg1, Arg2) Return Function(ArgTypeName Arg1 , ArgTypeName Arg2 ) \
-	{ return this->Contract. Function (ArgName Arg1 , ArgName Arg2 ); }
+	{ return this->Contract-> Function (ArgName Arg1 , ArgName Arg2 ); }
 #define Function3(Capability, Return, Function, Arg1, Arg2, Arg3) Return Function(ArgTypeName Arg1 , ArgTypeName Arg2 , ArgTypeName Arg3 ) \
-	{ return this->Contract. Function (ArgName Arg1 , ArgName Arg2 , ArgName Arg3); }
+	{ return this->Contract-> Function (ArgName Arg1 , ArgName Arg2 , ArgName Arg3); }
 #define Function4(Capability, Return, Function, Arg1, Arg2, Arg3, Arg4) Return Function(ArgTypeName Arg1 , ArgTypeName Arg2 , ArgTypeName Arg3 , ArgTypeName Arg4 ) \
-	{ return this->Contract. Function (ArgName Arg1 , ArgName Arg2 , ArgName Arg3 , ArgName Arg4); }
+	{ return this->Contract-> Function (ArgName Arg1 , ArgName Arg2 , ArgName Arg3 , ArgName Arg4); }
 
 #define ConstFunction0(Capability, Return, Function) Return Function() const \
-	{ return this->Contract. Function (); }
+	{ return this->Contract-> Function (); }
 #define ConstFunction1(Capability, Return, Function, Arg1) Return Function(ArgTypeName Arg1 ) const \
-	{ return this->Contract. Function (ArgName Arg1 ); }
+	{ return this->Contract-> Function (ArgName Arg1 ); }
 #define ConstFunction2(Capability, Return, Function, Arg1, Arg2) Return Function(ArgTypeName Arg1 , ArgTypeName Arg2 ) const \
-	{ return this->Contract. Function (ArgName Arg1 , ArgName Arg2 ); }
+	{ return this->Contract-> Function (ArgName Arg1 , ArgName Arg2 ); }
 #define ConstFunction3(Capability, Return, Function, Arg1, Arg2, Arg3) Return Function(ArgTypeName Arg1 , ArgTypeName Arg2 , ArgTypeName Arg3 ) const \
-	{ return this->Contract. Function (ArgName Arg1 , ArgName Arg2 , ArgName Arg3); }
+	{ return this->Contract-> Function (ArgName Arg1 , ArgName Arg2 , ArgName Arg3); }
 #define ConstFunction4(Capability, Return, Function, Arg1, Arg2, Arg3, Arg4) Return Function(ArgTypeName Arg1 , ArgTypeName Arg2 , ArgTypeName Arg3 , ArgTypeName Arg4 ) const \
-	{ return this->Contract. Function (ArgName Arg1 , ArgName Arg2 , ArgName Arg3 , ArgName Arg4); }
+	{ return this->Contract-> Function (ArgName Arg1 , ArgName Arg2 , ArgName Arg3 , ArgName Arg4); }
 
 #define CapabilityStart(Capability) \
 	template<typename... Capabilities> \
 	struct TForwarder< Capability , Capabilities...> : TForwarder<Capabilities...> \
 	{ \
+		TForwarder() = default; \
 		explicit TForwarder(IContextContract& Contract) : TForwarder<Capabilities...>(Contract) {}
 
 #define CapabilityEnd(Capability) \
diff --git a/Engine/Source/Runtime/TypedElementFramework/Public/Elements/Framework/TypedElementQueryContext.h b/Engine/Source/Runtime/TypedElementFramework/Public/Elements/Framework/TypedElementQueryContext.h
--- a/Engine/Source/Runtime/TypedElementFramework/Public/Elements/Framework/TypedElementQueryContext.h
+++ b/Engine/Source/Runtime/TypedElementFramework/Public/Elements/Framework/TypedElementQueryContext.h
@@ -28,6 +28,7 @@ namespace UE::Editor::DataStorage::Queries
 	{
 		using Capabilities = TCapabilityStore<CapabilityTypes...>;
 
+		TQueryContext() = default;
 		TQueryContext(IContextContract& Contract);
 	};
 	
diff --git a/Engine/Source/Runtime/TypedElementFramework/Public/Elements/Framework/TypedElementQueryFunctionArguments.h b/Engine/Source/Runtime/TypedElementFramework/Public/Elements/Framework/TypedElementQueryFunctionArguments.h
new file mode 100644
--- /dev/null
+++ b/Engine/Source/Runtime/TypedElementFramework/Public/Elements/Framework/TypedElementQueryFunctionArguments.h
@@ -0,0 +1,126 @@
+// Copyright Epic Games, Inc. All Rights Reserved.
+
+#pragma once
+
+#include "Containers/ContainersFwd.h"
+#include "DataStorage/CommonTypes.h"
+#include "Elements/Framework/TypedElementQueryCapabilityForwarder.h"
+#include "Misc/EnumClassFlags.h"
+#include "Templates/Tuple.h"
+
+class UScriptStruct;
+
+namespace UE::Editor::DataStorage::Queries::Private
+{
+	template<typename T> concept ContextType = std::is_base_of_v<FQueryContext, T>;
+
+	enum class EArgumentFlags
+	{
+		IsConst = 1 << 0,
+		IsMutable = 1 << 1,
+		BatchArgument = 1 << 2,
+
+		Type_Context = 1 << 3,
+		Type_RowHandle = 1 << 4,
+		Type_Column = 1 << 5,
+	};
+	ENUM_CLASS_FLAGS(EArgumentFlags)
+
+	template<typename T>
+	struct TPointerForwarder
+	{
+		TPointerForwarder() = default;
+		TPointerForwarder(T* Value) : Value(Value) {}
+		TPointerForwarder& operator=(T* InValue) { Value = InValue; return *this; }
+
+		operator T& () const { return *Value; }
+		T* operator->() const { return Value; }
+
+		T* Value = nullptr;
+	};
+	
+	template<typename T>
+	struct TArgument
+	{
+		static_assert(!std::is_same_v<T, void>, "Unsupported argument type for query callback.");
+	};
+
+	// Context type
+	template<typename T> requires ContextType<T>
+	struct TArgument<const T&>
+	{
+		constexpr static EArgumentFlags Flags = EArgumentFlags::IsConst | EArgumentFlags::Type_Context;
+		using ArgumentType = T;
+	};
+
+	template<typename T> requires ContextType<T>
+	struct TArgument<T&>
+	{
+		constexpr static EArgumentFlags Flags = EArgumentFlags::IsMutable | EArgumentFlags::Type_Context;
+		using ArgumentType = T;
+	};
+
+	// Column types
+
+	template<typename T> requires TColumnType<T>
+	struct TArgument<const T&>
+	{
+		constexpr static EArgumentFlags Flags = EArgumentFlags::IsConst | EArgumentFlags::Type_Column;
+		using ArgumentType = TPointerForwarder<const T>;
+		using ColumnType = T;
+		using PointerCastType = const T*;
+	};
+
+	template<typename T> requires TColumnType<T>
+	struct TArgument<T&>
+	{
+		constexpr static EArgumentFlags Flags = EArgumentFlags::IsMutable | EArgumentFlags::Type_Column;
+		using ArgumentType = TPointerForwarder<T>;
+		using ColumnType = T;
+		using PointerCastType = T*;
+	};
+
+	template<int32 Index, typename Tuple>
+	struct TIndexToTypeImpl
+	{
+		using Type = typename TTupleElement<Index, Tuple>::Type;
+	};
+
+	template<typename Tuple>
+	struct TIndexToTypeImpl<-1, Tuple>
+	{
+		using Type = void;
+	};
+	template<int32 Index, typename Tuple>
+	using TIndexToType = typename TIndexToTypeImpl<Index, Tuple>::Type;
+
+	template<typename... Args>
+	struct TArgumentInfo
+	{
+	private:
+		template<EArgumentFlags TypeFlag>
+		constexpr static int32 FirstIndexOfType();
+
+		template<EArgumentFlags Flags>
+		constexpr static int32 CountColumns();
+
+		template<EArgumentFlags Flags>
+		static TConstArrayView<const UScriptStruct*> ListColumns();
+	
+	public:
+		using ArgumentList = TTuple<typename TArgument<Args>::ArgumentType...>;
+
+		constexpr static int32 ContextIndex = FirstIndexOfType<EArgumentFlags::Type_Context>();
+		using ContextType = TIndexToType<ContextIndex, ArgumentList>;
+		static void SetContext(ArgumentList& Arguments, IContextContract& Contract);
+
+		constexpr static int32 CountConstColumns();
+		constexpr static int32 CountMutableColumns();
+		static TConstArrayView<const UScriptStruct*> ListConstColumns();
+		static TConstArrayView<const UScriptStruct*> ListMutableColumns();
+		static void SetConstColumns(ArgumentList& Arguments, TConstArrayView<const void*> Columns);
+		static void SetMutableColumns(ArgumentList& Arguments, TConstArrayView<void*> Columns);
+	};
+} // namespace UE::Editor::DataStorage::Queries::Private
+
+#include "Elements/Framework/TypedElementQueryFunctionArguments.inl"
diff --git a/Engine/Source/Runtime/TypedElementFramework/Public/Elements/Framework/TypedElementQueryFunctionArguments.inl b/Engine/Source/Runtime/TypedElementFramework/Public/Elements/Framework/TypedElementQueryFunctionArguments.inl
new file mode 100644
--- /dev/null
+++ b/Engine/Source/Runtime/TypedElementFramework/Public/Elements/Framework/TypedElementQueryFunctionArguments.inl
@@ -0,0 +1,128 @@
+// Copyright Epic Games, Inc. All Rights Reserved.
+
+#include <functional>
+#include <type_traits>
+
+namespace UE::Editor::DataStorage::Queries::Private
+{
+	template<EArgumentFlags TypeFlag, int32 Index>
+	constexpr int32 FirstIndexOfTypeImpl()
+	{
+		return -1;
+	}
+
+	template<EArgumentFlags TypeFlag, int32 Index, typename Arg, typename... Args>
+	constexpr int32 FirstIndexOfTypeImpl()
+	{
+		if constexpr (EnumHasAnyFlags(TArgument<Arg>::Flags, TypeFlag))
+		{
+			return Index;
+		}
+		else
+		{
+			return FirstIndexOfTypeImpl<TypeFlag, Index + 1, Args...>();
+		}
+	}
+
+	template<typename... Args>
+	template<EArgumentFlags TypeFlag>
+	constexpr int32 TArgumentInfo<Args...>::FirstIndexOfType()
+	{
+		return FirstIndexOfTypeImpl<TypeFlag, 0, Args...>();
+	}
+
+	template<typename... Args>
+	void TArgumentInfo<Args...>::SetContext(ArgumentList& Arguments, IContextContract& Contract)
+	{
+		Arguments.template Get<ContextIndex>() = typename TTupleElement<ContextIndex, ArgumentList>::Type(Contract);
+	}
+
+	template<typename... Args>
+	template<EArgumentFlags Flags>
+	constexpr int32 TArgumentInfo<Args...>::CountColumns()
+	{
+		return (0 + ... + (EnumHasAllFlags(TArgument<Args>::Flags, Flags) ? 1 : 0));
+	}
+
+	template<typename... Args>
+	constexpr int32 TArgumentInfo<Args...>::CountConstColumns()
+	{
+		return CountColumns<EArgumentFlags::IsConst | EArgumentFlags::Type_Column>();
+	}
+			
+	template<typename... Args>
+	constexpr int32 TArgumentInfo<Args...>::CountMutableColumns()
+	{
+		return CountColumns<EArgumentFlags::IsMutable | EArgumentFlags::Type_Column>();
+	}
+
+	template<typename... Args>
+	template<EArgumentFlags Flags>
+	TConstArrayView<const UScriptStruct*> TArgumentInfo<Args...>::ListColumns()
+	{
+		constexpr int32 ColumnCount = CountColumns<Flags>();
+		if constexpr (ColumnCount > 0)
+		{
+			static TConstArrayView<const UScriptStruct*> Result =
+				[&]()
+				{
+					static const UScriptStruct* Columns[ColumnCount];
+					int32 Index = 0;
+					(
+						[&]
+						{
+							if constexpr (EnumHasAllFlags(TArgument<Args>::Flags, Flags))
+							{
+								Columns[Index++] = TArgument<Args>::ColumnType::StaticStruct();
+							}
+						}(), ...
+					);
+
+					return TConstArrayView<const UScriptStruct*>(Columns, ColumnCount);
+				}();
+			return Result;
+		}
+		else
+		{
+			return TConstArrayView<const UScriptStruct*>();
+		}
+	}
+
+	template<typename... Args>
+	TConstArrayView<const UScriptStruct*> TArgumentInfo<Args...>::ListConstColumns()
+	{
+		return ListColumns<EArgumentFlags::IsConst | EArgumentFlags::Type_Column>();
+	}
+
+	template<typename... Args>
+	TConstArrayView<const UScriptStruct*> TArgumentInfo<Args...>::ListMutableColumns()
+	{
+		return ListColumns<EArgumentFlags::IsMutable | EArgumentFlags::Type_Column>();
+	}
+
+	template<typename... Args>
+	void TArgumentInfo<Args...>::SetConstColumns(ArgumentList& Arguments, TConstArrayView<const void*> Columns)
+	{
+		const void* const* ColumnIt = Columns.GetData();
+		([&]{
+				if constexpr (EnumHasAllFlags(TArgument<Args>::Flags, EArgumentFlags::IsConst | EArgumentFlags::Type_Column))
+				{
+					using ArgumentType = TArgument<Args>;
+					Arguments.template Get<typename ArgumentType::ArgumentType>() = static_cast<typename ArgumentType::PointerCastType>(*ColumnIt++);
+				}
+			}(), ...);
+	}
+
+	template<typename... Args>
+	void TArgumentInfo<Args...>::SetMutableColumns(ArgumentList& Arguments, TConstArrayView<void*> Columns)
+	{
+		void* const* ColumnIt = Columns.GetData();
+		([&]{
+				if constexpr (EnumHasAllFlags(TArgument<Args>::Flags, EArgumentFlags::IsMutable | EArgumentFlags::Type_Column))
+				{
+					using ArgumentType = TArgument<Args>;
+					Arguments.template Get<typename ArgumentType::ArgumentType>() = static_cast<typename ArgumentType::PointerCastType>(*ColumnIt++);
+				}
+			}(), ...);
+	}
+} // namespace UE::Editor::DataStorage::Queries::Private
diff --git a/Engine/Source/Runtime/TypedElementFramework/Public/Elements/Framework/TypedElementQueryFunctions.h b/Engine/Source/Runtime/TypedElementFramework/Public/Elements/Framework/TypedElementQueryFunctions.h
--- a/Engine/Source/Runtime/TypedElementFramework/Public/Elements/Framework/TypedElementQueryFunctions.h
+++ b/Engine/Source/Runtime/TypedElementFramework/Public/Elements/Framework/TypedElementQueryFunctions.h
@@ -7,16 +7,30 @@
 #include "Templates/Function.h"
 #include "UObject/NameTypes.h"
 
+class UScriptStruct;
+
 namespace UE::Editor::DataStorage::Queries
 {
 	struct IContextContract;
 
-	/** Storage for a function that used as part of a query.  */
+	struct IQueryFunctionResponse
+	{
+		virtual ~IQueryFunctionResponse() = default;
+		
+		virtual void GetConstColumns(TArrayView<const void*> ColumnsData, TConstArrayView<const UScriptStruct*> ColumnTypes) = 0;
+		virtual void GetMutableColumns(TArrayView<void*> ColumnsData, TConstArrayView<const UScriptStruct*> ColumnTypes) = 0;
+	};
+
+	/** Storage for a function that can be used as part of a query.  */
 	template<typename ReturnType>
 	struct TQueryFunction
 	{
 		TConstArrayView<FName> Capabilities;
-		TFunction<ReturnType(IContextContract& Contract)> Function;
+		TConstArrayView<const UScriptStruct*> ConstColumnTypes;
+		TConstArrayView<const UScriptStruct*> MutableColumnTypes;
+		TFunction<ReturnType(IContextContract& Contract, TConstArrayView<const void*> ConstColumns, TConstArrayView<void*> MutableColumns)> Function;
+
+		ReturnType Call(IContextContract& Contract, IQueryFunctionResponse& Response);
 	};
 
 	template<typename ReturnType, typename FunctionType>
diff --git a/Engine/Source/Runtime/TypedElementFramework/Public/Elements/Framework/TypedElementQueryFunctions.inl b/Engine/Source/Runtime/TypedElementFramework/Public/Elements/Framework/TypedElementQueryFunctions.inl
--- a/Engine/Source/Runtime/TypedElementFramework/Public/Elements/Framework/TypedElementQueryFunctions.inl
+++ b/Engine/Source/Runtime/TypedElementFramework/Public/Elements/Framework/TypedElementQueryFunctions.inl
@@ -1,55 +1,54 @@
 // Copyright Epic Games, Inc. All Rights Reserved.
 
+#include <functional>
 #include <type_traits>
+#include "DataStorage/CommonTypes.h"
+#include "DataStorage/Handles.h"
 #include "Elements/Framework/TypedElementQueryCapabilities.h"
 #include "Elements/Framework/TypedElementQueryCapabilityForwarder.h"
 #include "Elements/Framework/TypedElementQueryContract.h"
+#include "Elements/Framework/TypedElementQueryFunctionArguments.h"
 
 namespace UE::Editor::DataStorage::Queries
 {
 	namespace Private
 	{
 		template<typename T> concept FunctorType = requires{ &T::operator(); };
-		template<typename T> concept ContextType = std::is_base_of_v<FQueryContext, T>;
-	
+		
 		template<typename T>
 		struct TFunctionInfoImpl
 		{
+			using ArgumentInfo = TArgumentInfo<>;
+
 			static constexpr bool bIsValidQueryFunction = false;
-			using ContextType = void;
 			using ReturnType = void;
-			static constexpr int32 CapabilityCount = 0;
-			inline static const TConstArrayView<FName> Capabilities;
 		};
 
-		template<typename Return, ContextType Context, typename... Args>
-		struct TFunctionInfoImpl<Return (*)(Context&, Args...)>
+		template<typename Return, typename... Args>
+		struct TFunctionInfoImpl<Return(*)(Args...)>
 		{
+			using ArgumentInfo = TArgumentInfo<Args...>;
+
 			static constexpr bool bIsValidQueryFunction = true;
-			using ContextType = Context;
-			using ReturnType = Return;
-			static constexpr int32 CapabilityCount = IContextContract::CountSupportedCapabilities<typename Context::Capabilities>();
-			inline static const TConstArrayView<FName> Capabilities = IContextContract::SupportedCapabilitiesList<typename Context::Capabilities>();
+			using ReturnType = Return;			
 		};
 
-		template<typename Class, typename Return, ContextType Context, typename... Args>
-		struct TFunctionInfoImpl<Return(Class::*)(Context&, Args...)>
+		template<typename Class, typename Return, typename... Args>
+		struct TFunctionInfoImpl<Return(Class::*)(Args...)>
 		{
+			using ArgumentInfo = TArgumentInfo<Args...>;
+
 			static constexpr bool bIsValidQueryFunction = true;
-			using ContextType = Context;
 			using ReturnType = Return;
-			static constexpr int32 CapabilityCount = IContextContract::CountSupportedCapabilities<typename Context::Capabilities>();
-			inline static const TConstArrayView<FName> Capabilities = IContextContract::SupportedCapabilitiesList<typename Context::Capabilities>();
 		};
 
-		template<typename Class, typename Return, ContextType Context, typename... Args>
-		struct TFunctionInfoImpl<Return(Class::*)(Context&, Args...) const>
+		template<typename Class, typename Return, typename... Args>
+		struct TFunctionInfoImpl<Return(Class::*)(Args...) const>
 		{
+			using ArgumentInfo = TArgumentInfo<Args...>;
+
 			static constexpr bool bIsValidQueryFunction = true;
-			using ContextType = Context;
 			using ReturnType = Return;
-			static constexpr int32 CapabilityCount = IContextContract::CountSupportedCapabilities<typename Context::Capabilities>();
-			inline static const TConstArrayView<FName> Capabilities = IContextContract::SupportedCapabilitiesList<typename Context::Capabilities>();
 		};
 
 		template<bool IsFunctor, typename FunctionType>
@@ -68,13 +67,49 @@ namespace UE::Editor::DataStorage::Queries
 		};
 		
 		template<typename FunctionType>
-		using TFunctionInfo = TFunctionInfoSelection<FunctorType<FunctionType>, FunctionType>::Type;
+		struct TFunctionInfo
+		{
+			using Base = TFunctionInfoSelection<FunctorType<FunctionType>, FunctionType>::Type;
+			
+			static constexpr bool bIsValidQueryFunction = Base::bIsValidQueryFunction;
+			using ReturnType = typename Base::ReturnType;
+
+			using ArgumentInfo = typename Base::ArgumentInfo;
+			using ContextType = typename ArgumentInfo::ContextType;
+		};
+	}
+
+	template<typename ReturnType>
+	ReturnType TQueryFunction<ReturnType>::Call(IContextContract& Contract, IQueryFunctionResponse& Response)
+	{
+		const void** ConstColumnsData = static_cast<const void**>(alloca(sizeof(void*) * ConstColumnTypes.Num()));
+		void** MutableColumnsData = static_cast<void**>(alloca(sizeof(void*) * MutableColumnTypes.Num()));
+
+		Response.GetConstColumns(TArrayView<const void*>(ConstColumnsData, ConstColumnTypes.Num()), ConstColumnTypes);
+		Response.GetMutableColumns(TArrayView<void*>(MutableColumnsData, MutableColumnTypes.Num()), MutableColumnTypes);
+
+		if constexpr (std::is_same_v<ReturnType, void>)
+		{
+			Function(
+				Contract,
+				TConstArrayView<const void*>(ConstColumnsData, ConstColumnTypes.Num()),
+				TConstArrayView<void*>(MutableColumnsData, MutableColumnTypes.Num()));
+		}
+		else
+		{
+			return Function(
+				Contract,
+				TConstArrayView<const void*>(ConstColumnsData, ConstColumnTypes.Num()),
+				TConstArrayView<void*>(MutableColumnsData, MutableColumnTypes.Num()));
+		}
 	}
 
 	template<typename ReturnType, typename FunctionType>
 	TQueryFunction<ReturnType> BuildQueryFunction(FunctionType&& Function)
 	{
 		using FunctionInfo = Private::TFunctionInfo<FunctionType>;
+		using ArgumentInfo = typename FunctionInfo::ArgumentInfo;
+		using ContextType = typename FunctionInfo::ContextType;
 
 		static_assert(FunctionInfo::bIsValidQueryFunction, "The provided function is not compatible with a query callback. "
 			"Likely errors are missing a query context or having incompatible arguments.");
@@ -82,18 +117,33 @@ namespace UE::Editor::DataStorage::Queries
 			"The result of the provided function is not compatible with the query function that's being constructed.");
 
 		TQueryFunction<ReturnType> Result;
-		Result.Capabilities = FunctionInfo::Capabilities;
+		if constexpr (!std::is_same_v<void, ContextType>)
+		{
+			Result.Capabilities = IContextContract::SupportedCapabilitiesList<typename ContextType::Capabilities>();
+		}
+		Result.ConstColumnTypes = ArgumentInfo::ListConstColumns();
+		Result.MutableColumnTypes = ArgumentInfo::ListMutableColumns();
 		Result.Function =
-			[Callback = Forward<FunctionType>(Function)](IContextContract& Contract)
+			[Callback = Forward<FunctionType>(Function)]
+			(IContextContract& Contract, TConstArrayView<const void*> ConstColumns, TConstArrayView<void*> MutableColumns)
 			{
-				typename FunctionInfo::ContextType Context(Contract);
+				using ArgumentList = typename ArgumentInfo::ArgumentList;
+
+				ArgumentList Arguments;
+				if constexpr (!std::is_same_v<void, ContextType>)
+				{
+					ArgumentInfo::SetContext(Arguments, Contract);
+				}
+				ArgumentInfo::SetConstColumns(Arguments, ConstColumns);
+				ArgumentInfo::SetMutableColumns(Arguments, MutableColumns);
+
 				if constexpr (std::is_same_v<ReturnType, void>)
 				{
-					Callback(Context);
+					Arguments.ApplyBefore(Callback);
 				}
 				else
 				{
-					return Callback(Context);
+					return Arguments.ApplyBefore(Callback);
 				}
 			};
 		return Result;
diff --git a/Engine/Source/Runtime/TypedElementFramework/Tests/Elements/Framework/TypedElementQueryContextMock.h b/Engine/Source/Runtime/TypedElementFramework/Tests/Elements/Framework/TypedElementQueryContextMock.h
new file mode 100644
--- /dev/null
+++ b/Engine/Source/Runtime/TypedElementFramework/Tests/Elements/Framework/TypedElementQueryContextMock.h
@@ -0,0 +1,176 @@
+// Copyright Epic Games, Inc. All Rights Reserved.
+
+#pragma once
+
+#include "Elements/Framework/TypedElementQueryCapabilities.h"
+#include "Elements/Framework/TypedElementQueryContextImplementation.h"
+#include "Templates/FunctionFwd.h"
+
+namespace UE::Editor::DataStorage::Queries
+{
+	struct FMockChainTerminator {};
+
+	struct FQueryContextMock final : public
+
+#define Function0(Capability, Return, Function)
+#define Function1(Capability, Return, Function, Arg1)
+#define Function2(Capability, Return, Function, Arg1, Arg2)
+#define Function3(Capability, Return, Function, Arg1, Arg2, Arg3)
+#define Function4(Capability, Return, Function, Arg1, Arg2, Arg3, Arg4)
+
+#define ConstFunction0(Capability, Return, Function)
+#define ConstFunction1(Capability, Return, Function, Arg1)
+#define ConstFunction2(Capability, Return, Function, Arg1, Arg2)
+#define ConstFunction3(Capability, Return, Function, Arg1, Arg2, Arg3)
+#define ConstFunction4(Capability, Return, Function, Arg1, Arg2, Arg3, Arg4)
+
+#define CapabilityStart(Capability) Capability, 
+#define CapabilityEnd(Capability)
+
+#include "Elements/Framework/TypedElementQueryCapabilities.inl"
+
+#undef Function0
+#undef Function1
+#undef Function2
+#undef Function3
+#undef Function4
+#undef ConstFunction0
+#undef ConstFunction1
+#undef ConstFunction2
+#undef ConstFunction3
+#undef ConstFunction4
+#undef CapabilityStart
+#undef CapabilityEnd
+
+		FMockChainTerminator
+	{
+
+#define WithWrappers 1
+
+#define ArgTypeName(Type, Name) Type Name
+#define ArgName(Type, Name) Name
+#define MockFunction(Capability, Function) Capability##_##Function##_Mock
+
+#define Function0(Capability, Return, Function) \
+		TFunction<Return ()> MockFunction(Capability, Function); \
+		virtual Return Function() override \
+		{ \
+			return MockFunction(Capability, Function)(); \
+		}
+#define Function1(Capability, Return, Function, Arg1) \
+		TFunction<Return (ArgTypeName Arg1 )> MockFunction(Capability, Function); \
+		virtual Return Function(ArgTypeName Arg1) override \
+		{ \
+			return MockFunction(Capability, Function)(ArgName Arg1); \
+		}
+#define Function2(Capability, Return, Function, Arg1, Arg2) \
+		TFunction<Return (ArgTypeName Arg1, ArgTypeName Arg2)> MockFunction(Capability, Function); \
+		virtual Return Function(ArgTypeName Arg1, ArgTypeName Arg2) override \
+		{ \
+			return MockFunction(Capability, Function)(ArgName Arg1, ArgName Arg2); \
+		}
+#define Function3(Capability, Return, Function, Arg1, Arg2, Arg3) \
+		TFunction<Return (ArgTypeName Arg1, ArgTypeName Arg2, ArgTypeName Arg3)> MockFunction(Capability, Function); \
+		virtual Return Function(ArgTypeName Arg1, ArgTypeName Arg2, ArgTypeName Arg3) override \
+		{ \
+			return MockFunction(Capability, Function)(ArgName Arg1, ArgName Arg2, ArgName Arg3); \
+		}
+#define Function4(Capability, Return, Function, Arg1, Arg2, Arg3, Arg4) \
+		TFunction<Return (ArgTypeName Arg1, ArgTypeName Arg2, ArgTypeName Arg3, ArgTypeName Arg4)> MockFunction(Capability, Function); \
+		virtual Return Function(ArgTypeName Arg1, ArgTypeName Arg2, ArgTypeName Arg3, ArgTypeName Arg4) override \
+		{ \
+			return MockFunction(Capability, Function)(ArgName Arg1, ArgName Arg2, ArgName Arg3, ArgName Arg4); \
+		}
+
+#define ConstFunction0(Capability, Return, Function) \
+		TFunction<Return ()> MockFunction(Capability, Function); \
+		virtual Return Function() const override \
+		{ \
+			return MockFunction(Capability, Function)(); \
+		}
+#define ConstFunction1(Capability, Return, Function, Arg1) \
+		TFunction<Return (ArgTypeName Arg1)> MockFunction(Capability, Function); \
+		virtual Return Function(ArgTypeName Arg1) const override \
+		{ \
+			return MockFunction(Capability, Function)(ArgName Arg1); \
+		}
+#define ConstFunction2(Capability, Return, Function, Arg1, Arg2) \
+		TFunction<Return (ArgTypeName Arg1, ArgTypeName Arg2)> MockFunction(Capability, Function); \
+		virtual Return Function(ArgTypeName Arg1, ArgTypeName Arg2) const override \
+		{ \
+			return MockFunction(Capability, Function)(ArgName Arg1, ArgName Arg2); \
+		}
+#define ConstFunction3(Capability, Return, Function, Arg1, Arg2, Arg3) \
+		TFunction<Return (ArgTypeName Arg1, ArgTypeName Arg2, ArgTypeName Arg3)> MockFunction(Capability, Function); \
+		virtual Return Function(ArgTypeName Arg1, ArgTypeName Arg2, ArgTypeName Arg3) const override \
+		{ \
+			return MockFunction(Capability, Function)(ArgName Arg1, ArgName Arg2, ArgName Arg3); \
+		}
+#define ConstFunction4(Capability, Return, Function, Arg1, Arg2, Arg3, Arg4) \
+		TFunction<Return (ArgTypeName Arg1, ArgTypeName Arg2, ArgTypeName Arg3, ArgTypeName Arg4)> MockFunction(Capability, Function); \
+		virtual Return Function(ArgTypeName Arg1, ArgTypeName Arg2, ArgTypeName Arg3, ArgTypeName Arg4) const override \
+		{ \
+			return MockFunction(Capability, Function)(ArgName Arg1, ArgName Arg2, ArgName Arg3, ArgName Arg4); \
+		}
+
+#define CapabilityStart(Capability) 
+#define CapabilityEnd(Capability)
+
+#include "Elements/Framework/TypedElementQueryCapabilities.inl"
+
+#undef Function0
+#undef Function1
+#undef Function2
+#undef Function3
+#undef Function4
+#undef ConstFunction0
+#undef ConstFunction1
+#undef ConstFunction2
+#undef ConstFunction3
+#undef ConstFunction4
+#undef WithWrappers
+
+#define MockFunctionImplementation(Capability, Return, Function) \
+		checkf(false, TEXT( "Context function " #Function " in capability " #Capability " was unexpectedly called." )); return Return{};
+#define Function0(Capability, Return, Function) \
+		MockFunction(Capability, Function) = []() { MockFunctionImplementation(Capability, Return, Function) };
+#define Function1(Capability, Return, Function, Arg1) \
+		MockFunction(Capability, Function) = [](ArgTypeName Arg1) { MockFunctionImplementation(Capability, Return, Function) };
+#define Function2(Capability, Return, Function, Arg1, Arg2) \
+		MockFunction(Capability, Function) = [](ArgTypeName Arg1, ArgTypeName Arg2) { MockFunctionImplementation(Capability, Return, Function) };
+#define Function3(Capability, Return, Function, Arg1, Arg2, Arg3) \
+		MockFunction(Capability, Function) = [](ArgTypeName Arg1, ArgTypeName Arg2, ArgTypeName Arg3) { MockFunctionImplementation(Capability, Return, Function) };
+#define Function4(Capability, Return, Function, Arg1, Arg2, Arg3, Arg4) \
+		MockFunction(Capability, Function) = [](ArgTypeName Arg1, ArgTypeName Arg2, ArgTypeName Arg3, ArgTypeName Arg4) { MockFunctionImplementation(Capability, Return, Function) };
+
+#define ConstFunction0(Capability, Return, Function) Function0(Capability, Return, Function)
+#define ConstFunction1(Capability, Return, Function, Arg1) Function1(Capability, Return, Function, Arg1)
+#define ConstFunction2(Capability, Return, Function, Arg1, Arg2) Function2(Capability, Return, Function, Arg1, Arg2)
+#define ConstFunction3(Capability, Return, Function, Arg1, Arg2, Arg3) Function3(Capability, Return, Function, Arg1, Arg2, Arg3)
+#define ConstFunction4(Capability, Return, Function, Arg1, Arg2, Arg3, Arg4) Function4(Capability, Return, Function, Arg1, Arg2, Arg3, Arg4)
+
+		FQueryContextMock()
+		{
+#include "Elements/Framework/TypedElementQueryCapabilities.inl"
+		}
+
+#undef ArgTypeName
+#undef ArgName
+#undef MockFunctionName
+#undef Function0
+#undef Function1
+#undef Function2
+#undef Function3
+#undef Function4
+#undef ConstFunction0
+#undef ConstFunction1
+#undef ConstFunction2
+#undef ConstFunction3
+#undef ConstFunction4
+#undef CapabilityStart
+#undef CapabilityEnd
+#undef WithWrappers
+	};
+
+	using QueryContextMock = TQueryContextImpl<FQueryContextMock>;
+} // namespace UE::Editor::DataStorage::Queries
diff --git a/Engine/Source/Runtime/TypedElementFramework/Tests/Elements/Framework/TypedElementQueryFunctionsTests.cpp b/Engine/Source/Runtime/TypedElementFramework/Tests/Elements/Framework/TypedElementQueryFunctionsTests.cpp
new file mode 100644
--- /dev/null
+++ b/Engine/Source/Runtime/TypedElementFramework/Tests/Elements/Framework/TypedElementQueryFunctionsTests.cpp
@@ -0,0 +1,180 @@
+// Copyright Epic Games, Inc. All Rights Reserved.
+
+#if WITH_TESTS
+
+#include "Elements/Framework/TypedElementQueryContext.h"
+#include "Elements/Framework/TypedElementQueryContextMock.h"
+#include "Elements/Framework/TypedElementQueryFunctions.h"
+#include "Elements/Framework/TypedElementTestColumns.h"
+#include "Tests/TestHarnessAdapter.h"
+
+namespace UE::Editor::DataStorage::Queries::Tests
+{
+	struct FTestQueryFunctionResponse : IQueryFunctionResponse
+	{
+		virtual ~FTestQueryFunctionResponse() override = default;
+
+		virtual void GetConstColumns(TArrayView<const void*> ColumnsData, TConstArrayView<const UScriptStruct*> ColumnTypes) override
+		{
+			CHECK_MESSAGE(TEXT("Attempting to get more or less const columns than are available in the mock."), ConstColumns.Num() == ColumnTypes.Num());
+
+			for (int32 Index = 0, Num = ColumnTypes.Num(); Index < Num; ++Index)
+			{
+				CHECK_MESSAGE(TEXT("Type for requested const column doesn't match the stored type in the mock."),
+					ConstColumnTypes[Index] == ColumnTypes[Index]);
+				ColumnsData[Index] = ConstColumns[Index];
+			}
+		}
+
+		virtual void GetMutableColumns(TArrayView<void*> ColumnsData, TConstArrayView<const UScriptStruct*> ColumnTypes) override
+		{
+			CHECK_MESSAGE(TEXT("Attempting to get more or less mutable columns than are available in the mock."), MutableColumns.Num() == ColumnTypes.Num());
+
+			for (int32 Index = 0, Num = ColumnTypes.Num(); Index < Num; ++Index)
+			{
+				CHECK_MESSAGE(TEXT("Type for requested mutable column doesn't match the stored type in the mock."),
+					MutableColumnTypes[Index] == ColumnTypes[Index]);
+				ColumnsData[Index] = MutableColumns[Index];
+			}
+		}
+
+		template<TColumnType... ColumnType>
+		void SetConstColumns(const ColumnType*... Columns)
+		{
+			(ConstColumns.Add(Columns), ...);
+			(ConstColumnTypes.Add(Columns->StaticStruct()), ...);
+		}
+
+		template<TColumnType... ColumnType>
+		void SetMutableColumns(ColumnType*... Columns)
+		{
+			(MutableColumns.Add(Columns), ...);
+			(MutableColumnTypes.Add(Columns->StaticStruct()), ...);
+		}
+
+		TArray<const void*> ConstColumns;
+		TArray<const UScriptStruct*> ConstColumnTypes;
+		TArray<void*> MutableColumns;
+		TArray<const UScriptStruct*> MutableColumnTypes;
+	};
+
+	TEST_CASE_NAMED(TQueryFunction_Tests, "Editor::DataStorage::Queries::Query Function (TQueryFunction)", "[ApplicationContextMask][EngineFilter]")
+	{
+		using namespace UE::Editor::DataStorage;
+		using namespace UE::Editor::DataStorage::Queries;
+		
+		SECTION("Empty function (void)")
+		{
+			TQueryFunction<void> Result = BuildQueryFunction<void>([](){});
+			
+			CHECK_MESSAGE(TEXT("Empty query still has capabilities."), Result.Capabilities.IsEmpty());
+			CHECK_MESSAGE(TEXT("Empty query still has const columns."), Result.ConstColumnTypes.IsEmpty());
+			CHECK_MESSAGE(TEXT("Empty query still has mutable columns."), Result.MutableColumnTypes.IsEmpty());
+		}
+
+		SECTION("Empty function (bool)")
+		{
+			TQueryFunction<bool> Result = BuildQueryFunction<bool>([]() { return false; });
+
+			CHECK_MESSAGE(TEXT("Empty query still has capabilities."), Result.Capabilities.IsEmpty());
+			CHECK_MESSAGE(TEXT("Empty query still has const columns."), Result.ConstColumnTypes.IsEmpty());
+			CHECK_MESSAGE(TEXT("Empty query still has mutable columns."), Result.MutableColumnTypes.IsEmpty());
+		}
+
+		SECTION("Empty function with capture")
+		{
+			bool bResult = false;
+			TQueryFunction<bool> Result = BuildQueryFunction<bool>([&bResult]() { return bResult; });
+
+			CHECK_MESSAGE(TEXT("Empty query still has capabilities."), Result.Capabilities.IsEmpty());
+			CHECK_MESSAGE(TEXT("Empty query still has const columns."), Result.ConstColumnTypes.IsEmpty());
+			CHECK_MESSAGE(TEXT("Empty query still has mutable columns."), Result.MutableColumnTypes.IsEmpty());
+		}
+
+		SECTION("Function with context")
+		{
+			TQueryFunction<void> Result = BuildQueryFunction<void>([](TQueryContext<SingleRowInfo>& Context) {});
+
+			CHECK_MESSAGE(TEXT("Expected exactly one capability."), Result.Capabilities.Num() == 1);
+			CHECK_MESSAGE(TEXT("Query with only context shouldn't have const columns."), Result.ConstColumnTypes.IsEmpty());
+			CHECK_MESSAGE(TEXT("Query with only context shouldn't have mutable columns."), Result.MutableColumnTypes.IsEmpty());
+		}
+
+		SECTION("Function with columns")
+		{
+			TQueryFunction<void> Result = BuildQueryFunction<void>(
+				[](FTestColumnInt& ColumnA, const FTestColumnString& ColumnB)
+				{
+				});
+
+			CHECK_MESSAGE(TEXT("Query with only columns shouldn't have capabilities."), Result.Capabilities.IsEmpty());
+			CHECK_MESSAGE(TEXT("Expected exactly one const column."), Result.ConstColumnTypes.Num() == 1);
+			CHECK_MESSAGE(TEXT("Expected exactly one mutable column."), Result.MutableColumnTypes.Num() == 1);
+		}
+
+		SECTION("Function with context and columns")
+		{
+			TQueryFunction<void> Result = BuildQueryFunction<void>(
+				[](TQueryContext<SingleRowInfo>& Context, FTestColumnInt& ColumnA, const FTestColumnString& ColumnB)
+				{
+				});
+
+			CHECK_MESSAGE(TEXT("Expected exactly one capability."), Result.Capabilities.Num() == 1);
+			CHECK_MESSAGE(TEXT("Expected exactly one const column."), Result.ConstColumnTypes.Num() == 1);
+			CHECK_MESSAGE(TEXT("Expected exactly one mutable column."), Result.MutableColumnTypes.Num() == 1);
+
+			CHECK_MESSAGE(TEXT("Incorrect const column set."), Result.ConstColumnTypes[0] == FTestColumnString::StaticStruct());
+			CHECK_MESSAGE(TEXT("Incorrect mutable column set."), Result.MutableColumnTypes[0] == FTestColumnInt::StaticStruct());
+		}
+
+		SECTION("Function with context and columns at random location")
+		{
+			TQueryFunction<void> Result = BuildQueryFunction<void>(
+				[](FTestColumnInt& ColumnA, TQueryContext<SingleRowInfo>& Context, const FTestColumnString& ColumnB)
+				{
+				});
+
+			CHECK_MESSAGE(TEXT("Expected exactly one capability."), Result.Capabilities.Num() == 1);
+			CHECK_MESSAGE(TEXT("Expected exactly one const column."), Result.ConstColumnTypes.Num() == 1);
+			CHECK_MESSAGE(TEXT("Expected exactly one mutable column."), Result.MutableColumnTypes.Num() == 1);
+
+			CHECK_MESSAGE(TEXT("Incorrect const column set."), Result.ConstColumnTypes[0] == FTestColumnString::StaticStruct());
+			CHECK_MESSAGE(TEXT("Incorrect mutable column set."), Result.MutableColumnTypes[0] == FTestColumnInt::StaticStruct());
+		}
+
+		SECTION("Call without arguments")
+		{
+			TQueryFunction<bool> Function = BuildQueryFunction<bool>(
+				[]()
+				{
+					return true;
+				});
+			
+			bool bResult = false;
+			QueryContextMock Context;
+			FTestQueryFunctionResponse Response;
+			bResult = Function.Call(Context, Response);
+			CHECK_MESSAGE(TEXT("Query function wasn't called."), bResult);
+		}
+
+		SECTION("Call setting column value")
+		{
+			TQueryFunction<void> Function = BuildQueryFunction<void>(
+				[](FTestColumnString& Column)
+				{
+					Column.TestString = TEXT("Callback");
+				});
+
+			QueryContextMock Context;
+			FTestColumnString Column;
+			
+			FTestQueryFunctionResponse Response;
+			Response.SetMutableColumns(&Column);
+			Function.Call(Context, Response);
+			
+			CHECK_MESSAGE(TEXT("Query function wasn't called."), Column.TestString == TEXT("Callback"));
+		}
+	}
+} // namespace UE::Editor::DataStorage::Queries::Tests
+
+#endif // WITH_TEST
\ No newline at end of file
