Fixed compile error caused by UE_INLINE_GENERATED_CPP_BY_NAME being wrong after file rename

diff --git a/Engine/Source/Runtime/TypedElementFramework/Private/DataStorage/CommonTypes.cpp b/Engine/Source/Runtime/TypedElementFramework/Private/DataStorage/CommonTypes.cpp
--- a/Engine/Source/Runtime/TypedElementFramework/Private/DataStorage/CommonTypes.cpp
+++ b/Engine/Source/Runtime/TypedElementFramework/Private/DataStorage/CommonTypes.cpp
@@ -2,7 +2,7 @@
 
 #include "DataStorage/CommonTypes.h"
 
-#include UE_INLINE_GENERATED_CPP_BY_NAME(TypedElementCommonTypes)
+#include UE_INLINE_GENERATED_CPP_BY_NAME(CommonTypes)
 
 namespace UE
 {
