[TEDS] Workaround for the most recent version of MSVC not processing "if constexpr" correctly before evaluating the templates in the branch.

#rnx

diff --git a/Engine/Source/Runtime/TypedElementFramework/Public/Elements/Framework/TypedElementQueryFunctionArguments.h b/Engine/Source/Runtime/TypedElementFramework/Public/Elements/Framework/TypedElementQueryFunctionArguments.h
--- a/Engine/Source/Runtime/TypedElementFramework/Public/Elements/Framework/TypedElementQueryFunctionArguments.h
+++ b/Engine/Source/Runtime/TypedElementFramework/Public/Elements/Framework/TypedElementQueryFunctionArguments.h
@@ -51,6 +51,8 @@ namespace UE::Editor::DataStorage::Queries::Private
 	{
 		constexpr static EArgumentFlags Flags = EArgumentFlags::IsConst | EArgumentFlags::Type_Context;
 		using ArgumentType = T;
+		using ColumnType = void;
+		using PointerCastType = void*;
 	};
 
 	template<typename T> requires ContextType<T>
@@ -58,6 +60,8 @@ namespace UE::Editor::DataStorage::Queries::Private
 	{
 		constexpr static EArgumentFlags Flags = EArgumentFlags::IsMutable | EArgumentFlags::Type_Context;
 		using ArgumentType = T;
+		using ColumnType = void;
+		using PointerCastType = void*;
 	};
 
 	// Column types
diff --git a/Engine/Source/Runtime/TypedElementFramework/Public/Elements/Framework/TypedElementQueryFunctionArguments.inl b/Engine/Source/Runtime/TypedElementFramework/Public/Elements/Framework/TypedElementQueryFunctionArguments.inl
--- a/Engine/Source/Runtime/TypedElementFramework/Public/Elements/Framework/TypedElementQueryFunctionArguments.inl
+++ b/Engine/Source/Runtime/TypedElementFramework/Public/Elements/Framework/TypedElementQueryFunctionArguments.inl
@@ -105,9 +105,9 @@ namespace UE::Editor::DataStorage::Queries::Private
 	{
 		const void* const* ColumnIt = Columns.GetData();
 		([&]{
-				if constexpr (EnumHasAllFlags(TArgument<Args>::Flags, EArgumentFlags::IsConst | EArgumentFlags::Type_Column))
+				using ArgumentType = TArgument<Args>;
+				if constexpr (EnumHasAllFlags(ArgumentType::Flags, EArgumentFlags::IsConst | EArgumentFlags::Type_Column))
 				{
-					using ArgumentType = TArgument<Args>;
 					Arguments.template Get<typename ArgumentType::ArgumentType>() = static_cast<typename ArgumentType::PointerCastType>(*ColumnIt++);
 				}
 			}(), ...);
@@ -118,9 +118,9 @@ namespace UE::Editor::DataStorage::Queries::Private
 	{
 		void* const* ColumnIt = Columns.GetData();
 		([&]{
-				if constexpr (EnumHasAllFlags(TArgument<Args>::Flags, EArgumentFlags::IsMutable | EArgumentFlags::Type_Column))
+				using ArgumentType = TArgument<Args>;
+				if constexpr (EnumHasAllFlags(ArgumentType::Flags, EArgumentFlags::IsMutable | EArgumentFlags::Type_Column))
 				{
-					using ArgumentType = TArgument<Args>;
 					Arguments.template Get<typename ArgumentType::ArgumentType>() = static_cast<typename ArgumentType::PointerCastType>(*ColumnIt++);
 				}
 			}(), ...);
diff --git a/Engine/Source/Runtime/TypedElementFramework/Tests/Elements/Framework/TypedElementQueryFunctionsTests.cpp b/Engine/Source/Runtime/TypedElementFramework/Tests/Elements/Framework/TypedElementQueryFunctionsTests.cpp
--- a/Engine/Source/Runtime/TypedElementFramework/Tests/Elements/Framework/TypedElementQueryFunctionsTests.cpp
+++ b/Engine/Source/Runtime/TypedElementFramework/Tests/Elements/Framework/TypedElementQueryFunctionsTests.cpp
@@ -174,6 +174,27 @@ namespace UE::Editor::DataStorage::Queries::Tests
 			
 			CHECK_MESSAGE(TEXT("Query function wasn't called."), Column.TestString == TEXT("Callback"));
 		}
+
+		SECTION("Call setting const and mutable columns")
+		{
+			TQueryFunction<void> Function = BuildQueryFunction<void>(
+				[](FTestColumnString& ColumnA, const FTestColumnInt& ColumnB)
+				{
+					ColumnA.TestString = FString::Printf(TEXT("Test: %i"), ColumnB.TestInt);
+				});
+
+			QueryContextMock Context;
+			FTestColumnString ColumnA;
+			FTestColumnInt ColumnB;
+			ColumnB.TestInt = 42;
+
+			FTestQueryFunctionResponse Response;
+			Response.SetMutableColumns(&ColumnA);
+			Response.SetConstColumns(&ColumnB);
+			Function.Call(Context, Response);
+
+			CHECK_MESSAGE(TEXT("Query function wasn't called."), ColumnA.TestString == TEXT("Test: 42"));
+		}
 	}
 } // namespace UE::Editor::DataStorage::Queries::Tests
 
