[TEDS] Allowed the filter node to work in inverted mode.

The `FilterRowsBy` and filter node in the Query Stack can now work in opposite mode where all rows are kept that fail to pass the filter.
This change also includes additional changes:
- The filter function can no longer be triggered by invalid columns.
- The context compatibility check would always return true.
- The row filter node in the Query Stack didn't pass the update call to its parent.
- The direct query API implementation inherited from the capability interface instead of the capability, causing it to not match with the runtime.
- There's now an explicit way to reset the row filter node in the Query Stack in case external sources such as search text of filter groups change.

diff --git a/Engine/Plugins/Experimental/EditorDataStorage/Source/TedsCore/Private/Queries/TypedElementQueryContext_DirectSingleApi.h b/Engine/Plugins/Experimental/EditorDataStorage/Source/TedsCore/Private/Queries/TypedElementQueryContext_DirectSingleApi.h
--- a/Engine/Plugins/Experimental/EditorDataStorage/Source/TedsCore/Private/Queries/TypedElementQueryContext_DirectSingleApi.h
+++ b/Engine/Plugins/Experimental/EditorDataStorage/Source/TedsCore/Private/Queries/TypedElementQueryContext_DirectSingleApi.h
@@ -10,7 +10,7 @@ class UEditorDataStorage;
 namespace UE::Editor::DataStorage::Queries
 {
 	/** Query context implementation for single row queries that call into TEDS Direct API. */
-	class FQueryContext_DirectSingleApi final : public ISingleRowInfo<IQueryFunctionResponse>
+	class FQueryContext_DirectSingleApi final : public SingleRowInfo, public IQueryFunctionResponse
 	{
 	public:
 		static const FName Name;
diff --git a/Engine/Plugins/Experimental/EditorDataStorage/Source/TedsCore/Private/TypedElementDatabase.cpp b/Engine/Plugins/Experimental/EditorDataStorage/Source/TedsCore/Private/TypedElementDatabase.cpp
--- a/Engine/Plugins/Experimental/EditorDataStorage/Source/TedsCore/Private/TypedElementDatabase.cpp
+++ b/Engine/Plugins/Experimental/EditorDataStorage/Source/TedsCore/Private/TypedElementDatabase.cpp
@@ -578,8 +578,8 @@ bool UEditorDataStorage::IsRowAssignedUnsafe(UE::Editor::DataStorage::RowHandle
 	return UE::Editor::DataStorage::Legacy::FCommandBuffer::Execute_IsRowAssigned(*ActiveEditorEntityManager, Row);
 }
 
-void UEditorDataStorage::FilterRowsBy(UE::Editor::DataStorage::FRowHandleArray& Result, 
-	UE::Editor::DataStorage::FRowHandleArrayView Input, UE::Editor::DataStorage::Queries::TQueryFunction<bool>& Filter)
+void UEditorDataStorage::FilterRowsBy(UE::Editor::DataStorage::FRowHandleArray& Result,  UE::Editor::DataStorage::FRowHandleArrayView Input, 
+	EFilterOptions Options, UE::Editor::DataStorage::Queries::TQueryFunction<bool>& Filter)
 {
 	using namespace UE::Editor::DataStorage;
 	using namespace UE::Editor::DataStorage::Queries;
@@ -589,11 +589,11 @@ void UEditorDataStorage::FilterRowsBy(UE::Editor::DataStorage::FRowHandleArray&
 		QueryContext_DirectSingleApi Context(*this);
 		Context.CheckCompatiblity(Filter);
 
+		bool bAcceptedResult = Options == EFilterOptions::Inclusive ? true : false;
 		for (RowHandle Row : Input)
 		{
 			Context.GetContextImplementation().SetCurrentRow(Row);
-			
-			if (Filter.Call(Context, Context.GetContextImplementation()))
+			if (Filter.Call<EFunctionCallConfig::VerifyColumns>(Context, Context.GetContextImplementation()) == bAcceptedResult)
 			{
 				Result.Add(Row);
 			}
diff --git a/Engine/Plugins/Experimental/EditorDataStorage/Source/TedsCore/Private/TypedElementDatabase.h b/Engine/Plugins/Experimental/EditorDataStorage/Source/TedsCore/Private/TypedElementDatabase.h
--- a/Engine/Plugins/Experimental/EditorDataStorage/Source/TedsCore/Private/TypedElementDatabase.h
+++ b/Engine/Plugins/Experimental/EditorDataStorage/Source/TedsCore/Private/TypedElementDatabase.h
@@ -113,7 +113,7 @@ public:
 	/** Same as IsRowAssigned, but doesn't check if the data storage has been initialized. */
 	UE_API bool IsRowAssignedUnsafe(UE::Editor::DataStorage::RowHandle Row) const;
 	virtual void FilterRowsBy(UE::Editor::DataStorage::FRowHandleArray& Result, UE::Editor::DataStorage::FRowHandleArrayView Input, 
-		UE::Editor::DataStorage::Queries::TQueryFunction<bool>& Filter) override;
+		EFilterOptions Options, UE::Editor::DataStorage::Queries::TQueryFunction<bool>& Filter) override;
 
 	UE_API virtual void AddColumn(UE::Editor::DataStorage::RowHandle Row, const UScriptStruct* ColumnType) override;
 	UE_API virtual void AddColumn(UE::Editor::DataStorage::RowHandle Row, const UE::Editor::DataStorage::FValueTag& Tag, const FName& InValue) override;
diff --git a/Engine/Plugins/Experimental/EditorDataStorageFeatures/Source/TedsQueryStack/Private/TedsRowFilterNode.cpp b/Engine/Plugins/Experimental/EditorDataStorageFeatures/Source/TedsQueryStack/Private/TedsRowFilterNode.cpp
--- a/Engine/Plugins/Experimental/EditorDataStorageFeatures/Source/TedsQueryStack/Private/TedsRowFilterNode.cpp
+++ b/Engine/Plugins/Experimental/EditorDataStorageFeatures/Source/TedsQueryStack/Private/TedsRowFilterNode.cpp
@@ -21,12 +21,18 @@ namespace UE::Editor::DataStorage::QueryStack
 
 	void FRowFilterNode::Update()
 	{
+		ParentRowNode->Update();
 		if (CachedParentRevisionID != ParentRowNode->GetRevision())
 		{
 			Rows.Reset();
-			Storage->FilterRowsBy(Rows, ParentRowNode->GetRows(), Filter);
+			Storage->FilterRowsBy(Rows, ParentRowNode->GetRows(), Options, Filter);
 
 			CachedParentRevisionID = ParentRowNode->GetRevision();
 		}
 	}
+
+	void FRowFilterNode::ForceRefresh()
+	{
+		CachedParentRevisionID = 0;
+	}
 } // namespace UE::Editor::DataStorage::QueryStack
diff --git a/Engine/Plugins/Experimental/EditorDataStorageFeatures/Source/TedsQueryStack/Public/TedsRowFilterNode.h b/Engine/Plugins/Experimental/EditorDataStorageFeatures/Source/TedsQueryStack/Public/TedsRowFilterNode.h
--- a/Engine/Plugins/Experimental/EditorDataStorageFeatures/Source/TedsQueryStack/Public/TedsRowFilterNode.h
+++ b/Engine/Plugins/Experimental/EditorDataStorageFeatures/Source/TedsQueryStack/Public/TedsRowFilterNode.h
@@ -22,12 +22,24 @@ namespace UE::Editor::DataStorage::QueryStack
 	class FRowFilterNode final : public IRowNode
 	{
 	public:
+		template<typename FilterFunction>
+		FRowFilterNode(ICoreProvider* Storage, const TSharedPtr<IRowNode>& ParentRowNode, ICoreProvider::EFilterOptions Options, 
+			FilterFunction&& Filter)
+			: Storage(Storage)
+			, Filter(Queries::BuildQueryFunction<bool>(Forward<FilterFunction>(Filter)))
+			, ParentRowNode(ParentRowNode)
+			, CachedParentRevisionID(ParentRowNode->GetRevision())
+			, Options(Options)
+		{
+		}
+
 		template<typename FilterFunction>
 		FRowFilterNode(ICoreProvider* Storage, const TSharedPtr<IRowNode>& ParentRowNode, FilterFunction&& Filter)
 			: Storage(Storage)
 			, Filter(Queries::BuildQueryFunction<bool>(Forward<FilterFunction>(Filter)))
 			, ParentRowNode(ParentRowNode)
 			, CachedParentRevisionID(ParentRowNode->GetRevision())
+			, Options(ICoreProvider::EFilterOptions::Inclusive)
 		{
 		}
 
@@ -37,6 +49,9 @@ namespace UE::Editor::DataStorage::QueryStack
 		TEDSQUERYSTACK_API virtual FRowHandleArray& GetMutableRows() override;
 		TEDSQUERYSTACK_API virtual RevisionId GetRevision() const override;
 		TEDSQUERYSTACK_API virtual void Update() override;
+		
+		/** Call this when external inputs into the filter, such as search strings or filtering groups, have been updated. */
+		TEDSQUERYSTACK_API void ForceRefresh();
 
 	private:
 		ICoreProvider* Storage;
@@ -44,5 +59,6 @@ namespace UE::Editor::DataStorage::QueryStack
 		FRowHandleArray Rows;
 		TSharedPtr<IRowNode> ParentRowNode;
 		RevisionId CachedParentRevisionID;
+		ICoreProvider::EFilterOptions Options;
 	};
 }
diff --git a/Engine/Source/Runtime/TypedElementFramework/Public/Elements/Framework/TypedElementQueryContextImplementation.inl b/Engine/Source/Runtime/TypedElementFramework/Public/Elements/Framework/TypedElementQueryContextImplementation.inl
--- a/Engine/Source/Runtime/TypedElementFramework/Public/Elements/Framework/TypedElementQueryContextImplementation.inl
+++ b/Engine/Source/Runtime/TypedElementFramework/Public/Elements/Framework/TypedElementQueryContextImplementation.inl
@@ -15,13 +15,27 @@ namespace UE::Editor::DataStorage::Queries
 	template<ContextCapability RequestedCapability>
 	constexpr bool TQueryContextImpl<QueryContext, SupportedCapabilities...>::SupportsCapability()
 	{
-		return std::is_base_of_v<RequestedCapability, QueryContext> && (std::is_same_v<RequestedCapability, SupportedCapabilities> || ...);
+		if constexpr (sizeof...(SupportedCapabilities) > 0)
+		{
+			return std::is_base_of_v<RequestedCapability, QueryContext> && (std::is_same_v<RequestedCapability, SupportedCapabilities> || ...);
+		}
+		else
+		{
+			return std::is_base_of_v<RequestedCapability, QueryContext>;
+		}
 	}
 
 	template<typename QueryContext, ContextCapability... SupportedCapabilities>
 	bool TQueryContextImpl<QueryContext, SupportedCapabilities...>::SupportsCapability(const FName& Capability)
 	{
-		return IContextContract::SupportsCapability<QueryContext>(Capability) && ((SupportedCapabilities::Name == Capability) || ...);
+		if constexpr (sizeof...(SupportedCapabilities) > 0)
+		{
+			return IContextContract::SupportsCapability<QueryContext>(Capability) && ((SupportedCapabilities::Name == Capability) || ...);
+		}
+		else
+		{
+			return IContextContract::SupportsCapability<QueryContext>(Capability);
+		}
 	}
 
 	template<typename QueryContext, ContextCapability... SupportedCapabilities>
@@ -157,11 +171,11 @@ namespace UE::Editor::DataStorage::Queries
 			if (!SupportsCapability(Capability))
 			{
 				Result = false;
-				checkf(false, TEXT("Requested query context '%s' is not support by query context implementation '%s'."), 
-					*Capability.ToString(), *(QueryContext::Name).ToString());
+				checkf(false, TEXT("Requested query context '%s' is not supported by query context implementation."), 
+					*Capability.ToString());
 			}
 		}
-		return false;
+		return Result;
 	}
 
 	template<typename QueryContext, ContextCapability... SupportedCapabilities>
diff --git a/Engine/Source/Runtime/TypedElementFramework/Public/Elements/Framework/TypedElementQueryFunctions.h b/Engine/Source/Runtime/TypedElementFramework/Public/Elements/Framework/TypedElementQueryFunctions.h
--- a/Engine/Source/Runtime/TypedElementFramework/Public/Elements/Framework/TypedElementQueryFunctions.h
+++ b/Engine/Source/Runtime/TypedElementFramework/Public/Elements/Framework/TypedElementQueryFunctions.h
@@ -4,6 +4,7 @@
 
 #include "Containers/ContainersFwd.h"
 #include "Elements/Framework/TypedElementQueryCapabilities.h"
+#include "Misc/EnumClassFlags.h"
 #include "Templates/Function.h"
 #include "UObject/NameTypes.h"
 
@@ -21,6 +22,14 @@ namespace UE::Editor::DataStorage::Queries
 		virtual void GetMutableColumns(TArrayView<void*> ColumnsData, TConstArrayView<const UScriptStruct*> ColumnTypes) = 0;
 	};
 
+	enum class EFunctionCallConfig
+	{
+		None = 0,
+		/** If not all columns could be retrieved, return the default value for the return type instead of calling the function. */
+		VerifyColumns = 1 << 0,
+	};
+	ENUM_CLASS_FLAGS(EFunctionCallConfig);
+
 	/** Storage for a function that can be used as part of a query.  */
 	template<typename ReturnType>
 	struct TQueryFunction
@@ -30,6 +39,7 @@ namespace UE::Editor::DataStorage::Queries
 		TConstArrayView<const UScriptStruct*> MutableColumnTypes;
 		TFunction<ReturnType(IContextContract& Contract, TConstArrayView<const void*> ConstColumns, TConstArrayView<void*> MutableColumns)> Function;
 
+		template<EFunctionCallConfig Config = EFunctionCallConfig::None>
 		ReturnType Call(IContextContract& Contract, IQueryFunctionResponse& Response);
 	};
 
diff --git a/Engine/Source/Runtime/TypedElementFramework/Public/Elements/Framework/TypedElementQueryFunctions.inl b/Engine/Source/Runtime/TypedElementFramework/Public/Elements/Framework/TypedElementQueryFunctions.inl
--- a/Engine/Source/Runtime/TypedElementFramework/Public/Elements/Framework/TypedElementQueryFunctions.inl
+++ b/Engine/Source/Runtime/TypedElementFramework/Public/Elements/Framework/TypedElementQueryFunctions.inl
@@ -80,6 +80,7 @@ namespace UE::Editor::DataStorage::Queries
 	}
 
 	template<typename ReturnType>
+	template<EFunctionCallConfig Config>
 	ReturnType TQueryFunction<ReturnType>::Call(IContextContract& Contract, IQueryFunctionResponse& Response)
 	{
 		const void** ConstColumnsData = static_cast<const void**>(alloca(sizeof(void*) * ConstColumnTypes.Num()));
@@ -88,20 +89,34 @@ namespace UE::Editor::DataStorage::Queries
 		Response.GetConstColumns(TArrayView<const void*>(ConstColumnsData, ConstColumnTypes.Num()), ConstColumnTypes);
 		Response.GetMutableColumns(TArrayView<void*>(MutableColumnsData, MutableColumnTypes.Num()), MutableColumnTypes);
 
-		if constexpr (std::is_same_v<ReturnType, void>)
+		if constexpr (EnumHasAnyFlags(Config, EFunctionCallConfig::VerifyColumns))
 		{
-			Function(
-				Contract,
-				TConstArrayView<const void*>(ConstColumnsData, ConstColumnTypes.Num()),
-				TConstArrayView<void*>(MutableColumnsData, MutableColumnTypes.Num()));
-		}
-		else
-		{
-			return Function(
-				Contract,
-				TConstArrayView<const void*>(ConstColumnsData, ConstColumnTypes.Num()),
-				TConstArrayView<void*>(MutableColumnsData, MutableColumnTypes.Num()));
+			bool bMissingColumn = false;
+			for (const void* Column : TArrayView<const void*>(ConstColumnsData, ConstColumnTypes.Num()))
+			{
+				bMissingColumn = bMissingColumn || Column == nullptr;
+			}
+			for (void* Column : TArrayView<void*>(MutableColumnsData, MutableColumnTypes.Num()))
+			{
+				bMissingColumn = bMissingColumn || Column == nullptr;
+			}
+			if (bMissingColumn)
+			{
+				if constexpr (std::is_same_v<void, ReturnType>)
+				{
+					return;
+				}
+				else
+				{
+					return ReturnType{};
+				}
+			}
 		}
+
+		return Function(
+			Contract,
+			TConstArrayView<const void*>(ConstColumnsData, ConstColumnTypes.Num()),
+			TConstArrayView<void*>(MutableColumnsData, MutableColumnTypes.Num()));
 	}
 
 	template<typename ReturnType, typename FunctionType>
@@ -137,14 +152,7 @@ namespace UE::Editor::DataStorage::Queries
 				ArgumentInfo::SetConstColumns(Arguments, ConstColumns);
 				ArgumentInfo::SetMutableColumns(Arguments, MutableColumns);
 
-				if constexpr (std::is_same_v<ReturnType, void>)
-				{
-					Arguments.ApplyBefore(Callback);
-				}
-				else
-				{
-					return Arguments.ApplyBefore(Callback);
-				}
+				return Arguments.ApplyBefore(Callback);
 			};
 		return Result;
 	}
diff --git a/Engine/Source/Runtime/TypedElementFramework/Public/Elements/Interfaces/TypedElementDataStorageInterface.h b/Engine/Source/Runtime/TypedElementFramework/Public/Elements/Interfaces/TypedElementDataStorageInterface.h
--- a/Engine/Source/Runtime/TypedElementFramework/Public/Elements/Interfaces/TypedElementDataStorageInterface.h
+++ b/Engine/Source/Runtime/TypedElementFramework/Public/Elements/Interfaces/TypedElementDataStorageInterface.h
@@ -184,11 +184,20 @@ public:
 	/** Checks whether or not a row has been reserved but not yet assigned to a table. */
 	virtual bool IsRowAssigned(RowHandle Row) const = 0;
 
+	enum class EFilterOptions : uint8
+	{
+		/** Include rows that have the columns required for the filter function and the filter function returns true. */
+		Inclusive,
+		/** Exclude rows that have the columns required for the filter function and the filter function returns true. */
+		Exclusive
+	};
+
 	/** Filters the provided rows using the query filter and stores the results in the provided list. */
-	virtual void FilterRowsBy(FRowHandleArray& Result, FRowHandleArrayView Input, Queries::TQueryFunction<bool>& Filter) = 0;
+	virtual void FilterRowsBy(
+		FRowHandleArray& Result, FRowHandleArrayView Input, EFilterOptions Options, Queries::TQueryFunction<bool>& Filter) = 0;
 	/** Filters the provided rows using the query filter and stores the results in the provided list. */
 	template<typename FilterFunction>
-	void FilterRowsBy(FRowHandleArray& Result, FRowHandleArrayView Input, FilterFunction&& Filter);
+	void FilterRowsBy(FRowHandleArray& Result, FRowHandleArrayView Input, EFilterOptions Options, FilterFunction&& Filter);
 
 	/**
 	 * @section Column management
@@ -571,9 +580,9 @@ TableHandle ICoreProvider::RegisterTable(
 }
 
 template<typename FilterFunction>
-void ICoreProvider::FilterRowsBy(FRowHandleArray& Result, FRowHandleArrayView Input, FilterFunction&& Filter)
+void ICoreProvider::FilterRowsBy(FRowHandleArray& Result, FRowHandleArrayView Input, EFilterOptions Options, FilterFunction&& Filter)
 {
-	FilterRowsBy(Result, Input, Queries::BuildQueryFunction<bool>(Forward<FilterFunction>(Filter)));
+	FilterRowsBy(Result, Input, Options, Queries::BuildQueryFunction<bool>(Forward<FilterFunction>(Filter)));
 }
 
 template<TColumnType Column>
diff --git a/Engine/Source/Runtime/TypedElementFramework/Tests/Elements/Framework/TypedElementQueryFunctionsTests.cpp b/Engine/Source/Runtime/TypedElementFramework/Tests/Elements/Framework/TypedElementQueryFunctionsTests.cpp
--- a/Engine/Source/Runtime/TypedElementFramework/Tests/Elements/Framework/TypedElementQueryFunctionsTests.cpp
+++ b/Engine/Source/Runtime/TypedElementFramework/Tests/Elements/Framework/TypedElementQueryFunctionsTests.cpp
@@ -142,6 +142,33 @@ namespace UE::Editor::DataStorage::Queries::Tests
 			CHECK_MESSAGE(TEXT("Incorrect mutable column set."), Result.MutableColumnTypes[0] == FTestColumnInt::StaticStruct());
 		}
 
+		SECTION("Check context compatibility without arguments")
+		{
+			TQueryFunction<void> Function = BuildQueryFunction<void>([](){});
+
+			QueryContextMock Context;
+			bool bResult = Context.CheckCompatiblity(Function);
+			CHECK_MESSAGE(TEXT("Function without context incorrectly found to not be matching."), bResult);
+		}
+
+		SECTION("Check context compatibility with context")
+		{
+			TQueryFunction<void> Function = BuildQueryFunction<void>([](TQueryContext<SingleRowInfo>& Context) {});
+
+			QueryContextMock Context;
+			bool bResult = Context.CheckCompatiblity(Function);
+			CHECK_MESSAGE(TEXT("Function with context incorrectly found to not be matching."), bResult);
+		}
+
+		SECTION("Check context compatibility with restricted context")
+		{
+			TQueryFunction<void> Function = BuildQueryFunction<void>([](TQueryContext<SingleRowInfo>& Context) {});
+
+			TQueryContextImpl<FQueryContextMock, SingleRowInfo> Context;
+			bool bResult = Context.CheckCompatiblity(Function);
+			CHECK_MESSAGE(TEXT("Function with context incorrectly found to not be matching."), bResult);
+		}
+
 		SECTION("Call without arguments")
 		{
 			TQueryFunction<bool> Function = BuildQueryFunction<bool>(
@@ -153,7 +180,7 @@ namespace UE::Editor::DataStorage::Queries::Tests
 			bool bResult = false;
 			QueryContextMock Context;
 			FTestQueryFunctionResponse Response;
-			bResult = Function.Call(Context, Response);
+			bResult = Function.Call<EFunctionCallConfig::VerifyColumns>(Context, Response);
 			CHECK_MESSAGE(TEXT("Query function wasn't called."), bResult);
 		}
 
@@ -170,7 +197,7 @@ namespace UE::Editor::DataStorage::Queries::Tests
 			
 			FTestQueryFunctionResponse Response;
 			Response.SetMutableColumns(&Column);
-			Function.Call(Context, Response);
+			Function.Call<EFunctionCallConfig::VerifyColumns>(Context, Response);
 			
 			CHECK_MESSAGE(TEXT("Query function wasn't called."), Column.TestString == TEXT("Callback"));
 		}
@@ -191,10 +218,78 @@ namespace UE::Editor::DataStorage::Queries::Tests
 			FTestQueryFunctionResponse Response;
 			Response.SetMutableColumns(&ColumnA);
 			Response.SetConstColumns(&ColumnB);
-			Function.Call(Context, Response);
+			Function.Call<EFunctionCallConfig::VerifyColumns>(Context, Response);
 
 			CHECK_MESSAGE(TEXT("Query function wasn't called."), ColumnA.TestString == TEXT("Test: 42"));
 		}
+
+		SECTION("Call and use context")
+		{
+			TQueryFunction<void> Function = BuildQueryFunction<void>(
+				[](TQueryContext<SingleRowInfo>& Context)
+				{
+					Context.GetCurrentRow();
+					Context.CurrentRowHasColumns<FTestColumnString>();
+				});
+
+			QueryContextMock Context;
+			
+			bool bGetCurrentRowCalled = false;
+			Context.GetContextImplementation().SingleRowInfo_GetCurrentRow_Mock =
+				[&bGetCurrentRowCalled]()
+				{
+					bGetCurrentRowCalled = true;
+					return 0;
+				};
+			
+			bool bCurrentRowHasColumnsCalled = false;
+			Context.GetContextImplementation().SingleRowInfo_CurrentRowHasColumns_Mock =
+				[&bCurrentRowHasColumnsCalled](TConstArrayView<const UScriptStruct*>)
+				{
+					bCurrentRowHasColumnsCalled = true;
+					return false;
+				};
+
+			FTestQueryFunctionResponse Response;
+			Function.Call<>(Context, Response);
+
+			CHECK_MESSAGE(TEXT("GetCurrentRow not called on context."), bGetCurrentRowCalled);
+			CHECK_MESSAGE(TEXT("CurrentRowHasColumns not called on context."), bCurrentRowHasColumnsCalled);
+		}
+
+		SECTION("Call and use restricted context")
+		{
+			TQueryFunction<void> Function = BuildQueryFunction<void>(
+				[](TQueryContext<SingleRowInfo>& Context)
+				{
+					Context.GetCurrentRow();
+					Context.CurrentRowHasColumns<FTestColumnString>();
+				});
+
+			TQueryContextImpl<FQueryContextMock, SingleRowInfo> Context;
+
+			bool bGetCurrentRowCalled = false;
+			Context.GetContextImplementation().SingleRowInfo_GetCurrentRow_Mock =
+				[&bGetCurrentRowCalled]()
+				{
+					bGetCurrentRowCalled = true;
+					return 0;
+				};
+
+			bool bCurrentRowHasColumnsCalled = false;
+			Context.GetContextImplementation().SingleRowInfo_CurrentRowHasColumns_Mock =
+				[&bCurrentRowHasColumnsCalled](TConstArrayView<const UScriptStruct*>)
+				{
+					bCurrentRowHasColumnsCalled = true;
+					return false;
+				};
+
+			FTestQueryFunctionResponse Response;
+			Function.Call<>(Context, Response);
+
+			CHECK_MESSAGE(TEXT("GetCurrentRow not called on context."), bGetCurrentRowCalled);
+			CHECK_MESSAGE(TEXT("CurrentRowHasColumns not called on context."), bCurrentRowHasColumnsCalled);
+		}
 	}
 } // namespace UE::Editor::DataStorage::Queries::Tests
 
