[AutoRTFM] Update TypedElementFramework for transactional safety.

This code is accessed via editor-only code paths in the ControlRig unit tests.

diff --git a/Engine/Source/Runtime/TypedElementFramework/Private/Elements/Framework/TypedElementRegistry.cpp b/Engine/Source/Runtime/TypedElementFramework/Private/Elements/Framework/TypedElementRegistry.cpp
--- a/Engine/Source/Runtime/TypedElementFramework/Private/Elements/Framework/TypedElementRegistry.cpp
+++ b/Engine/Source/Runtime/TypedElementFramework/Private/Elements/Framework/TypedElementRegistry.cpp
@@ -84,7 +84,7 @@ void UTypedElementRegistry::AddReferencedObjects(UObject* InThis, FReferenceColl
 
 	UTypedElementRegistry* This = CastChecked<UTypedElementRegistry>(InThis);
 
-	FReadScopeLock RegisteredElementTypesLock(This->RegisteredElementTypesRW);
+	UE::TReadScopeLock RegisteredElementTypesLock(This->RegisteredElementTypesRW);
 	
 	for (TUniquePtr<FRegisteredElementType>& RegisteredElementType : This->RegisteredElementTypes)
 	{
@@ -151,7 +151,7 @@ void UTypedElementRegistry::ProcessDeferredElementsToDestroy()
 {
 	OnProcessingDeferredElementsToDestroyDelegate.Broadcast();
 
-	FReadScopeLock RegisteredElementTypesLock(RegisteredElementTypesRW);
+	UE::TReadScopeLock RegisteredElementTypesLock(RegisteredElementTypesRW);
 
 	LLM_SCOPE(ELLMTag::EngineMisc);
 	for (TUniquePtr<FRegisteredElementType>& RegisteredElementType : RegisteredElementTypes)
@@ -222,7 +222,7 @@ FString UTypedElementRegistry::RegistredElementTypesAndInterfacesToString() cons
 	TArray<FString> Lines;
 
 	{
-		FReadScopeLock RegisteredElementTypesLock(RegisteredElementTypesRW);
+		UE::TReadScopeLock RegisteredElementTypesLock(RegisteredElementTypesRW);
 		Lines.Reserve(RegisteredElementTypesNameToId.Num());
 
 		{
diff --git a/Engine/Source/Runtime/TypedElementFramework/Public/Elements/Framework/TypedElementData.h b/Engine/Source/Runtime/TypedElementFramework/Public/Elements/Framework/TypedElementData.h
--- a/Engine/Source/Runtime/TypedElementFramework/Public/Elements/Framework/TypedElementData.h
+++ b/Engine/Source/Runtime/TypedElementFramework/Public/Elements/Framework/TypedElementData.h
@@ -5,6 +5,7 @@
 #include "CoreMinimal.h"
 #include "Misc/ScopeLock.h"
 #include "Misc/ScopeRWLock.h"
+#include "Misc/TransactionallySafeRWLock.h"
 #include "HAL/PlatformStackWalk.h"
 #include "Containers/SparseArray.h"
 #include "Containers/ChunkedArray.h"
@@ -186,7 +187,19 @@ public:
 	{
 		Id.Private_DestroyNoRef();
 #if UE_TYPED_ELEMENT_HAS_REFCOUNTING
-		FPlatformAtomics::InterlockedExchange(&RefCount, 0);
+		FTypedElementRefCount OldRefCount = 0;
+		UE_AUTORTFM_OPEN
+		{
+			OldRefCount = FPlatformAtomics::InterlockedExchange(&RefCount, 0);
+		};
+		UE_AUTORTFM_ONABORT(this, OldRefCount)
+		{
+			FTypedElementRefCount RefCountDuringAbort = FPlatformAtomics::InterlockedAdd(&RefCount, OldRefCount);
+			// Is there a real-world scenario where rolling back a Reset() could yield a non-zero value here?
+			// Any further AddRefs would also be paired with an OnAbort handler, undoing its increment.
+			// Any ReleaseRefs would be inert until commit time.
+			ensure(RefCountDuringAbort == 0);
+		};
 #endif	// UE_TYPED_ELEMENT_HAS_REFCOUNTING
 #if UE_TYPED_ELEMENT_HAS_REFTRACKING
 		if (References)
@@ -206,7 +219,15 @@ public:
 		FTypedElementReferenceId ReferenceId = INDEX_NONE;
 #if UE_TYPED_ELEMENT_HAS_REFCOUNTING
 		checkSlow(RefCount < TNumericLimits<FTypedElementRefCount>::Max());
-		FPlatformAtomics::InterlockedIncrement(&RefCount);
+		UE_AUTORTFM_OPEN
+		{
+			FPlatformAtomics::InterlockedIncrement(&RefCount);
+		};
+		UE_AUTORTFM_ONABORT(this)
+		{
+			checkSlow(RefCount > 0);
+			FPlatformAtomics::InterlockedDecrement(&RefCount);
+		};
 #endif	// UE_TYPED_ELEMENT_HAS_REFCOUNTING
 #if UE_TYPED_ELEMENT_HAS_REFTRACKING
 		if (bCanTrackReference && References)
@@ -221,7 +242,11 @@ public:
 	{
 #if UE_TYPED_ELEMENT_HAS_REFCOUNTING
 		checkSlow(RefCount > 0);
-		FPlatformAtomics::InterlockedDecrement(&RefCount);
+		UE_AUTORTFM_ONCOMMIT(this)
+		{
+			checkSlow(RefCount > 0);
+			FPlatformAtomics::InterlockedDecrement(&RefCount);
+		};
 #endif	// UE_TYPED_ELEMENT_HAS_REFCOUNTING
 #if UE_TYPED_ELEMENT_HAS_REFTRACKING
 		if (References)
@@ -231,6 +256,7 @@ public:
 #endif	// UE_TYPED_ELEMENT_HAS_REFTRACKING
 	}
 
+	UE_AUTORTFM_ALWAYS_OPEN 
 	inline FTypedElementRefCount GetRefCount() const
 	{
 #if UE_TYPED_ELEMENT_HAS_REFCOUNTING
@@ -515,7 +541,7 @@ public:
 
 	TTypedElementInternalData<ElementDataType>& AddDataForElement(const FTypedHandleTypeId InTypeId, FTypedHandleElementId& InOutElementId)
 	{
-		FWriteScopeLock InternalDataLock(InternalDataRW);
+		UE::TWriteScopeLock InternalDataLock(InternalDataRW);
 
 		checkSlow(InOutElementId < 0);
 
@@ -537,7 +563,7 @@ public:
 			return *ScriptElementInternalData;
 		}
 
-		FReadScopeLock InternalDataLock(InternalDataRW);
+		UE::TReadScopeLock InternalDataLock(InternalDataRW);
 		checkSlow(InternalDataArray.IsValidIndex(InElementId));
 		return ScriptInternalDataMap.EmplaceByHash(Hash,InElementId, InternalDataArray[InElementId]);
 	}
@@ -550,7 +576,7 @@ public:
 
 	void RemoveDataForElement(const FTypedHandleElementId InElementId, const FTypedElementInternalData* InExpectedDataPtr)
 	{
-		FWriteScopeLock InternalDataLock(InternalDataRW);
+		UE::TWriteScopeLock InternalDataLock(InternalDataRW);
 
 		checkSlow(InternalDataArray.IsValidIndex(InElementId));
 		
@@ -563,7 +589,7 @@ public:
 
 	const TTypedElementInternalData<ElementDataType>& GetDataForElement(const FTypedHandleElementId InElementId) const
 	{
-		FReadScopeLock InternalDataLock(InternalDataRW);
+		UE::TReadScopeLock InternalDataLock(InternalDataRW);
 
 		checkSlow(InternalDataArray.IsValidIndex(InElementId));
 		return InternalDataArray[InElementId];
@@ -586,7 +612,7 @@ public:
 	}
 
 private:
-	mutable FRWLock InternalDataRW;
+	mutable FTransactionallySafeRWLock InternalDataRW;
 	TArray<int32> InternalDataFreeIndices;
 	TChunkedArray<TTypedElementInternalData<ElementDataType>> InternalDataArray;
 	TMap<FTypedHandleElementId, FScriptTypedElementInternalDataOwner> ScriptInternalDataMap;
@@ -604,7 +630,7 @@ public:
 
 	TTypedElementInternalData<void>& AddDataForElement(const FTypedHandleTypeId InTypeId, FTypedHandleElementId& InOutElementId)
 	{
-		FWriteScopeLock InternalDataLock(InternalDataRW);
+		UE::TWriteScopeLock InternalDataLock(InternalDataRW);
 
 		checkSlow(InOutElementId >= 0);
 		checkSlow(!ElementIdToArrayIndex.Contains(InOutElementId));
@@ -629,7 +655,7 @@ public:
 			return *ScriptElementInternalData;
 		}
 
-		FReadScopeLock InternalDataLock(InternalDataRW);
+		UE::TReadScopeLock InternalDataLock(InternalDataRW);
 		const int32 Index = ElementIdToArrayIndex.FindChecked(InElementId);
 		return ScriptInternalDataMap.EmplaceByHash(Hash, InElementId, InternalDataArray[Index]);
 	}
@@ -642,7 +668,7 @@ public:
 
 	void RemoveDataForElement(const FTypedHandleElementId InElementId, const FTypedElementInternalData* InExpectedDataPtr)
 	{
-		FWriteScopeLock InternalDataLock(InternalDataRW);
+		UE::TWriteScopeLock InternalDataLock(InternalDataRW);
 
 		int32 InternalDataArrayIndex = INDEX_NONE;
 		ElementIdToArrayIndex.RemoveAndCopyValue(InElementId, InternalDataArrayIndex);
@@ -658,7 +684,7 @@ public:
 
 	const TTypedElementInternalData<void>& GetDataForElement(const FTypedHandleElementId InElementId) const
 	{
-		FReadScopeLock InternalDataLock(InternalDataRW);
+		UE::TReadScopeLock InternalDataLock(InternalDataRW);
 
 		const int32* InternalDataArrayIndexPtr = ElementIdToArrayIndex.Find(InElementId);
 		checkSlow(InternalDataArrayIndexPtr && InternalDataArray.IsValidIndex(*InternalDataArrayIndexPtr));
@@ -680,7 +706,7 @@ public:
 	}
 
 private:
-	mutable FRWLock InternalDataRW;
+	mutable FTransactionallySafeRWLock InternalDataRW;
 	TChunkedArray<TTypedElementInternalData<void>> InternalDataArray;
 	TArray<int32> InternalDataFreeIndices;
 	TMap<FTypedHandleElementId, int32> ElementIdToArrayIndex;
diff --git a/Engine/Source/Runtime/TypedElementFramework/Public/Elements/Framework/TypedElementOwnerStore.h b/Engine/Source/Runtime/TypedElementFramework/Public/Elements/Framework/TypedElementOwnerStore.h
--- a/Engine/Source/Runtime/TypedElementFramework/Public/Elements/Framework/TypedElementOwnerStore.h
+++ b/Engine/Source/Runtime/TypedElementFramework/Public/Elements/Framework/TypedElementOwnerStore.h
@@ -4,6 +4,7 @@
 
 #include "CoreMinimal.h"
 #include "Misc/ScopeLock.h"
+#include "Misc/TransactionallySafeCriticalSection.h"
 #include "Elements/Framework/TypedElementHandle.h"
 
 /**
@@ -15,7 +16,7 @@ template <typename ElementOwnerType>
 struct TTypedElementOwnerScopedAccessImpl
 {
 public:
-	explicit TTypedElementOwnerScopedAccessImpl(FCriticalSection* InSynchObject)
+	explicit TTypedElementOwnerScopedAccessImpl(FTransactionallySafeCriticalSection* InSynchObject)
 		: SynchObject(InSynchObject)
 	{
 		check(SynchObject);
@@ -81,7 +82,7 @@ public:
 	}
 
 private:
-	FCriticalSection* SynchObject = nullptr;
+	FTransactionallySafeCriticalSection* SynchObject = nullptr;
 	ElementOwnerType* ElementOwner = nullptr;
 };
 
@@ -141,7 +142,7 @@ public:
 	{
 		TTypedElementOwner<ElementDataType> ElementOwner;
 		{
-			FScopeLock Lock(&ElementOwnerMapCS);
+			UE::TScopeLock Lock(ElementOwnerMapCS);
 			ElementOwnerMap.RemoveAndCopyValue(InKey, ElementOwner);
 		}
 		return ElementOwner;
@@ -152,7 +153,7 @@ public:
 	 */
 	void UnregisterElementOwners(TFunctionRef<bool(const TTypedElementOwner<ElementDataType>&)> InShouldUnregisterElement, TFunctionRef<void(TTypedElementOwner<ElementDataType>&&)> InOnUnregisterElement)
 	{
-		FScopeLock Lock(&ElementOwnerMapCS);
+		UE::TScopeLock Lock(ElementOwnerMapCS);
 		for (auto It = ElementOwnerMap.CreateIterator(); It; ++It)
 		{
 			if (InShouldUnregisterElement(It->Value))
@@ -208,11 +209,11 @@ public:
 	 */
 	bool ContainsElementOwner(const KeyDataType& InKey) const
 	{
-		FScopeLock Lock(&ElementOwnerMapCS);
+		UE::TScopeLock Lock(ElementOwnerMapCS);
 		return ElementOwnerMap.Contains(InKey);
 	}
 
 private:
 	TMap<KeyDataType, TTypedElementOwner<ElementDataType>> ElementOwnerMap;
-	mutable FCriticalSection ElementOwnerMapCS;
+	mutable FTransactionallySafeCriticalSection ElementOwnerMapCS;
 };
diff --git a/Engine/Source/Runtime/TypedElementFramework/Public/Elements/Framework/TypedElementRegistry.h b/Engine/Source/Runtime/TypedElementFramework/Public/Elements/Framework/TypedElementRegistry.h
--- a/Engine/Source/Runtime/TypedElementFramework/Public/Elements/Framework/TypedElementRegistry.h
+++ b/Engine/Source/Runtime/TypedElementFramework/Public/Elements/Framework/TypedElementRegistry.h
@@ -27,6 +27,7 @@
 #include "Misc/AssertionMacros.h"
 #include "Misc/ScopeLock.h"
 #include "Misc/ScopeRWLock.h"
+#include "Misc/TransactionallySafeRWLock.h"
 #include "Templates/Casts.h"
 #include "Templates/SharedPointer.h"
 #include "Templates/SubclassOf.h"
@@ -612,7 +613,7 @@ private:
 		checkf(!GetRegisteredElementTypeFromName(InRegisteredElementType->TypeName), TEXT("Element type '%s' has already been registered!"), *InRegisteredElementType->TypeName.ToString());
 
 		{
-			FWriteScopeLock RegisteredElementTypesLock(RegisteredElementTypesRW);
+			UE::TWriteScopeLock RegisteredElementTypesLock(RegisteredElementTypesRW);
 
 			LLM_SCOPE(ELLMTag::EngineMisc);
 			RegisteredElementTypesNameToId.Add(InRegisteredElementType->TypeName, InRegisteredElementType->TypeId);
@@ -622,7 +623,7 @@ private:
 
 	FRegisteredElementType* GetRegisteredElementTypeFromId(const FTypedHandleTypeId InTypeId) const
 	{
-		FReadScopeLock RegisteredElementTypesLock(RegisteredElementTypesRW);
+		UE::TReadScopeLock RegisteredElementTypesLock(RegisteredElementTypesRW);
 
 		return InTypeId > 0
 			? RegisteredElementTypes[InTypeId - 1].Get()
@@ -631,7 +632,7 @@ private:
 
 	FRegisteredElementType* GetRegisteredElementTypeFromName(const FName& InTypeName) const
 	{
-		FReadScopeLock RegisteredElementTypesLock(RegisteredElementTypesRW);
+		UE::TReadScopeLock RegisteredElementTypesLock(RegisteredElementTypesRW);
 
 		if (const FTypedHandleTypeId* TypeId = RegisteredElementTypesNameToId.Find(InTypeName))
 		{
@@ -659,7 +660,7 @@ private:
 		--DisableElementDestructionOnGCCount;
 	}
 
-	mutable FRWLock RegisteredElementTypesRW;
+	mutable FTransactionallySafeRWLock RegisteredElementTypesRW;
 	TUniquePtr<FRegisteredElementType> RegisteredElementTypes[TypedHandleMaxTypeId - 1];
 	TSortedMap<FName, FTypedHandleTypeId, FDefaultAllocator, FNameFastLess> RegisteredElementTypesNameToId;
 
