[TEDS] Added support for batching to the new Query Context.

The following changes were made:
- Context capabilities now have flags that can be used to pass along additional information about them. This is currently used to indicate if a capability works with single row, row batching or both.
- Query callback funtions can now provide results through a TResult argument. Single row query callbacks will continue to support returning results as well.
- If single row and batch row capabilities are mixed when creating a context, a compile error will be thrown.
- If single row and batch row style arguments are mixed there will now be a compile error pointing this out. There are also checks to make sure that number of arguments is in range, such as have no more than 1 context or that functions requiring a result have a query that either returns that results or has a TResult argument.
- Changed the way arguments are extracted from functions. This makes the code itself a bit more cryptic as it requires more data extraction, but also allows far more errors to be detected at compile time with more specific messages. Soime examples are:
- Enforce that TQueryContext is taken by value. This is done because it's essentially a smart pointer behind the scenes so taking by reference isn't necessary.
- Redirect users to TBatch if they take a column by pointer.
- Report when arguments like EFlow can't be taken by value.
- Added support to exit a query callback early. This can be done by adding an EFlow& argument to the function and setting it to break when work is done. This is useful for exiting linear searches early or to do time slicing work.
- FilterRowsBy now uses the updated way to call a query callback function.
- The RowBatchInfo capability now has a ForEachRow function that can be used to pass in one or more column batches and then have a lamba take individual columns while its iterating over the rows.
- FQueryContext_DirectSingleApi now supports batched operations, although it will create a batch for the one row it operates on. This was added for compatibility and offers no performance gain.
- Several unit tests were added and updated.

diff --git a/Engine/Plugins/Experimental/EditorDataStorage/Source/TedsCore/Private/Queries/TypedElementQueryContext_DirectSingleApi.cpp b/Engine/Plugins/Experimental/EditorDataStorage/Source/TedsCore/Private/Queries/TypedElementQueryContext_DirectSingleApi.cpp
--- a/Engine/Plugins/Experimental/EditorDataStorage/Source/TedsCore/Private/Queries/TypedElementQueryContext_DirectSingleApi.cpp
+++ b/Engine/Plugins/Experimental/EditorDataStorage/Source/TedsCore/Private/Queries/TypedElementQueryContext_DirectSingleApi.cpp
@@ -27,6 +27,16 @@ namespace UE::Editor::DataStorage::Queries
 		CurrentRow = Row;
 	}
 
+	bool FQueryContext_DirectSingleApi::NextBatch()
+	{
+		return false; // Only operates on a single row, so there are no more batches.
+	}
+
+	bool FQueryContext_DirectSingleApi::NextRow()
+	{
+		return false; // Only operates on a single row, so there are no more rows.
+	}
+
 	void FQueryContext_DirectSingleApi::GetConstColumns(TArrayView<const void*> ColumnsData, TConstArrayView<const UScriptStruct*> ColumnTypes)
 	{
 		const void** ColumnsDataIt = ColumnsData.GetData();
@@ -46,4 +56,19 @@ namespace UE::Editor::DataStorage::Queries
 		}
 	}
 
+	uint32 FQueryContext_DirectSingleApi::GetBatchRowCount() const
+	{
+		return 1;
+	}
+
+	TConstArrayView<RowHandle> FQueryContext_DirectSingleApi::GetBatchRowHandles() const
+	{
+		return TConstArrayView<RowHandle>(&CurrentRow, 1);
+	}
+
+	bool FQueryContext_DirectSingleApi::CurrentBatchTableHasColumns(TConstArrayView<const UScriptStruct*> Columns) const
+	{
+		return CurrentRowHasColumns(Columns);
+	}
+
 } // namespace UE::Editor::DataStorage::Queries
\ No newline at end of file
diff --git a/Engine/Plugins/Experimental/EditorDataStorage/Source/TedsCore/Private/Queries/TypedElementQueryContext_DirectSingleApi.h b/Engine/Plugins/Experimental/EditorDataStorage/Source/TedsCore/Private/Queries/TypedElementQueryContext_DirectSingleApi.h
--- a/Engine/Plugins/Experimental/EditorDataStorage/Source/TedsCore/Private/Queries/TypedElementQueryContext_DirectSingleApi.h
+++ b/Engine/Plugins/Experimental/EditorDataStorage/Source/TedsCore/Private/Queries/TypedElementQueryContext_DirectSingleApi.h
@@ -9,8 +9,13 @@ class UEditorDataStorage;
 
 namespace UE::Editor::DataStorage::Queries
 {
-	/** Query context implementation for single row queries that call into TEDS Direct API. */
-	class FQueryContext_DirectSingleApi final : public SingleRowInfo, public IQueryFunctionResponse
+	/** 
+	 * Query context implementation for single row queries that call into TEDS Direct API.
+	 * Note: batch operations, for instance RowBatchInfo, are supported by this context, but
+	 * will in practice only run a single row. Batch mode was added for compatibility and offers
+	 * no performance benefits when using this context.
+	 */
+	class FQueryContext_DirectSingleApi final : public SingleRowInfo, public RowBatchInfo, public IQueryFunctionResponse
 	{
 	public:
 		static const FName Name;
@@ -24,9 +29,18 @@ namespace UE::Editor::DataStorage::Queries
 		virtual bool CurrentRowHasColumns(TConstArrayView<const UScriptStruct*> ColumnTypes) const override;
 		virtual RowHandle GetCurrentRow() const override;
 
+		/**
+		 * @section IRowBatchInfo
+		 */
+		virtual uint32 GetBatchRowCount() const override;
+		virtual TConstArrayView<RowHandle> GetBatchRowHandles() const override;
+		virtual bool CurrentBatchTableHasColumns(TConstArrayView<const UScriptStruct*> Columns) const override;
+
 		/**
 		 * @section IQueryFunctionResponse
 		 */
+		virtual bool NextBatch() override;
+		virtual bool NextRow() override;
 		virtual void GetConstColumns(TArrayView<const void*> ColumnsData, TConstArrayView<const UScriptStruct*> ColumnTypes) override;
 		virtual void GetMutableColumns(TArrayView<void*> ColumnsData, TConstArrayView<const UScriptStruct*> ColumnTypes) override;
 
diff --git a/Engine/Plugins/Experimental/EditorDataStorage/Source/TedsCore/Private/TypedElementDatabase.cpp b/Engine/Plugins/Experimental/EditorDataStorage/Source/TedsCore/Private/TypedElementDatabase.cpp
--- a/Engine/Plugins/Experimental/EditorDataStorage/Source/TedsCore/Private/TypedElementDatabase.cpp
+++ b/Engine/Plugins/Experimental/EditorDataStorage/Source/TedsCore/Private/TypedElementDatabase.cpp
@@ -582,6 +582,15 @@ void UEditorDataStorage::FilterRowsBy(UE::Editor::DataStorage::FRowHandleArray&
 
 	if (ActiveEditorEntityManager)
 	{
+		struct FQueryResult : TResult<bool>
+		{
+			bool bResult = true;
+			virtual void Add(bool ResultValue) override
+			{
+				bResult = ResultValue;
+			}
+		};
+
 		QueryContext_DirectSingleApi Context(*this);
 		Context.CheckCompatiblity(Filter);
 
@@ -589,7 +598,9 @@ void UEditorDataStorage::FilterRowsBy(UE::Editor::DataStorage::FRowHandleArray&
 		for (RowHandle Row : Input)
 		{
 			Context.GetContextImplementation().SetCurrentRow(Row);
-			if (Filter.Call<EFunctionCallConfig::VerifyColumns>(Context, Context.GetContextImplementation()) == bAcceptedResult)
+			FQueryResult QueryResult;
+			Filter.Call<EFunctionCallConfig::VerifyColumns>(QueryResult, Context, Context.GetContextImplementation());
+			if (QueryResult.bResult == bAcceptedResult)
 			{
 				Result.Add(Row);
 			}
diff --git a/Engine/Plugins/Experimental/EditorDataStorageFeatures/Source/TedsTableViewer/Private/Widgets/STedsSearchBox.cpp b/Engine/Plugins/Experimental/EditorDataStorageFeatures/Source/TedsTableViewer/Private/Widgets/STedsSearchBox.cpp
--- a/Engine/Plugins/Experimental/EditorDataStorageFeatures/Source/TedsTableViewer/Private/Widgets/STedsSearchBox.cpp
+++ b/Engine/Plugins/Experimental/EditorDataStorageFeatures/Source/TedsTableViewer/Private/Widgets/STedsSearchBox.cpp
@@ -36,7 +36,7 @@ namespace UE::Editor::DataStorage
 		if(InArgs._OutSearchNode)
 		{
 			SearchNode = MakeShared<QueryStack::FRowFilterNode>(Storage, InArgs._InSearchableRowNode,
-				[this](TQueryContext<SingleRowInfo>& Context, const FTypedElementLabelColumn& LabelColumn)
+				[this](TQueryContext<SingleRowInfo> Context, const FTypedElementLabelColumn& LabelColumn)
 				{
 					if(SearchText.IsEmpty())
 					{
diff --git a/Engine/Source/Runtime/TypedElementFramework/Public/Elements/Framework/TypedElementQueryCapabilities.h b/Engine/Source/Runtime/TypedElementFramework/Public/Elements/Framework/TypedElementQueryCapabilities.h
--- a/Engine/Source/Runtime/TypedElementFramework/Public/Elements/Framework/TypedElementQueryCapabilities.h
+++ b/Engine/Source/Runtime/TypedElementFramework/Public/Elements/Framework/TypedElementQueryCapabilities.h
@@ -5,12 +5,20 @@
 #include <type_traits>
 #include "Containers/ContainersFwd.h"
 #include "DataStorage/Handles.h"
+#include "Misc/EnumClassFlags.h"
 #include "UObject/NameTypes.h"
 
 class UScriptStruct;
 
 namespace UE::Editor::DataStorage::Queries
 {
+	enum class EContextCapabilityFlags
+	{
+		SupportsSingle = 1 << 0,
+		SupportsBatch = 1 << 1
+	};
+	ENUM_CLASS_FLAGS(EContextCapabilityFlags)
+
 	/** 
 	 * Class fragments used to composite a context.
 	 * A query context is created by combining various context capabilities together. Users can request a query context
@@ -38,12 +46,13 @@ namespace UE::Editor::DataStorage::Queries
 #define ConstFunction3(Capability, Return, Function, Arg1, Arg2, Arg3) virtual Return Function(ArgTypeName Arg1 , ArgTypeName Arg2 , ArgTypeName Arg3 ) const = 0;
 #define ConstFunction4(Capability, Return, Function, Arg1, Arg2, Arg3, Arg4) virtual Return Function(ArgTypeName Arg1 , ArgTypeName Arg2 , ArgTypeName Arg3 , ArgTypeName Arg4 ) const = 0;
 
-#define CapabilityStart(Capability) \
+#define CapabilityStart(InCapability, InFlags) \
 	template<typename Base> \
-	struct I##Capability : Base \
+	struct I##InCapability : Base \
 	{ \
-		inline static const FName Name = #Capability ; \
-		virtual ~I##Capability () override = default;
+		inline static const FName Name = #InCapability ; \
+		static constexpr EContextCapabilityFlags Flags = InFlags; \
+		virtual ~I##InCapability () override = default;
 
 #define CapabilityEnd(Capability) \
 	}; \
diff --git a/Engine/Source/Runtime/TypedElementFramework/Public/Elements/Framework/TypedElementQueryCapabilities.inl b/Engine/Source/Runtime/TypedElementFramework/Public/Elements/Framework/TypedElementQueryCapabilities.inl
--- a/Engine/Source/Runtime/TypedElementFramework/Public/Elements/Framework/TypedElementQueryCapabilities.inl
+++ b/Engine/Source/Runtime/TypedElementFramework/Public/Elements/Framework/TypedElementQueryCapabilities.inl
@@ -8,9 +8,9 @@
 /**
  * Capability to provide information about the currently active row.
  */
-CapabilityStart(SingleRowInfo)
+CapabilityStart(SingleRowInfo, EContextCapabilityFlags::SupportsSingle)
 	/*
-	 * Returns whether a column matches the requested type or not. This version only applies to current row. This version is faster
+	 * Returns whether a column matches the requested type or not. This version only applies to current row. This version is faster than
 	 * querying for an arbitrary row.
 	 */
 	ConstFunction1(SingleRowInfo, bool, CurrentRowHasColumns, (TConstArrayView<const UScriptStruct*>, Columns))
@@ -19,7 +19,7 @@ CapabilityStart(SingleRowInfo)
 
 #if defined(WithWrappers)
 	/*
-	 * Returns whether a column matches the requested type or not. This version only applies to current row. This version is faster
+	 * Returns whether a column matches the requested type or not. This version only applies to current row. This version is faster than
 	 * querying for an arbitrary row.
 	 */
 	template<TColumnType... ColumnTypes>
@@ -33,7 +33,7 @@ CapabilityEnd(SingleRowInfo)
 /**
  * Capability to provide information about the currently active batch.
  */
-CapabilityStart(RowBatchInfo)
+CapabilityStart(RowBatchInfo, EContextCapabilityFlags::SupportsBatch)
 	/** Returns the number of rows in the current batch. */
 	ConstFunction0(RowBatchInfo, uint32, GetBatchRowCount)
 	/** Returns an view with the rows used by this batch. */
@@ -48,5 +48,19 @@ CapabilityStart(RowBatchInfo)
 	{
 		return this->CurrentBatchTableHasColumns(TConstArrayView<const UScriptStruct*>({ ColumnTypes::StaticStruct()... }));
 	}
+
+	template<typename CallbackType, typename... Args>
+	void ForEachRow(CallbackType&& Callback, TBatch<Args>... Batches)
+	{
+		uint32 RowCount = this->GetBatchRowCount();
+		TTuple<Args*...> AddressTuple = MakeTuple(Batches.GetData()...);
+		for (uint32 Index = 0; Index < RowCount; ++Index)
+		{
+			AddressTuple.ApplyBefore([&](Args*&... Pointers)
+				{
+					Callback((*Pointers++)...);
+				});
+		}
+	}
 #endif
 CapabilityEnd(RowBatchInfo)
\ No newline at end of file
diff --git a/Engine/Source/Runtime/TypedElementFramework/Public/Elements/Framework/TypedElementQueryCapabilityForwarder.h b/Engine/Source/Runtime/TypedElementFramework/Public/Elements/Framework/TypedElementQueryCapabilityForwarder.h
--- a/Engine/Source/Runtime/TypedElementFramework/Public/Elements/Framework/TypedElementQueryCapabilityForwarder.h
+++ b/Engine/Source/Runtime/TypedElementFramework/Public/Elements/Framework/TypedElementQueryCapabilityForwarder.h
@@ -5,6 +5,7 @@
 #include "DataStorage/CommonTypes.h"
 #include "Elements/Framework/TypedElementQueryCapabilities.h"
 #include "Elements/Framework/TypedElementQueryContract.h"
+#include "Elements/Framework/TypedElementQueryFunctionArgumentTypes.h"
 
 namespace UE::Editor::DataStorage::Queries::Private
 {
@@ -65,7 +66,7 @@ namespace UE::Editor::DataStorage::Queries::Private
 #define ConstFunction4(Capability, Return, Function, Arg1, Arg2, Arg3, Arg4) Return Function(ArgTypeName Arg1 , ArgTypeName Arg2 , ArgTypeName Arg3 , ArgTypeName Arg4 ) const \
 	{ return this->Contract-> Function (ArgName Arg1 , ArgName Arg2 , ArgName Arg3 , ArgName Arg4); }
 
-#define CapabilityStart(Capability) \
+#define CapabilityStart(Capability, Flags) \
 	template<typename... Capabilities> \
 	struct TForwarder< Capability , Capabilities...> : TForwarder<Capabilities...> \
 	{ \
diff --git a/Engine/Source/Runtime/TypedElementFramework/Public/Elements/Framework/TypedElementQueryContext.h b/Engine/Source/Runtime/TypedElementFramework/Public/Elements/Framework/TypedElementQueryContext.h
--- a/Engine/Source/Runtime/TypedElementFramework/Public/Elements/Framework/TypedElementQueryContext.h
+++ b/Engine/Source/Runtime/TypedElementFramework/Public/Elements/Framework/TypedElementQueryContext.h
@@ -14,7 +14,11 @@ namespace UE::Editor::DataStorage::Queries
 	concept ContextCapability = std::is_base_of_v<IContextCapability, T> && !std::is_same_v<IContextCapability, T>;
 
 	template<ContextCapability... Capabilities>
-	struct TCapabilityStore : Capabilities... {};
+	struct TCapabilityStore : Capabilities... 
+	{
+		static constexpr bool bIsSingle = (EnumHasAnyFlags(Capabilities::Flags, EContextCapabilityFlags::SupportsSingle) && ...);
+		static constexpr bool bIsBatch = (EnumHasAnyFlags(Capabilities::Flags, EContextCapabilityFlags::SupportsBatch) && ...);
+	};
 
 	/** 
 	 * Template to composite a query context using context capabilities.
@@ -28,6 +32,9 @@ namespace UE::Editor::DataStorage::Queries
 	{
 		using Capabilities = TCapabilityStore<CapabilityTypes...>;
 
+		static_assert(Capabilities::bIsSingle || Capabilities::bIsBatch,
+			"One or more capabilities that only support single or batch processing were mixed.");
+
 		TQueryContext() = default;
 		TQueryContext(IContextContract& Contract);
 	};
diff --git a/Engine/Source/Runtime/TypedElementFramework/Public/Elements/Framework/TypedElementQueryContextImplementation.h b/Engine/Source/Runtime/TypedElementFramework/Public/Elements/Framework/TypedElementQueryContextImplementation.h
--- a/Engine/Source/Runtime/TypedElementFramework/Public/Elements/Framework/TypedElementQueryContextImplementation.h
+++ b/Engine/Source/Runtime/TypedElementFramework/Public/Elements/Framework/TypedElementQueryContextImplementation.h
@@ -6,7 +6,7 @@
 #include "Elements/Framework/TypedElementQueryContract.h"
 #include "Elements/Framework/TypedElementQueryContext.h"
 #include "Elements/Framework/TypedElementQueryFunctions.h"
-
+#include "Elements/Framework/TypedElementQueryFunctionArgumentTypes.h"
 
 #define ArgTypeName(Type, Name) Type Name
 
@@ -22,7 +22,7 @@
 #define ConstFunction3(Capability, Return, Function, Arg1, Arg2, Arg3) virtual Return Function(ArgTypeName Arg1 , ArgTypeName Arg2 , ArgTypeName Arg3 ) const override;
 #define ConstFunction4(Capability, Return, Function, Arg1, Arg2, Arg3, Arg4) virtual Return Function(ArgTypeName Arg1 , ArgTypeName Arg2 , ArgTypeName Arg3 , ArgTypeName Arg4 ) const override;
 
-#define CapabilityStart(Capability)
+#define CapabilityStart(Capability, Flags)
 #define CapabilityEnd(Capability)
 
 #define WithWrappers 1
diff --git a/Engine/Source/Runtime/TypedElementFramework/Public/Elements/Framework/TypedElementQueryContextImplementation.inl b/Engine/Source/Runtime/TypedElementFramework/Public/Elements/Framework/TypedElementQueryContextImplementation.inl
--- a/Engine/Source/Runtime/TypedElementFramework/Public/Elements/Framework/TypedElementQueryContextImplementation.inl
+++ b/Engine/Source/Runtime/TypedElementFramework/Public/Elements/Framework/TypedElementQueryContextImplementation.inl
@@ -123,21 +123,9 @@ namespace UE::Editor::DataStorage::Queries
 		CallFunction(Capability, Return, Function, ArgName Arg1, ArgName Arg2, ArgName Arg3, ArgName Arg4 ) \
 	}
 
-#define CapabilityStart(Capability)
+#define CapabilityStart(Capability, Flags)
 #define CapabilityEnd(Capability)
 
-
-#define CallFunction(Capability, ReturnValue, FunctionName, ...) \
-	if constexpr (SupportsCapability< Capability >()) \
-	{ \
-		return Implementation.FunctionName( __VA_ARGS__ ); \
-	} \
-	else \
-	{ \
-		checkf(false, TEXT("Function '" #FunctionName "' in capability '" #Capability"' is not supported by the current query context implementation.")); \
-		return ReturnValue {} ; \
-	}
-
 #include "Elements/Framework/TypedElementQueryCapabilities.inl"
 
 #undef ArgTypeName
diff --git a/Engine/Source/Runtime/TypedElementFramework/Public/Elements/Framework/TypedElementQueryContract.h b/Engine/Source/Runtime/TypedElementFramework/Public/Elements/Framework/TypedElementQueryContract.h
--- a/Engine/Source/Runtime/TypedElementFramework/Public/Elements/Framework/TypedElementQueryContract.h
+++ b/Engine/Source/Runtime/TypedElementFramework/Public/Elements/Framework/TypedElementQueryContract.h
@@ -96,7 +96,7 @@ namespace UE::Editor::DataStorage::Queries
 #define ConstFunction3(Capability, Return, Function, Arg1, Arg2, Arg3)
 #define ConstFunction4(Capability, Return, Function, Arg1, Arg2, Arg3, Arg4)
 
-#define CapabilityStart(Capability) , I##Capability
+#define CapabilityStart(Capability, Flags) , I##Capability
 #define CapabilityEnd(Capability)
 
 	/** 
diff --git a/Engine/Source/Runtime/TypedElementFramework/Public/Elements/Framework/TypedElementQueryFunctionArgumentTypes.h b/Engine/Source/Runtime/TypedElementFramework/Public/Elements/Framework/TypedElementQueryFunctionArgumentTypes.h
new file mode 100644
--- /dev/null
+++ b/Engine/Source/Runtime/TypedElementFramework/Public/Elements/Framework/TypedElementQueryFunctionArgumentTypes.h
@@ -0,0 +1,83 @@
+// Copyright Epic Games, Inc. All Rights Reserved.
+
+#pragma once
+
+#include "Containers/ArrayView.h"
+
+namespace UE::Editor::DataStorage::Queries
+{
+	/**
+	 * This file contains the arguments that can be used in functions for Query Callbacks. In addition
+	 * to the options in this file, the following arguments can be added as well:
+	 *	- TQueryContext - Addition is typically required to interact with TEDS, e.g. to add/remove columns.
+	 *	- Columns - When using the single row mode any column with data can be added by reference and it will
+	 *		automatically be bound. Note that the use of const is important as it will determine if read-only
+	 *		or read-write access is requested from TEDS, which can have performance implications as read-only 
+	 *		generally provides better parallelization.
+	 */
+	 
+	 // Forward declarations.
+	namespace Private
+	{
+		template<typename... Args>
+		struct TArgumentInfo;
+	}
+
+	/** 
+	 * Optional argument that allows controlling the flow of the query callback.
+	 * This can be used to early out if for example a searched for value is found or if the number of
+	 * processed rows has reached a predetermined limit.
+	 */
+	enum class EFlowControl
+	{
+		Continue, /** Continue processing with the next row/batch. */
+		Break /** Stop processing at the earliest possible breakpoint. */
+	};
+
+	/**
+	 * Functions that return anything other than void need to provide an interface to collect results into.
+	 * It'll be up to the implementation to determine what to do with the results, such as track each registered
+	 * value or accumulate all results together. Query Callback functions that operate in single mode can return
+	 * the result type from the function and the result will automatically be added. Query Callback functions that
+	 * work in batch mode require having a TResult<...>& argument where results are written to.
+	 */
+	template<typename ResultType>
+	struct TResult
+	{
+		virtual void Add(ResultType ResultValue) = 0;
+	};
+	template<> struct TResult<void>{}; // Used as a placeholder for internal use.
+
+	/** 
+	 * Storage for batch of objects used in a query.
+	 * Batches provide the following guarantees:
+	 *	- Contiguous in memory up to the number of entries in the batch.
+	 *	- All values apply to the same table.
+	 */
+	template<typename T>
+	class TBatch
+	{
+		template<typename... Args>
+		friend struct Private::TArgumentInfo;
+
+	public:
+		TBatch() = default;
+
+		/** Provides direct access to the stored objects. */
+		TArrayView<T> GetView(int32 BatchSize) { return TArrayView<T>(Values, BatchSize); };
+		
+		/** Returns a pointer to the stored column. */
+		T* GetData() { return Values; }
+		/** Returns a pointer to the stored column. */
+		const T* GetData() const { return Values; }
+
+	private:
+		TBatch(T* Values) : Values(Values) {}
+		TBatch& operator=(T* InValues) { Values = InValues; return *this; }
+
+		T* Values = nullptr;
+	};
+
+	template<typename T>
+	using TConstBatch = TBatch<const T>;
+} // namespace UE::Editor::DataStorage::Queries
diff --git a/Engine/Source/Runtime/TypedElementFramework/Public/Elements/Framework/TypedElementQueryFunctionArguments.h b/Engine/Source/Runtime/TypedElementFramework/Public/Elements/Framework/TypedElementQueryFunctionArguments.h
--- a/Engine/Source/Runtime/TypedElementFramework/Public/Elements/Framework/TypedElementQueryFunctionArguments.h
+++ b/Engine/Source/Runtime/TypedElementFramework/Public/Elements/Framework/TypedElementQueryFunctionArguments.h
@@ -2,9 +2,11 @@
 
 #pragma once
 
+#include <type_traits>
 #include "Containers/ContainersFwd.h"
 #include "DataStorage/CommonTypes.h"
 #include "Elements/Framework/TypedElementQueryCapabilityForwarder.h"
+#include "Elements/Framework/TypedElementQueryFunctionArgumentTypes.h"
 #include "Misc/EnumClassFlags.h"
 #include "Templates/Tuple.h"
 
@@ -12,20 +14,28 @@ class UScriptStruct;
 
 namespace UE::Editor::DataStorage::Queries::Private
 {
+	template<typename T>
+	using UnreferencedType = std::remove_reference_t<std::remove_pointer_t<T>>;
+	template<typename T>
+	using FoundationalType = std::remove_cv_t<UnreferencedType<T>>;
 	template<typename T> concept ContextType = std::is_base_of_v<FQueryContext, T>;
 
 	enum class EArgumentFlags
 	{
-		IsConst = 1 << 0,
-		IsMutable = 1 << 1,
-		BatchArgument = 1 << 2,
-
-		Type_Context = 1 << 3,
-		Type_RowHandle = 1 << 4,
-		Type_Column = 1 << 5,
+		IsConst = 1 << 0,			// Indicates the data is immutable.
+		IsMutable = 1 << 1,			// Indicates that the data can change.
+		SingleArgument = 1 << 2,	// The argument can be used when a single row is being processed.
+		BatchArgument = 1 << 3,		// The argument can be used when a batch of rows is being processed.
+
+		Type_Result = 1 << 4,		// The argument is used to pass back results.
+		Type_Context = 1 << 5,		// The argument contains contextual information about the current row or batch.
+		Type_RowHandle = 1 << 6,	// The argument stores a row handle.
+		Type_Column = 1 << 7,		// The argument stores a column, including data columns and tag columns.
+		Type_FlowControl = 1 << 8,	// The argument can be used to control the next iteration step, including stopping work.
 	};
 	ENUM_CLASS_FLAGS(EArgumentFlags)
 
+	// Wrapper around a pointer to act as a reference.
 	template<typename T>
 	struct TPointerForwarder
 	{
@@ -34,69 +44,146 @@ namespace UE::Editor::DataStorage::Queries::Private
 		TPointerForwarder& operator=(T* InValue) { Value = InValue; return *this; }
 
 		operator T& () const { return *Value; }
-		T* operator->() const { return Value; }
-
+		void Increment() { Value++; }
+		
 		T* Value = nullptr;
 	};
-	
+
+	// Base descriptor for query callback function arguments. Not directly usable and requires specializations.
 	template<typename T>
 	struct TArgument
 	{
 		static_assert(!std::is_same_v<T, void>, "Unsupported argument type for query callback.");
 	};
 
-	// Context type
-	template<typename T> requires ContextType<T>
-	struct TArgument<const T&>
+	// Void type (used as a placeholder)
+	template<>
+	struct TArgument<void>
 	{
-		constexpr static EArgumentFlags Flags = EArgumentFlags::IsConst | EArgumentFlags::Type_Context;
-		using ArgumentType = T;
+		constexpr static EArgumentFlags Flags = EArgumentFlags(0);
+		using BaseType = void;
+		using ArgumentType = void;
+		using ColumnType = void;
+		using PointerCastType = void;
+	};
+
+	// Result type
+	template<typename T>
+	struct ResultTypeInfo
+	{
+		static constexpr bool bIsValid = false;
+		using BaseType = void;
+	};
+	template<typename T>
+	struct ResultTypeInfo<TResult<T>>
+	{
+		static constexpr bool bIsValid = true;
+		using BaseType = T;
+	};
+
+	template<typename T> requires ResultTypeInfo<FoundationalType<T>>::bIsValid
+	struct TArgument<T>
+	{
+		static_assert(std::is_reference_v<T>, "Results can only be accessed by reference.");
+		static_assert(!std::is_const_v<T>, "Results requires write access to function.");
+
+		constexpr static EArgumentFlags Flags =
+			EArgumentFlags::SingleArgument | EArgumentFlags::BatchArgument | EArgumentFlags::IsMutable | EArgumentFlags::Type_Result;
+		using BaseType = typename ResultTypeInfo<FoundationalType<T>>::BaseType;
+		using ArgumentType = TPointerForwarder<FoundationalType<T>>;
 		using ColumnType = void;
 		using PointerCastType = void*;
 	};
 
-	template<typename T> requires ContextType<T>
-	struct TArgument<T&>
+	// Context type
+	template<typename T> requires ContextType<FoundationalType<T>>
+	struct TArgument<T>
 	{
-		constexpr static EArgumentFlags Flags = EArgumentFlags::IsMutable | EArgumentFlags::Type_Context;
-		using ArgumentType = T;
+		static_assert(!std::is_reference_v<T> && !std::is_pointer_v<T>, "The contexts can only be passed in by value.");
+		
+		using BaseType = FoundationalType<T>;
+		using ContextCapabilities = BaseType::Capabilities;
+
+		constexpr static EArgumentFlags Flags =
+			(ContextCapabilities::bIsSingle ? EArgumentFlags::SingleArgument : EArgumentFlags(0)) |
+			(ContextCapabilities::bIsBatch ? EArgumentFlags::BatchArgument : EArgumentFlags(0)) |
+			(std::is_const_v<UnreferencedType<T>> ? EArgumentFlags::IsConst : EArgumentFlags::IsMutable) |
+			EArgumentFlags::Type_Context; 
+		using ArgumentType = BaseType; // Take by value because the context is a forwarder that only stores a single pointer.
 		using ColumnType = void;
 		using PointerCastType = void*;
 	};
 
-	// Column types
+	// Column reference types
+	template<typename T> requires TColumnType<FoundationalType<T>>
+	struct TArgument<T>
+	{
+		static_assert(!std::is_pointer_v<T>, "Pointers to columns to work with batches are no longer supported. Use TBatch instead.");
+		static_assert(std::is_reference_v<T>, "Columns can only be accessed by reference.");
+
+		static constexpr bool bIsConst = std::is_const_v<UnreferencedType<T>>;
+
+		constexpr static EArgumentFlags Flags =
+			EArgumentFlags::SingleArgument |
+			(bIsConst ? EArgumentFlags::IsConst : EArgumentFlags::IsMutable) |
+			EArgumentFlags::Type_Column;
+		using BaseType = FoundationalType<T>;
+		using ArgumentType = TPointerForwarder<UnreferencedType<T>>;
+		using ColumnType = BaseType;
+		using PointerCastType = std::conditional_t<bIsConst, const BaseType*, BaseType*>;
+	};
 
+	// Column batch types
 	template<typename T> requires TColumnType<T>
-	struct TArgument<const T&>
+	struct TArgument<TBatch<T>>
 	{
-		constexpr static EArgumentFlags Flags = EArgumentFlags::IsConst | EArgumentFlags::Type_Column;
-		using ArgumentType = TPointerForwarder<const T>;
-		using ColumnType = T;
-		using PointerCastType = const T*;
+		static constexpr bool bIsConst = std::is_const_v<UnreferencedType<T>>;
+		
+		constexpr static EArgumentFlags Flags =
+			EArgumentFlags::BatchArgument | 
+			(bIsConst ? EArgumentFlags::IsConst : EArgumentFlags::IsMutable) |
+			EArgumentFlags::Type_Column;
+		using BaseType = T;
+		using ArgumentType = TBatch<T>;
+		using ColumnType = BaseType;
+		using PointerCastType = std::conditional_t<bIsConst, const BaseType*, BaseType*>;
 	};
 
-	template<typename T> requires TColumnType<T>
-	struct TArgument<T&>
+	// Flow control
+	template<typename T> requires std::is_same_v<FoundationalType<T>, EFlowControl>
+	struct TArgument<T>
 	{
-		constexpr static EArgumentFlags Flags = EArgumentFlags::IsMutable | EArgumentFlags::Type_Column;
-		using ArgumentType = TPointerForwarder<T>;
-		using ColumnType = T;
-		using PointerCastType = T*;
+		static_assert(std::is_reference_v<T>, "Flow control can only be accessed by reference.");
+		static_assert(!std::is_const_v<T>, "Flow control requires write access to function.");
+
+		constexpr static EArgumentFlags Flags = 
+			EArgumentFlags::SingleArgument | EArgumentFlags::BatchArgument | EArgumentFlags::IsMutable | EArgumentFlags::Type_FlowControl;
+		using BaseType = FoundationalType<T>;
+		using ArgumentType = TPointerForwarder<BaseType>;
+		using ColumnType = void;
+		using PointerCastType = BaseType*;
 	};
 
-	template<int32 Index, typename Tuple>
-	struct TIndexToTypeImpl
+	template<int32 Index, typename... Args>
+	struct TIndexToArgInfoImpl {};
+
+	template<int32 Index, typename Front, typename... Args>
+	struct TIndexToArgInfoImpl<Index, Front, Args...>
 	{
-		using Type = typename TTupleElement<Index, Tuple>::Type;
+		using Type = 
+			std::conditional_t<Index < 0, TArgument<void>,
+				std::conditional_t<Index == 0, TArgument<Front>,
+					typename TIndexToArgInfoImpl<Index - 1, Args...>::Type>>;
 	};
 
-	template<typename Tuple>
-	struct TIndexToTypeImpl<-1, Tuple>
+	template<int32 Index>
+	struct TIndexToArgInfoImpl<Index> // Empty argument list.
 	{
-		using Type = void;
+		using Type = TArgument<void>;
 	};
-	template<int32 Index, typename Tuple>
-	using TIndexToType = typename TIndexToTypeImpl<Index, Tuple>::Type;
+
+	template<int32 Index, typename... Args>
+	using TIndexToArgInfo = typename TIndexToArgInfoImpl<Index, Args...>::Type;
 
 	template<typename... Args>
 	struct TArgumentInfo
@@ -106,17 +193,33 @@ namespace UE::Editor::DataStorage::Queries::Private
 		constexpr static int32 FirstIndexOfType();
 
 		template<EArgumentFlags Flags>
-		constexpr static int32 CountColumns();
+		constexpr static int32 CountFlags();
 
 		template<EArgumentFlags Flags>
 		static TConstArrayView<const UScriptStruct*> ListColumns();
 	
 	public:
+		static constexpr bool bIsSingle = (EnumHasAnyFlags(TArgument<Args>::Flags, EArgumentFlags::SingleArgument) && ...);
+		static constexpr bool bIsBatch = (EnumHasAnyFlags(TArgument<Args>::Flags, EArgumentFlags::BatchArgument) && ...);
+
+		static_assert(bIsSingle || bIsBatch, "One or more query callback arguments for single or batch processing were mixed.");
+
 		using ArgumentList = TTuple<typename TArgument<Args>::ArgumentType...>;
 
+		constexpr static int32 ResultIndex = FirstIndexOfType<EArgumentFlags::Type_Result>();
+		using ResultType = typename TIndexToArgInfo<ResultIndex, Args...>::BaseType;
+		template<typename T>
+		static void SetResult(ArgumentList& Arguments, TResult<T>& Result);
+		constexpr static int32 CountResults();
+
 		constexpr static int32 ContextIndex = FirstIndexOfType<EArgumentFlags::Type_Context>();
-		using ContextType = TIndexToType<ContextIndex, ArgumentList>;
+		using ContextType = typename TIndexToArgInfo<ContextIndex, Args...>::BaseType;
 		static void SetContext(ArgumentList& Arguments, IContextContract& Contract);
+		constexpr static int32 CountContexts();
+
+		constexpr static int32 FlowIndex = FirstIndexOfType<EArgumentFlags::Type_FlowControl>();
+		static void SetFlowControl(ArgumentList& Arguments, EFlowControl& FlowControl);
+		constexpr static int32 CountFlowControls();
 
 		constexpr static int32 CountConstColumns();
 		constexpr static int32 CountMutableColumns();
@@ -124,6 +227,15 @@ namespace UE::Editor::DataStorage::Queries::Private
 		static TConstArrayView<const UScriptStruct*> ListMutableColumns();
 		static void SetConstColumns(ArgumentList& Arguments, TConstArrayView<const void*> Columns);
 		static void SetMutableColumns(ArgumentList& Arguments, TConstArrayView<void*> Columns);
+		static void IncrementColumns(ArgumentList& Arguments);
+
+	private:
+		template<int32 Index, typename T>
+		static void SetResultUnguarded(ArgumentList& Arguments, TResult<T>& Result);
+		template<int32 Index>
+		static void SetContextUnguarded(ArgumentList& Arguments, IContextContract& Contract);
+		template<int32 Index>
+		static void SetFlowControlUnguarded(ArgumentList& Arguments, EFlowControl& FlowControl);
 	};
 } // namespace UE::Editor::DataStorage::Queries::Private
 
diff --git a/Engine/Source/Runtime/TypedElementFramework/Public/Elements/Framework/TypedElementQueryFunctionArguments.inl b/Engine/Source/Runtime/TypedElementFramework/Public/Elements/Framework/TypedElementQueryFunctionArguments.inl
--- a/Engine/Source/Runtime/TypedElementFramework/Public/Elements/Framework/TypedElementQueryFunctionArguments.inl
+++ b/Engine/Source/Runtime/TypedElementFramework/Public/Elements/Framework/TypedElementQueryFunctionArguments.inl
@@ -28,39 +28,86 @@ namespace UE::Editor::DataStorage::Queries::Private
 	template<EArgumentFlags TypeFlag>
 	constexpr int32 TArgumentInfo<Args...>::FirstIndexOfType()
 	{
-		return FirstIndexOfTypeImpl<TypeFlag, 0, Args...>();
+		if constexpr (sizeof...(Args) > 0)
+		{
+			return FirstIndexOfTypeImpl<TypeFlag, 0, Args...>();
+		}
+		else
+		{
+			return INDEX_NONE;
+		}
+	}
+
+	template<typename... Args>
+	template<EArgumentFlags Flags>
+	constexpr int32 TArgumentInfo<Args...>::CountFlags()
+	{
+		return (0 + ... + (EnumHasAllFlags(TArgument<Args>::Flags, Flags) ? 1 : 0));
+	}
+
+	template<typename... Args>
+	template<typename T>
+	void TArgumentInfo<Args...>::SetResult(ArgumentList& Arguments, TResult<T>& Result)
+	{
+		if constexpr (ResultIndex >= 0)
+		{
+			SetResultUnguarded<ResultIndex>(Arguments, Result);
+		}
+	}
+
+	template<typename... Args>
+	constexpr int32 TArgumentInfo<Args...>::CountResults()
+	{
+		return CountFlags<EArgumentFlags::Type_Result>();
 	}
 
 	template<typename... Args>
 	void TArgumentInfo<Args...>::SetContext(ArgumentList& Arguments, IContextContract& Contract)
 	{
-		Arguments.template Get<ContextIndex>() = typename TTupleElement<ContextIndex, ArgumentList>::Type(Contract);
+		if constexpr (ContextIndex >= 0)
+		{
+			SetContextUnguarded<ContextIndex>(Arguments, Contract);
+		}
 	}
 
 	template<typename... Args>
-	template<EArgumentFlags Flags>
-	constexpr int32 TArgumentInfo<Args...>::CountColumns()
+	constexpr int32 TArgumentInfo<Args...>::CountContexts()
 	{
-		return (0 + ... + (EnumHasAllFlags(TArgument<Args>::Flags, Flags) ? 1 : 0));
+		return CountFlags<EArgumentFlags::Type_Context>();
+	}
+	
+	template<typename... Args>
+	void TArgumentInfo<Args...>::SetFlowControl(ArgumentList& Arguments, EFlowControl& FlowControl)
+	{
+		if constexpr (FlowIndex >= 0)
+		{
+			SetFlowControlUnguarded<FlowIndex>(Arguments, FlowControl);
+		}
+	}
+
+	template<typename... Args>
+	constexpr int32 TArgumentInfo<Args...>::CountFlowControls()
+	{
+		return CountFlags<EArgumentFlags::Type_FlowControl>();
 	}
 
 	template<typename... Args>
 	constexpr int32 TArgumentInfo<Args...>::CountConstColumns()
 	{
-		return CountColumns<EArgumentFlags::IsConst | EArgumentFlags::Type_Column>();
+		return CountFlags<EArgumentFlags::IsConst | EArgumentFlags::Type_Column>();
 	}
 			
 	template<typename... Args>
 	constexpr int32 TArgumentInfo<Args...>::CountMutableColumns()
 	{
-		return CountColumns<EArgumentFlags::IsMutable | EArgumentFlags::Type_Column>();
+		return CountFlags<EArgumentFlags::IsMutable | EArgumentFlags::Type_Column>();
 	}
 
 	template<typename... Args>
 	template<EArgumentFlags Flags>
 	TConstArrayView<const UScriptStruct*> TArgumentInfo<Args...>::ListColumns()
 	{
-		constexpr int32 ColumnCount = CountColumns<Flags>();
+		constexpr int32 ColumnCount = CountFlags<Flags>();
 		if constexpr (ColumnCount > 0)
 		{
 			static TConstArrayView<const UScriptStruct*> Result =
@@ -125,4 +172,39 @@ namespace UE::Editor::DataStorage::Queries::Private
 				}
 			}(), ...);
 	}
+
+	template<typename... Args>
+	void TArgumentInfo<Args...>::IncrementColumns(ArgumentList& Arguments)
+	{
+		([&] {
+			using ArgumentType = TArgument<Args>;
+			if constexpr (EnumHasAllFlags(ArgumentType::Flags, EArgumentFlags::SingleArgument | EArgumentFlags::Type_Column))
+			{
+				Arguments.template Get<typename ArgumentType::ArgumentType>().Increment();
+			}
+			}(), ...);
+	}
+
+	template<typename... Args>
+	template<int32 Index, typename T>
+	void TArgumentInfo<Args...>::SetResultUnguarded(ArgumentList& Arguments, TResult<T>& Result)
+	{
+		static_assert(std::is_same_v<T, typename TIndexToArgInfo<Index, Args...>::BaseType>, 
+			"Result type used doesn't match the result type in the function's arguments.");
+		Arguments.template Get<Index>() = &Result;
+	}
+
+	template<typename... Args>
+	template<int32 Index>
+	void TArgumentInfo<Args...>::SetContextUnguarded(ArgumentList& Arguments, IContextContract& Contract)
+	{
+		Arguments.template Get<Index>() = typename TTupleElement<Index, ArgumentList>::Type(Contract);
+	}
+
+	template<typename... Args>
+	template<int32 Index>
+	void TArgumentInfo<Args...>::SetFlowControlUnguarded(ArgumentList& Arguments, EFlowControl& FlowControl)
+	{
+		Arguments.template Get<Index>() = &FlowControl;
+	}
 } // namespace UE::Editor::DataStorage::Queries::Private
diff --git a/Engine/Source/Runtime/TypedElementFramework/Public/Elements/Framework/TypedElementQueryFunctions.h b/Engine/Source/Runtime/TypedElementFramework/Public/Elements/Framework/TypedElementQueryFunctions.h
--- a/Engine/Source/Runtime/TypedElementFramework/Public/Elements/Framework/TypedElementQueryFunctions.h
+++ b/Engine/Source/Runtime/TypedElementFramework/Public/Elements/Framework/TypedElementQueryFunctions.h
@@ -4,6 +4,7 @@
 
 #include "Containers/ContainersFwd.h"
 #include "Elements/Framework/TypedElementQueryCapabilities.h"
+#include "Elements/Framework/TypedElementQueryFunctionArgumentTypes.h"
 #include "Misc/EnumClassFlags.h"
 #include "Templates/Function.h"
 #include "UObject/NameTypes.h"
@@ -18,6 +19,9 @@ namespace UE::Editor::DataStorage::Queries
 	{
 		virtual ~IQueryFunctionResponse() = default;
 		
+		virtual bool NextBatch() = 0;
+		virtual bool NextRow() = 0;
+
 		virtual void GetConstColumns(TArrayView<const void*> ColumnsData, TConstArrayView<const UScriptStruct*> ColumnTypes) = 0;
 		virtual void GetMutableColumns(TArrayView<void*> ColumnsData, TConstArrayView<const UScriptStruct*> ColumnTypes) = 0;
 	};
@@ -25,22 +29,69 @@ namespace UE::Editor::DataStorage::Queries
 	enum class EFunctionCallConfig
 	{
 		None = 0,
-		/** If not all columns could be retrieved, return the default value for the return type instead of calling the function. */
+		/** If not all columns could be retrieved, return the default value for the return type, if applicable, and don't call the function. */
 		VerifyColumns = 1 << 0,
 	};
 	ENUM_CLASS_FLAGS(EFunctionCallConfig);
 
 	/** Storage for a function that can be used as part of a query.  */
 	template<typename ReturnType>
-	struct TQueryFunction
+	class TQueryFunctionBase
 	{
+	public:
+		using FunctionSpecializationCallback = 
+			bool(*)(IQueryFunctionResponse& Response, TArrayView<const void*> ConstColumns, TArrayView<void*> MutableColumns);
+	
+		using WrapperFunctionType = TFunction<
+			void(
+				TResult<ReturnType>& Result, 
+				IContextContract& Contract, 
+				IQueryFunctionResponse& Response,
+				FunctionSpecializationCallback Specialization)
+		>;
+
 		TConstArrayView<FName> Capabilities;
 		TConstArrayView<const UScriptStruct*> ConstColumnTypes;
 		TConstArrayView<const UScriptStruct*> MutableColumnTypes;
-		TFunction<ReturnType(IContextContract& Contract, TConstArrayView<const void*> ConstColumns, TConstArrayView<void*> MutableColumns)> Function;
+		WrapperFunctionType Function;
+		bool bIsSingleRowProcessor;
+
+	protected:
+		template<EFunctionCallConfig Config>
+		void CallInternal(TResult<ReturnType>& Result, IContextContract& Contract, IQueryFunctionResponse& Response);
+	};
+
+	template<typename ReturnType>
+	class TQueryFunction final : public TQueryFunctionBase<ReturnType>
+	{
+	public:
+		template<EFunctionCallConfig Config>
+		void Call(TResult<ReturnType>& Result, IContextContract& Contract, IQueryFunctionResponse& Response)
+		{
+			this->template CallInternal<Config>(Result, Contract, Response);
+		}
+
+		void Call(TResult<ReturnType>& Result, IContextContract& Contract, IQueryFunctionResponse& Response)
+		{
+			this->template CallInternal<EFunctionCallConfig::None>(Result, Contract, Response);
+		}
+	};
+
+	template<>
+	class TQueryFunction<void> final : public TQueryFunctionBase<void>
+	{
+	public:
+		template<EFunctionCallConfig Config>
+		void Call(IContextContract& Contract, IQueryFunctionResponse& Response)
+		{
+			TResult<void> Dummy;
+			this->template CallInternal<Config>(Dummy, Contract, Response);
+		}
 
-		template<EFunctionCallConfig Config = EFunctionCallConfig::None>
-		ReturnType Call(IContextContract& Contract, IQueryFunctionResponse& Response);
+		void Call(IContextContract& Contract, IQueryFunctionResponse& Response)
+		{
+			this->Call<EFunctionCallConfig::None>(Contract, Response);
+		}
 	};
 
 	template<typename ReturnType, typename FunctionType>
diff --git a/Engine/Source/Runtime/TypedElementFramework/Public/Elements/Framework/TypedElementQueryFunctions.inl b/Engine/Source/Runtime/TypedElementFramework/Public/Elements/Framework/TypedElementQueryFunctions.inl
--- a/Engine/Source/Runtime/TypedElementFramework/Public/Elements/Framework/TypedElementQueryFunctions.inl
+++ b/Engine/Source/Runtime/TypedElementFramework/Public/Elements/Framework/TypedElementQueryFunctions.inl
@@ -71,52 +71,102 @@ namespace UE::Editor::DataStorage::Queries
 		{
 			using Base = TFunctionInfoSelection<FunctorType<FunctionType>, FunctionType>::Type;
 			
-			static constexpr bool bIsValidQueryFunction = Base::bIsValidQueryFunction;
 			using ReturnType = typename Base::ReturnType;
-
 			using ArgumentInfo = typename Base::ArgumentInfo;
 			using ContextType = typename ArgumentInfo::ContextType;
-		};
-	}
-
-	template<typename ReturnType>
-	template<EFunctionCallConfig Config>
-	ReturnType TQueryFunction<ReturnType>::Call(IContextContract& Contract, IQueryFunctionResponse& Response)
-	{
-		const void** ConstColumnsData = static_cast<const void**>(alloca(sizeof(void*) * ConstColumnTypes.Num()));
-		void** MutableColumnsData = static_cast<void**>(alloca(sizeof(void*) * MutableColumnTypes.Num()));
 
-		Response.GetConstColumns(TArrayView<const void*>(ConstColumnsData, ConstColumnTypes.Num()), ConstColumnTypes);
-		Response.GetMutableColumns(TArrayView<void*>(MutableColumnsData, MutableColumnTypes.Num()), MutableColumnTypes);
+			static constexpr bool bIsValidQueryFunction = 
+				Base::bIsValidQueryFunction &&
+				(ArgumentInfo::bIsSingle || ArgumentInfo::bIsBatch);
+			static constexpr bool bIsSingleRowProcessor = ArgumentInfo::bIsSingle;
 
-		if constexpr (EnumHasAnyFlags(Config, EFunctionCallConfig::VerifyColumns))
-		{
-			bool bMissingColumn = false;
-			for (const void* Column : TArrayView<const void*>(ConstColumnsData, ConstColumnTypes.Num()))
+			template<typename RequestedReturnType>
+			static constexpr void StaticValidation()
 			{
-				bMissingColumn = bMissingColumn || Column == nullptr;
+				static_assert(bIsValidQueryFunction, "The provided function is not compatible with a query callback. "
+					"One or more arguments are possibly incompatible.");
+
+				static_assert(ArgumentInfo::CountContexts() <= 1, "Only zero or one context can be added as an argument.");
+				static_assert(ArgumentInfo::CountFlowControls() <= 1, "Only zero or one flow control arguments can be added.");
+
+				if constexpr (std::is_same_v<RequestedReturnType, void>)
+				{
+					static_assert(ArgumentInfo::ResultIndex == INDEX_NONE,
+						"The query function being build has no return value and can't therefore accept a TResult.");
+					static_assert(std::is_same_v<ReturnType, void>,
+						"The return type of the provided query is expected to be void.");
+				}
+				else
+				{
+					if constexpr (bIsSingleRowProcessor)
+					{
+						if constexpr (ArgumentInfo::ResultIndex >= 0)
+						{
+							static_assert(ArgumentInfo::CountResults() == 1, "Only one TResult can be present.");
+							static_assert(std::is_same_v<typename ArgumentInfo::ResultType, RequestedReturnType>,
+								"The type used for TResult is not compatible with the result type expected for the query function.");
+							static_assert(std::is_same_v<ReturnType, void>,
+								"The return type of the provided query is expected to be void when a TResult is provided.");
+						}
+						else
+						{
+							static_assert(std::is_convertible_v<ReturnType, RequestedReturnType>,
+								"The type returned is not compatible with the result type expected for the query function.");
+						}
+					}
+					else
+					{
+						static_assert(ArgumentInfo::ResultIndex >= 0,
+							"Batch processing a function that expects a result it's required that a TResult is used.");
+						static_assert(ArgumentInfo::CountResults() == 1, "Only one TResult can be present.");
+						static_assert(std::is_same_v<typename ArgumentInfo::ResultType, RequestedReturnType>,
+							"The type used for TResult is not compatible with the result type expected for the query function.");
+						static_assert(std::is_same_v<ReturnType, void>, 
+							"The return type of the provided query is expected to be void for batch processing functions.");
+					}
+				}
 			}
-			for (void* Column : TArrayView<void*>(MutableColumnsData, MutableColumnTypes.Num()))
+
+			template<typename RequestedReturnType>
+			static constexpr void SetupResult(TQueryFunction<RequestedReturnType>& Result)
 			{
-				bMissingColumn = bMissingColumn || Column == nullptr;
+				if constexpr (!std::is_same_v<void, ContextType>)
+				{
+					Result.Capabilities = IContextContract::SupportedCapabilitiesList<typename ContextType::Capabilities>();
+				}
+				Result.ConstColumnTypes = ArgumentInfo::ListConstColumns();
+				Result.MutableColumnTypes = ArgumentInfo::ListMutableColumns();
+				Result.bIsSingleRowProcessor = bIsSingleRowProcessor;
 			}
-			if (bMissingColumn)
+		};
+	}
+
+	template<typename ReturnType>
+	template<EFunctionCallConfig Config>
+	void TQueryFunctionBase<ReturnType>::CallInternal(TResult<ReturnType>& Result, IContextContract& Contract, IQueryFunctionResponse& Response)
+	{
+		auto Specialization = [](IQueryFunctionResponse& Response, TArrayView<const void*> ConstColumns, TArrayView<void*> MutableColumns)
 			{
-				if constexpr (std::is_same_v<void, ReturnType>)
+				if constexpr (EnumHasAnyFlags(Config, EFunctionCallConfig::VerifyColumns))
 				{
-					return;
+					bool bMissingColumn = false;
+					for (const void* Column : ConstColumns)
+					{
+						bMissingColumn = bMissingColumn || Column == nullptr;
+					}
+					for (void* Column : MutableColumns)
+					{
+						bMissingColumn = bMissingColumn || Column == nullptr;
+					}
+					return !bMissingColumn;
 				}
 				else
 				{
-					return ReturnType{};
+					return true;
 				}
-			}
-		}
+			};
 
-		return Function(
-			Contract,
-			TConstArrayView<const void*>(ConstColumnsData, ConstColumnTypes.Num()),
-			TConstArrayView<void*>(MutableColumnsData, MutableColumnTypes.Num()));
+		Function(Result, Contract, Response, Specialization);
 	}
 
 	template<typename ReturnType, typename FunctionType>
@@ -124,35 +174,71 @@ namespace UE::Editor::DataStorage::Queries
 	{
 		using FunctionInfo = Private::TFunctionInfo<FunctionType>;
 		using ArgumentInfo = typename FunctionInfo::ArgumentInfo;
-		using ContextType = typename FunctionInfo::ContextType;
-
-		static_assert(FunctionInfo::bIsValidQueryFunction, "The provided function is not compatible with a query callback. "
-			"Likely errors are missing a query context or having incompatible arguments.");
-		static_assert(std::is_convertible_v<typename FunctionInfo::ReturnType, ReturnType>,
-			"The result of the provided function is not compatible with the query function that's being constructed.");
-
+		
 		TQueryFunction<ReturnType> Result;
-		if constexpr (!std::is_same_v<void, ContextType>)
-		{
-			Result.Capabilities = IContextContract::SupportedCapabilitiesList<typename ContextType::Capabilities>();
-		}
-		Result.ConstColumnTypes = ArgumentInfo::ListConstColumns();
-		Result.MutableColumnTypes = ArgumentInfo::ListMutableColumns();
+		FunctionInfo::template StaticValidation<ReturnType>();
+		FunctionInfo::SetupResult(Result);
+
 		Result.Function =
-			[Callback = Forward<FunctionType>(Function)]
-			(IContextContract& Contract, TConstArrayView<const void*> ConstColumns, TConstArrayView<void*> MutableColumns)
+			[Callback = Forward<FunctionType>(Function)](
+				TResult<ReturnType>& QueryResult, 
+				IContextContract& Contract, 
+				IQueryFunctionResponse& Response, 
+				TQueryFunction<ReturnType>::FunctionSpecializationCallback Specialization)
 			{
 				using ArgumentList = typename ArgumentInfo::ArgumentList;
 
+				TConstArrayView<const UScriptStruct*> ConstColumnTypes = ArgumentInfo::ListConstColumns();
+				TConstArrayView<const UScriptStruct*> MutableColumnTypes = ArgumentInfo::ListMutableColumns();
+
+				const void** ConstColumnsData = static_cast<const void**>(FMemory_Alloca(sizeof(void*) * ConstColumnTypes.Num()));
+				void** MutableColumnsData = static_cast<void**>(FMemory_Alloca(sizeof(void*) * MutableColumnTypes.Num()));
+
+				TArrayView<const void*> ConstColumns(ConstColumnsData, ConstColumnTypes.Num());
+				TArrayView<void*> MutableColumns(MutableColumnsData, MutableColumnTypes.Num());
+
+				EFlowControl Flow = EFlowControl::Continue;
+
 				ArgumentList Arguments;
-				if constexpr (!std::is_same_v<void, ContextType>)
-				{
-					ArgumentInfo::SetContext(Arguments, Contract);
-				}
-				ArgumentInfo::SetConstColumns(Arguments, ConstColumns);
-				ArgumentInfo::SetMutableColumns(Arguments, MutableColumns);
+				ArgumentInfo::SetResult(Arguments, QueryResult);
+				ArgumentInfo::SetContext(Arguments, Contract);
+				ArgumentInfo::SetFlowControl(Arguments, Flow);
 
-				return Arguments.ApplyBefore(Callback);
+				do
+				{
+					Response.GetConstColumns(ConstColumns, ConstColumnTypes);
+					Response.GetMutableColumns(MutableColumns, MutableColumnTypes);
+
+					if (Specialization(Response, ConstColumns, MutableColumns))
+					{
+						ArgumentInfo::SetConstColumns(Arguments, ConstColumns);
+						ArgumentInfo::SetMutableColumns(Arguments, MutableColumns);
+
+						if constexpr (FunctionInfo::bIsSingleRowProcessor)
+						{
+							do
+							{
+								if constexpr (ArgumentInfo::ResultIndex >= 0 || std::is_same_v<ReturnType, void>)
+								{
+									Arguments.ApplyBefore(Callback);
+								}
+								else
+								{
+									QueryResult.Add(Arguments.ApplyBefore(Callback));
+								}
+								ArgumentInfo::IncrementColumns(Arguments);
+							} while (Flow == EFlowControl::Continue && Response.NextRow());
+						}
+						else
+						{
+							Arguments.ApplyBefore(Callback);
+						}
+					}
+					else
+					{
+						break;
+					}
+				} while (Flow == EFlowControl::Continue && Response.NextBatch());
 			};
 		return Result;
 	}
diff --git a/Engine/Source/Runtime/TypedElementFramework/Tests/Elements/Framework/TypedElementQueryContextMock.h b/Engine/Source/Runtime/TypedElementFramework/Tests/Elements/Framework/TypedElementQueryContextMock.h
--- a/Engine/Source/Runtime/TypedElementFramework/Tests/Elements/Framework/TypedElementQueryContextMock.h
+++ b/Engine/Source/Runtime/TypedElementFramework/Tests/Elements/Framework/TypedElementQueryContextMock.h
@@ -4,6 +4,7 @@
 
 #include "Elements/Framework/TypedElementQueryCapabilities.h"
 #include "Elements/Framework/TypedElementQueryContextImplementation.h"
+#include "Elements/Framework/TypedElementQueryFunctionArgumentTypes.h"
 #include "Templates/FunctionFwd.h"
 
 namespace UE::Editor::DataStorage::Queries
@@ -24,7 +25,7 @@ namespace UE::Editor::DataStorage::Queries
 #define ConstFunction3(Capability, Return, Function, Arg1, Arg2, Arg3)
 #define ConstFunction4(Capability, Return, Function, Arg1, Arg2, Arg3, Arg4)
 
-#define CapabilityStart(Capability) Capability, 
+#define CapabilityStart(Capability, Flags) Capability, 
 #define CapabilityEnd(Capability)
 
 #include "Elements/Framework/TypedElementQueryCapabilities.inl"
@@ -113,7 +114,7 @@ namespace UE::Editor::DataStorage::Queries
 			return MockFunction(Capability, Function)(ArgName Arg1, ArgName Arg2, ArgName Arg3, ArgName Arg4); \
 		}
 
-#define CapabilityStart(Capability) 
+#define CapabilityStart(Capability, Flags) 
 #define CapabilityEnd(Capability)
 
 #include "Elements/Framework/TypedElementQueryCapabilities.inl"
diff --git a/Engine/Source/Runtime/TypedElementFramework/Tests/Elements/Framework/TypedElementQueryFunctionsTests.cpp b/Engine/Source/Runtime/TypedElementFramework/Tests/Elements/Framework/TypedElementQueryFunctionsTests.cpp
--- a/Engine/Source/Runtime/TypedElementFramework/Tests/Elements/Framework/TypedElementQueryFunctionsTests.cpp
+++ b/Engine/Source/Runtime/TypedElementFramework/Tests/Elements/Framework/TypedElementQueryFunctionsTests.cpp
@@ -10,52 +10,95 @@
 
 namespace UE::Editor::DataStorage::Queries::Tests
 {
+	struct FBoolAndResultAccumulator : TResult<bool>
+	{
+		bool Value = true;
+
+		virtual void Add(bool ResultValue) override
+		{
+			Value = Value && ResultValue;
+		}
+	};
+
+	struct FBoolOrResultAccumulator : TResult<bool>
+	{
+		bool Value = false;
+
+		virtual void Add(bool ResultValue) override
+		{
+			Value = Value || ResultValue;
+		}
+	};
+
+	struct FInt32ResultAccumulator : TResult<int32>
+	{
+		int32 Value = 0;
+
+		virtual void Add(int32 ResultValue) override
+		{
+			Value += ResultValue;
+		}
+	};
+	
 	struct FTestQueryFunctionResponse : IQueryFunctionResponse
 	{
 		virtual ~FTestQueryFunctionResponse() override = default;
 
-		virtual void GetConstColumns(TArrayView<const void*> ColumnsData, TConstArrayView<const UScriptStruct*> ColumnTypes) override
+		virtual bool NextBatch() override
+		{
+			CHECK_MESSAGE(TEXT("Row count was not set in test."), RowCount > 0);
+			return false;
+		}
+
+		virtual bool NextRow() override
 		{
-			CHECK_MESSAGE(TEXT("Attempting to get more or less const columns than are available in the mock."), ConstColumns.Num() == ColumnTypes.Num());
+			CHECK_MESSAGE(TEXT("Row count was not set in test."), RowCount > 0);
+			CurrentRow++;
+			return CurrentRow < RowCount;
+		}
 
+		virtual void GetConstColumns(TArrayView<const void*> ColumnsData, TConstArrayView<const UScriptStruct*> ColumnTypes) override
+		{
+			CHECK_MESSAGE(TEXT("Row count was not set in test."), RowCount > 0);
 			for (int32 Index = 0, Num = ColumnTypes.Num(); Index < Num; ++Index)
 			{
-				CHECK_MESSAGE(TEXT("Type for requested const column doesn't match the stored type in the mock."),
-					ConstColumnTypes[Index] == ColumnTypes[Index]);
-				ColumnsData[Index] = ConstColumns[Index];
+				const void** Data = ConstColumns.Find(ColumnTypes[Index]);
+				ColumnsData[Index] = Data ? (static_cast<const char*>(*Data) + ColumnTypes[Index]->GetStructureSize() * CurrentRow) : nullptr;
 			}
 		}
 
 		virtual void GetMutableColumns(TArrayView<void*> ColumnsData, TConstArrayView<const UScriptStruct*> ColumnTypes) override
 		{
-			CHECK_MESSAGE(TEXT("Attempting to get more or less mutable columns than are available in the mock."), MutableColumns.Num() == ColumnTypes.Num());
-
+			CHECK_MESSAGE(TEXT("Row count was not set in test."), RowCount > 0);
 			for (int32 Index = 0, Num = ColumnTypes.Num(); Index < Num; ++Index)
 			{
-				CHECK_MESSAGE(TEXT("Type for requested mutable column doesn't match the stored type in the mock."),
-					MutableColumnTypes[Index] == ColumnTypes[Index]);
-				ColumnsData[Index] = MutableColumns[Index];
+				void** Data = MutableColumns.Find(ColumnTypes[Index]);
+				ColumnsData[Index] = Data ? (static_cast<char*>(*Data) + ColumnTypes[Index]->GetStructureSize() * CurrentRow) : nullptr;
 			}
 		}
 
+		void SetRowCount(int32 InRowCount)
+		{
+			CHECK_MESSAGE(TEXT("Row count was not set in test."), InRowCount > 0);
+			RowCount = InRowCount;
+		}
+
 		template<TColumnType... ColumnType>
 		void SetConstColumns(const ColumnType*... Columns)
 		{
-			(ConstColumns.Add(Columns), ...);
-			(ConstColumnTypes.Add(Columns->StaticStruct()), ...);
+			(ConstColumns.Add(ColumnType::StaticStruct(), Columns), ...);
 		}
 
 		template<TColumnType... ColumnType>
 		void SetMutableColumns(ColumnType*... Columns)
 		{
-			(MutableColumns.Add(Columns), ...);
-			(MutableColumnTypes.Add(Columns->StaticStruct()), ...);
+			(MutableColumns.Add(ColumnType::StaticStruct(), Columns), ...);
 		}
 
-		TArray<const void*> ConstColumns;
-		TArray<const UScriptStruct*> ConstColumnTypes;
-		TArray<void*> MutableColumns;
-		TArray<const UScriptStruct*> MutableColumnTypes;
+		TMap<const UScriptStruct*, const void*> ConstColumns;
+		TMap<const UScriptStruct*, void*> MutableColumns;
+		int32 CurrentRow = 0;
+		int32 RowCount = 0;
 	};
 
 	TEST_CASE_NAMED(TQueryFunction_Tests, "Editor::DataStorage::Queries::Query Function (TQueryFunction)", "[ApplicationContextMask][EngineFilter]")
@@ -93,7 +136,7 @@ namespace UE::Editor::DataStorage::Queries::Tests
 
 		SECTION("Function with context")
 		{
-			TQueryFunction<void> Result = BuildQueryFunction<void>([](TQueryContext<SingleRowInfo>& Context) {});
+			TQueryFunction<void> Result = BuildQueryFunction<void>([](TQueryContext<SingleRowInfo> Context) {});
 
 			CHECK_MESSAGE(TEXT("Expected exactly one capability."), Result.Capabilities.Num() == 1);
 			CHECK_MESSAGE(TEXT("Query with only context shouldn't have const columns."), Result.ConstColumnTypes.IsEmpty());
@@ -115,7 +158,7 @@ namespace UE::Editor::DataStorage::Queries::Tests
 		SECTION("Function with context and columns")
 		{
 			TQueryFunction<void> Result = BuildQueryFunction<void>(
-				[](TQueryContext<SingleRowInfo>& Context, FTestColumnInt& ColumnA, const FTestColumnString& ColumnB)
+				[](TQueryContext<SingleRowInfo> Context, FTestColumnInt& ColumnA, const FTestColumnString& ColumnB)
 				{
 				});
 
@@ -130,7 +173,22 @@ namespace UE::Editor::DataStorage::Queries::Tests
 		SECTION("Function with context and columns at random location")
 		{
 			TQueryFunction<void> Result = BuildQueryFunction<void>(
-				[](FTestColumnInt& ColumnA, TQueryContext<SingleRowInfo>& Context, const FTestColumnString& ColumnB)
+				[](FTestColumnInt& ColumnA, TQueryContext<SingleRowInfo> Context, const FTestColumnString& ColumnB)
+				{
+				});
+
+			CHECK_MESSAGE(TEXT("Expected exactly one capability."), Result.Capabilities.Num() == 1);
+			CHECK_MESSAGE(TEXT("Expected exactly one const column."), Result.ConstColumnTypes.Num() == 1);
+			CHECK_MESSAGE(TEXT("Expected exactly one mutable column."), Result.MutableColumnTypes.Num() == 1);
+
+			CHECK_MESSAGE(TEXT("Incorrect const column set."), Result.ConstColumnTypes[0] == FTestColumnString::StaticStruct());
+			CHECK_MESSAGE(TEXT("Incorrect mutable column set."), Result.MutableColumnTypes[0] == FTestColumnInt::StaticStruct());
+		}
+
+		SECTION("Function with context and column batch")
+		{
+			TQueryFunction<void> Result = BuildQueryFunction<void>(
+				[](TQueryContext<RowBatchInfo> Context, TBatch<FTestColumnInt> ColumnsA, TConstBatch<FTestColumnString> ColumnB)
 				{
 				});
 
@@ -153,7 +211,7 @@ namespace UE::Editor::DataStorage::Queries::Tests
 
 		SECTION("Check context compatibility with context")
 		{
-			TQueryFunction<void> Function = BuildQueryFunction<void>([](TQueryContext<SingleRowInfo>& Context) {});
+			TQueryFunction<void> Function = BuildQueryFunction<void>([](TQueryContext<SingleRowInfo> Context) {});
 
 			QueryContextMock Context;
 			bool bResult = Context.CheckCompatiblity(Function);
@@ -162,7 +220,7 @@ namespace UE::Editor::DataStorage::Queries::Tests
 
 		SECTION("Check context compatibility with restricted context")
 		{
-			TQueryFunction<void> Function = BuildQueryFunction<void>([](TQueryContext<SingleRowInfo>& Context) {});
+			TQueryFunction<void> Function = BuildQueryFunction<void>([](TQueryContext<SingleRowInfo> Context) {});
 
 			TQueryContextImpl<FQueryContextMock, SingleRowInfo> Context;
 			bool bResult = Context.CheckCompatiblity(Function);
@@ -177,11 +235,13 @@ namespace UE::Editor::DataStorage::Queries::Tests
 					return true;
 				});
 			
-			bool bResult = false;
 			QueryContextMock Context;
 			FTestQueryFunctionResponse Response;
-			bResult = Function.Call<EFunctionCallConfig::VerifyColumns>(Context, Response);
-			CHECK_MESSAGE(TEXT("Query function wasn't called."), bResult);
+			Response.SetRowCount(1);
+			
+			FBoolOrResultAccumulator Result;
+			Function.Call<EFunctionCallConfig::VerifyColumns>(Result, Context, Response);
+			CHECK_MESSAGE(TEXT("Query function wasn't called."), Result.Value);
 		}
 
 		SECTION("Call setting column value")
@@ -196,6 +256,7 @@ namespace UE::Editor::DataStorage::Queries::Tests
 			FTestColumnString Column;
 			
 			FTestQueryFunctionResponse Response;
+			Response.SetRowCount(1);
 			Response.SetMutableColumns(&Column);
 			Function.Call<EFunctionCallConfig::VerifyColumns>(Context, Response);
 			
@@ -216,6 +277,7 @@ namespace UE::Editor::DataStorage::Queries::Tests
 			ColumnB.TestInt = 42;
 
 			FTestQueryFunctionResponse Response;
+			Response.SetRowCount(1);
 			Response.SetMutableColumns(&ColumnA);
 			Response.SetConstColumns(&ColumnB);
 			Function.Call<EFunctionCallConfig::VerifyColumns>(Context, Response);
@@ -223,10 +285,42 @@ namespace UE::Editor::DataStorage::Queries::Tests
 			CHECK_MESSAGE(TEXT("Query function wasn't called."), ColumnA.TestString == TEXT("Test: 42"));
 		}
 
+		SECTION("Call setting const and mutable batch columns")
+		{
+			TQueryFunction<void> Function = BuildQueryFunction<void>(
+				[](TQueryContext<RowBatchInfo> Context, TBatch<FTestColumnString> ColumnsA, TConstBatch<FTestColumnInt> ColumnsB)
+				{
+					Context.ForEachRow([](FTestColumnString& ColumnA, const FTestColumnInt& ColumnB)
+						{
+							ColumnA.TestString = FString::Printf(TEXT("Test: %i"), ColumnB.TestInt);
+						}, ColumnsA, ColumnsB);
+				});
+
+			static constexpr uint32 RowCount = 2;
+
+			QueryContextMock Context;
+			Context.GetContextImplementation().RowBatchInfo_GetBatchRowCount_Mock = []() { return RowCount; };
+			FTestColumnString ColumnsA[RowCount];
+			FTestColumnInt ColumnsB[] =
+			{
+				FTestColumnInt{ .TestInt = 42 },
+				FTestColumnInt{ .TestInt = 88 }
+			};
+
+			FTestQueryFunctionResponse Response;
+			Response.SetRowCount(RowCount);
+			Response.SetMutableColumns(ColumnsA);
+			Response.SetConstColumns(ColumnsB);
+			Function.Call<EFunctionCallConfig::VerifyColumns>(Context, Response);
+
+			CHECK_MESSAGE(TEXT("Query function wasn't called."), ColumnsA[0].TestString == TEXT("Test: 42"));
+			CHECK_MESSAGE(TEXT("Query function wasn't called."), ColumnsA[1].TestString == TEXT("Test: 88"));
+		}
+
 		SECTION("Call and use context")
 		{
 			TQueryFunction<void> Function = BuildQueryFunction<void>(
-				[](TQueryContext<SingleRowInfo>& Context)
+				[](TQueryContext<SingleRowInfo> Context)
 				{
 					Context.GetCurrentRow();
 					Context.CurrentRowHasColumns<FTestColumnString>();
@@ -251,7 +345,8 @@ namespace UE::Editor::DataStorage::Queries::Tests
 				};
 
 			FTestQueryFunctionResponse Response;
-			Function.Call<>(Context, Response);
+			Response.SetRowCount(1);
+			Function.Call(Context, Response);
 
 			CHECK_MESSAGE(TEXT("GetCurrentRow not called on context."), bGetCurrentRowCalled);
 			CHECK_MESSAGE(TEXT("CurrentRowHasColumns not called on context."), bCurrentRowHasColumnsCalled);
@@ -260,7 +355,7 @@ namespace UE::Editor::DataStorage::Queries::Tests
 		SECTION("Call and use restricted context")
 		{
 			TQueryFunction<void> Function = BuildQueryFunction<void>(
-				[](TQueryContext<SingleRowInfo>& Context)
+				[](TQueryContext<SingleRowInfo> Context)
 				{
 					Context.GetCurrentRow();
 					Context.CurrentRowHasColumns<FTestColumnString>();
@@ -285,11 +380,111 @@ namespace UE::Editor::DataStorage::Queries::Tests
 				};
 
 			FTestQueryFunctionResponse Response;
-			Function.Call<>(Context, Response);
+			Response.SetRowCount(1);
+			Function.Call(Context, Response);
 
 			CHECK_MESSAGE(TEXT("GetCurrentRow not called on context."), bGetCurrentRowCalled);
 			CHECK_MESSAGE(TEXT("CurrentRowHasColumns not called on context."), bCurrentRowHasColumnsCalled);
 		}
+
+		SECTION("Call with result from callback")
+		{
+			TQueryFunction<bool> Function = BuildQueryFunction<bool>([]() { return true; });
+
+			QueryContextMock Context;
+			FTestColumnString Column;
+
+			FBoolOrResultAccumulator Result;
+
+			FTestQueryFunctionResponse Response;
+			Response.SetRowCount(1);
+			Response.SetMutableColumns(&Column);
+			Function.Call<EFunctionCallConfig::VerifyColumns>(Result, Context, Response);
+
+			CHECK_MESSAGE(TEXT("Result wasn't passed back in."), Result.Value);
+		}
+
+		SECTION("Call with result through TResult")
+		{
+			TQueryFunction<bool> Function = BuildQueryFunction<bool>(
+				[](TResult<bool>& Result)
+				{ 
+					Result.Add(true);
+				});
+
+			QueryContextMock Context;
+			FTestColumnString Column;
+
+			FBoolOrResultAccumulator Result;
+
+			FTestQueryFunctionResponse Response;
+			Response.SetRowCount(1);
+			Response.SetMutableColumns(&Column);
+			Function.Call<EFunctionCallConfig::VerifyColumns>(Result, Context, Response);
+
+			CHECK_MESSAGE(TEXT("Result wasn't passed back in."), Result.Value);
+		}
+
+		SECTION("Call with batching results.")
+		{
+			TQueryFunction<int32> Function = BuildQueryFunction<int32>(
+				[](TQueryContext<RowBatchInfo> Context, TResult<int32>& Result, TConstBatch<FTestColumnInt> Columns)
+				{
+					Context.ForEachRow([&Result](const FTestColumnInt& Column)
+						{
+							Result.Add(Column.TestInt);
+						}, Columns);
+				});
+
+			static constexpr int32 RowCount = 4;
+
+			QueryContextMock Context;
+			Context.GetContextImplementation().RowBatchInfo_GetBatchRowCount_Mock = []() { return RowCount; };
+
+			FTestColumnInt Columns[] =
+			{
+				FTestColumnInt{.TestInt = 1 },
+				FTestColumnInt{.TestInt = 2 },
+				FTestColumnInt{.TestInt = 4 },
+				FTestColumnInt{.TestInt = 8 }
+			};
+
+			FInt32ResultAccumulator Result;
+
+			FTestQueryFunctionResponse Response;
+			Response.SetRowCount(RowCount);
+			Response.SetConstColumns(Columns);
+			Function.Call<EFunctionCallConfig::VerifyColumns>(Result, Context, Response);
+
+			CHECK_MESSAGE(TEXT("Result wasn't passed back in."), Result.Value == 1 + 2 + 4 + 8);
+		}
+
+		SECTION("Call and break on first row")
+		{
+			int32 LastValue = 0;
+			TQueryFunction<void> Function = BuildQueryFunction<void>(
+				[&LastValue](EFlowControl& Flow, const FTestColumnInt& Column)
+				{
+					LastValue = Column.TestInt;
+					Flow = EFlowControl::Break;
+				});
+
+			QueryContextMock Context;
+
+			FTestColumnInt Columns[] =
+			{
+				FTestColumnInt{.TestInt = 1 },
+				FTestColumnInt{.TestInt = 2 }
+			};
+
+			FTestQueryFunctionResponse Response;
+			Response.SetRowCount(2);
+			Response.SetConstColumns(Columns);
+
+			Function.Call(Context, Response);
+
+			CHECK_MESSAGE(TEXT("Iteration didn't stop on the first row."), LastValue == 1);
+		}
 	}
 } // namespace UE::Editor::DataStorage::Queries::Tests
 
