TEDS UI: Fix cvar to print widget purposes

diff --git a/Engine/Plugins/Experimental/EditorDataStorage/Source/TedsCore/Private/TypedElementDatabaseUI.cpp b/Engine/Plugins/Experimental/EditorDataStorage/Source/TedsCore/Private/TypedElementDatabaseUI.cpp
--- a/Engine/Plugins/Experimental/EditorDataStorage/Source/TedsCore/Private/TypedElementDatabaseUI.cpp
+++ b/Engine/Plugins/Experimental/EditorDataStorage/Source/TedsCore/Private/TypedElementDatabaseUI.cpp
@@ -719,12 +719,13 @@ void UEditorDataStorageUi::ListWidgetPurposes(const WidgetPurposeCallback& Callb
 	using namespace UE::Editor::DataStorage::Queries;
 	
 	static QueryHandle PurposeQueryHandle =
-		Storage->RegisterQuery(Select().ReadOnly<FWidgetConstructorNameColumn, FWidgetPurposeColumn, FDescriptionColumn>().Compile());
+		Storage->RegisterQuery(Select().ReadOnly<FWidgetPurposeNameColumn, FWidgetPurposeColumn, FDescriptionColumn>().Compile());
 
 	Storage->RunQuery(PurposeQueryHandle, CreateDirectQueryCallbackBinding(
-		[Callback](const FWidgetConstructorNameColumn& NameColumn, const FWidgetPurposeColumn& PurposeColumn, const FDescriptionColumn& DescriptionColumn)
+		[Callback](const FWidgetPurposeNameColumn& NameColumn, const FWidgetPurposeColumn& PurposeColumn, const FDescriptionColumn& DescriptionColumn)
 	{
-			Callback(NameColumn.Name, PurposeColumn.PurposeType, DescriptionColumn.Description);
+		FPurposeInfo PurposeInfo(NameColumn);
+		Callback(PurposeInfo.ToString(), PurposeColumn.PurposeType, DescriptionColumn.Description);
 	}));
 }
 
diff --git a/Engine/Source/Runtime/TypedElementFramework/Private/Elements/Interfaces/TypedElementDataStorageUiInterface.cpp b/Engine/Source/Runtime/TypedElementFramework/Private/Elements/Interfaces/TypedElementDataStorageUiInterface.cpp
--- a/Engine/Source/Runtime/TypedElementFramework/Private/Elements/Interfaces/TypedElementDataStorageUiInterface.cpp
+++ b/Engine/Source/Runtime/TypedElementFramework/Private/Elements/Interfaces/TypedElementDataStorageUiInterface.cpp
@@ -410,7 +410,18 @@ UE::Editor::DataStorage::IUiProvider::FPurposeInfo::FPurposeInfo(const FName& In
 	ParentPurposeID = InParent;
 }
 
+UE::Editor::DataStorage::IUiProvider::FPurposeInfo::FPurposeInfo(const FWidgetPurposeNameColumn& WidgetPurposeNameColumn)
+	: FPurposeInfo(WidgetPurposeNameColumn.Namespace, WidgetPurposeNameColumn.Name, WidgetPurposeNameColumn.Frame)
+{
+	
+}
+
 UE::Editor::DataStorage::IUiProvider::FPurposeID UE::Editor::DataStorage::IUiProvider::FPurposeInfo::GeneratePurposeID() const
+{
+	return FMapKey(ToString());
+}
+
+FName UE::Editor::DataStorage::IUiProvider::FPurposeInfo::ToString() const
 {
 	// Create the unique ID from the purpose by combining the namespace, name and frame
 	FStringBuilderBase PurposeIDBuilder;
@@ -427,10 +438,9 @@ UE::Editor::DataStorage::IUiProvider::FPurposeID UE::Editor::DataStorage::IUiPro
 	{
 		Frame.AppendString(PurposeIDBuilder);
 		PurposeIDBuilder.Append(TEXT("."));
-
 	}
-	
-	return FMapKey(FName(PurposeIDBuilder));
+
+	return FName(PurposeIDBuilder);
 }
 
 
diff --git a/Engine/Source/Runtime/TypedElementFramework/Public/Elements/Interfaces/TypedElementDataStorageUiInterface.h b/Engine/Source/Runtime/TypedElementFramework/Public/Elements/Interfaces/TypedElementDataStorageUiInterface.h
--- a/Engine/Source/Runtime/TypedElementFramework/Public/Elements/Interfaces/TypedElementDataStorageUiInterface.h
+++ b/Engine/Source/Runtime/TypedElementFramework/Public/Elements/Interfaces/TypedElementDataStorageUiInterface.h
@@ -26,6 +26,7 @@ namespace UE::Editor::DataStorage
 }
 
 class SWidget;
+struct FWidgetPurposeNameColumn;
 
 /**
  * Base class for Widget Constructors in TEDS.
@@ -370,8 +371,14 @@ namespace UE::Editor::DataStorage
 			TYPEDELEMENTFRAMEWORK_API FPurposeInfo(const FName& InLegacyPurposeName, EPurposeType InPurposeType, const FText& InDescription,
 				const FPurposeID& InParent = FPurposeID());
 
+			/** Create an FPurposeInfo struct using FWidgetPurposeNameColumn */
+			TYPEDELEMENTFRAMEWORK_API explicit FPurposeInfo(const FWidgetPurposeNameColumn& WidgetPurposeNameColumn);
+
 			/** Convert this PurposeInfo to a PurposeID that can be used to look up the purpose row */
 			TYPEDELEMENTFRAMEWORK_API FPurposeID GeneratePurposeID() const;
+
+			/** Convert this PurposeInfo to a human-readable name in the legacy format ("Namespace.Name.Frame") */
+			TYPEDELEMENTFRAMEWORK_API FName ToString() const;
 		};
 
 		using WidgetCreatedCallback = TFunctionRef<void(const TSharedRef<SWidget>& NewWidget, UE::Editor::DataStorage::RowHandle Row)>;
