TEDS UI: Fix incorrect trailing dot in FPurposeInfo::ToString

diff --git a/Engine/Source/Runtime/TypedElementFramework/Private/Elements/Interfaces/TypedElementDataStorageUiInterface.cpp b/Engine/Source/Runtime/TypedElementFramework/Private/Elements/Interfaces/TypedElementDataStorageUiInterface.cpp
--- a/Engine/Source/Runtime/TypedElementFramework/Private/Elements/Interfaces/TypedElementDataStorageUiInterface.cpp
+++ b/Engine/Source/Runtime/TypedElementFramework/Private/Elements/Interfaces/TypedElementDataStorageUiInterface.cpp
@@ -429,15 +429,17 @@ FName UE::Editor::DataStorage::IUiProvider::FPurposeInfo::ToString() const
 	if (!Namespace.IsNone())
 	{
 		Namespace.AppendString(PurposeIDBuilder);
-		PurposeIDBuilder.Append(TEXT("."));
 	}
 
+	// We want the dots even if the namespace/frame is missing since it signifies there is something missing (".Name.Frame")
+	// instead of being ambigous ("Namespace.Name", which is identical to "Name.Frame")
+	PurposeIDBuilder.Append(TEXT("."));
 	Name.AppendString(PurposeIDBuilder);
+	PurposeIDBuilder.Append(TEXT("."));
 
 	if (!Frame.IsNone())
 	{
 		Frame.AppendString(PurposeIDBuilder);
-		PurposeIDBuilder.Append(TEXT("."));
 	}
 
 	return FName(PurposeIDBuilder);
