[TEDS] MSVC 14.44 reported a compile error with TEDS Query Context.

A workaround was added by wrapping a templated function into another templated function instead of directly calling it. This changes the way MSVC evaluates template validity and seems to do the trick.

#rnx

diff --git a/Engine/Source/Runtime/TypedElementFramework/Public/Elements/Framework/TypedElementQueryFunctionArguments.h b/Engine/Source/Runtime/TypedElementFramework/Public/Elements/Framework/TypedElementQueryFunctionArguments.h
--- a/Engine/Source/Runtime/TypedElementFramework/Public/Elements/Framework/TypedElementQueryFunctionArguments.h
+++ b/Engine/Source/Runtime/TypedElementFramework/Public/Elements/Framework/TypedElementQueryFunctionArguments.h
@@ -236,6 +236,8 @@ namespace UE::Editor::DataStorage::Queries::Private
 		static void SetContextUnguarded(ArgumentList& Arguments, IContextContract& Contract);
 		template<int32 Index>
 		static void SetFlowControlUnguarded(ArgumentList& Arguments, EFlowControl& FlowControl);
+		template<typename ArgumentType>
+		static void IncrementColumnUnguarded(ArgumentList& Arguments);
 	};
 } // namespace UE::Editor::DataStorage::Queries::Private
 
diff --git a/Engine/Source/Runtime/TypedElementFramework/Public/Elements/Framework/TypedElementQueryFunctionArguments.inl b/Engine/Source/Runtime/TypedElementFramework/Public/Elements/Framework/TypedElementQueryFunctionArguments.inl
--- a/Engine/Source/Runtime/TypedElementFramework/Public/Elements/Framework/TypedElementQueryFunctionArguments.inl
+++ b/Engine/Source/Runtime/TypedElementFramework/Public/Elements/Framework/TypedElementQueryFunctionArguments.inl
@@ -180,7 +180,7 @@ namespace UE::Editor::DataStorage::Queries::Private
 			using ArgumentType = TArgument<Args>;
 			if constexpr (EnumHasAllFlags(ArgumentType::Flags, EArgumentFlags::SingleArgument | EArgumentFlags::Type_Column))
 			{
-				Arguments.template Get<typename ArgumentType::ArgumentType>().Increment();
+				IncrementColumnUnguarded<typename ArgumentType::ArgumentType>(Arguments);
 			}
 			}(), ...);
 	}
@@ -207,4 +207,11 @@ namespace UE::Editor::DataStorage::Queries::Private
 	{
 		Arguments.template Get<Index>() = &FlowControl;
 	}
+
+	template<typename... Args>
+	template<typename ArgumentType>
+	void TArgumentInfo<Args...>::IncrementColumnUnguarded(ArgumentList& Arguments)
+	{
+		Arguments.template Get<ArgumentType>().Increment();
+	}
 } // namespace UE::Editor::DataStorage::Queries::Private
