[TEDS] Added an overloaded version of RunQuery that uses the new query callback.

diff --git a/Engine/Plugins/Experimental/EditorDataStorage/Source/TedsCore/Private/Processors/TypedElementProcessorAdaptors.cpp b/Engine/Plugins/Experimental/EditorDataStorage/Source/TedsCore/Private/Processors/TypedElementProcessorAdaptors.cpp
--- a/Engine/Plugins/Experimental/EditorDataStorage/Source/TedsCore/Private/Processors/TypedElementProcessorAdaptors.cpp
+++ b/Engine/Plugins/Experimental/EditorDataStorage/Source/TedsCore/Private/Processors/TypedElementProcessorAdaptors.cpp
@@ -7,6 +7,7 @@
 #include "MassEntityView.h"
 #include "MassExecutionContext.h"
 #include "Queries/TypedElementExtendedQueryStore.h"
+#include "Queries/TypedElementQueryContext_MassForwarder.h"
 #include "StructUtils/StructTypeBitSet.h"
 #include "TypedElementDatabase.h"
 #include "TypedElementDatabaseEnvironment.h"
@@ -1636,6 +1637,53 @@ UE::Editor::DataStorage::FQueryResult FTypedElementQueryProcessorData::Execute(
 	return Result;
 }
 
+UE::Editor::DataStorage::FQueryResult FTypedElementQueryProcessorData::Execute(
+	UE::Editor::DataStorage::Queries::TQueryFunction<void>& Callback,
+	UE::Editor::DataStorage::FQueryDescription& Description,
+	FMassEntityQuery& NativeQuery,
+	FMassEntityManager& EntityManager,
+	UE::Editor::DataStorage::FEnvironment& Environment,
+	UE::Editor::DataStorage::EDirectQueryExecutionFlags ExecutionFlags)
+{
+	using namespace UE::Editor::DataStorage;
+	using namespace UE::Editor::DataStorage::Queries;
+
+	FQueryResult Result;
+	Result.Completed = FQueryResult::ECompletion::Fully;
+
+	if (EnumHasAnyFlags(ExecutionFlags, EDirectQueryExecutionFlags::AllowBoundQueries) || !Description.Callback.Function)
+	{
+		FMassExecutionContext Context(EntityManager);
+		auto ExecuteFunction = [&Result, &Callback, &Description, &Environment](FMassExecutionContext& Context)
+			{
+				using namespace Processors::Private;
+				// No need to cache any subsystem dependencies as these are not accessible from a direct query.
+				QueryContext_MassForwarder QueryContext(Context);
+				// This assumes that there already has been a check that the provided function is compatible with the query.
+				Callback.Call(QueryContext, QueryContext.GetContextImplementation());
+				Result.Count += Context.GetNumEntities();
+			};
+		if (EnumHasAnyFlags(ExecutionFlags, EDirectQueryExecutionFlags::ParallelizeChunks))
+		{
+			FMassEntityQuery::EParallelExecutionFlags Mode =
+				EnumHasAnyFlags(ExecutionFlags, EDirectQueryExecutionFlags::AutoBalanceParallelChunkProcessing)
+				? FMassEntityQuery::EParallelExecutionFlags::AutoBalance
+				: FMassEntityQuery::EParallelExecutionFlags::Default;
+			NativeQuery.ParallelForEachEntityChunk(Context, ExecuteFunction, Mode);
+		}
+		else
+		{
+			NativeQuery.ForEachEntityChunk(Context, ExecuteFunction);
+		}
+		Context.FlushDeferred();
+	}
+	else
+	{
+		Result.Completed = FQueryResult::ECompletion::Unsupported;
+	}
+	return Result;
+}
+
 UE::Editor::DataStorage::FQueryResult FTypedElementQueryProcessorData::Execute(
 	UE::Editor::DataStorage::SubqueryCallbackRef& Callback,
 	UE::Editor::DataStorage::FQueryDescription& Description,
diff --git a/Engine/Plugins/Experimental/EditorDataStorage/Source/TedsCore/Private/Processors/TypedElementProcessorAdaptors.h b/Engine/Plugins/Experimental/EditorDataStorage/Source/TedsCore/Private/Processors/TypedElementProcessorAdaptors.h
--- a/Engine/Plugins/Experimental/EditorDataStorage/Source/TedsCore/Private/Processors/TypedElementProcessorAdaptors.h
+++ b/Engine/Plugins/Experimental/EditorDataStorage/Source/TedsCore/Private/Processors/TypedElementProcessorAdaptors.h
@@ -60,6 +60,13 @@ struct FTypedElementQueryProcessorData
 		FMassEntityManager& EntityManager,
 		UE::Editor::DataStorage::FEnvironment& Environment,
 		UE::Editor::DataStorage::EDirectQueryExecutionFlags ExecutionFlags);
+	static UE::Editor::DataStorage::FQueryResult Execute(
+		UE::Editor::DataStorage::Queries::TQueryFunction<void>& Callback,
+		UE::Editor::DataStorage::FQueryDescription& Description,
+		FMassEntityQuery& NativeQuery,
+		FMassEntityManager& EntityManager,
+		UE::Editor::DataStorage::FEnvironment& Environment,
+		UE::Editor::DataStorage::EDirectQueryExecutionFlags ExecutionFlags);
 	static UE::Editor::DataStorage::FQueryResult Execute(
 		UE::Editor::DataStorage::SubqueryCallbackRef& Callback,
 		UE::Editor::DataStorage::FQueryDescription& Description,
diff --git a/Engine/Plugins/Experimental/EditorDataStorage/Source/TedsCore/Private/Queries/TypedElementExtendedQueryStore.cpp b/Engine/Plugins/Experimental/EditorDataStorage/Source/TedsCore/Private/Queries/TypedElementExtendedQueryStore.cpp
--- a/Engine/Plugins/Experimental/EditorDataStorage/Source/TedsCore/Private/Queries/TypedElementExtendedQueryStore.cpp
+++ b/Engine/Plugins/Experimental/EditorDataStorage/Source/TedsCore/Private/Queries/TypedElementExtendedQueryStore.cpp
@@ -388,6 +388,18 @@ namespace UE::Editor::DataStorage
 						Result = FTypedElementQueryProcessorData::Execute(
 							Callback, QueryData->Description, QueryData->NativeQuery, EntityManager, Environment, DirectExecutionFlags);
 					}
+					else if constexpr (std::is_same_v<CallbackReference, Queries::TQueryFunction<void>>)
+					{
+						if (CheckCompatibility(QueryData->Description, Callback))
+						{
+							Result = FTypedElementQueryProcessorData::Execute(
+								Callback, QueryData->Description, QueryData->NativeQuery, EntityManager, Environment, DirectExecutionFlags);
+						}
+						else
+						{
+							Result.Completed = FQueryResult::ECompletion::Incompatible;
+						}	
+					}
 					else
 					{
 						Result = FTypedElementQueryProcessorData::Execute(
@@ -423,6 +435,13 @@ namespace UE::Editor::DataStorage
 	{
 		return RunQueryCallbackCommon(EntityManager, Environment, nullptr, Query, DirectExecutionFlags, Callback);
 	}
+
+	FQueryResult FExtendedQueryStore::RunQuery(FMassEntityManager& EntityManager,
+		FEnvironment& Environment, Handle Query, ERunQueryFlags Flags,
+		const Queries::TQueryFunction<void>& Callback)
+	{
+		return RunQueryCallbackCommon(EntityManager, Environment, nullptr, Query, static_cast<EDirectQueryExecutionFlags>(Flags), Callback);
+	}
 	
 	FQueryResult FExtendedQueryStore::RunQuery(FMassEntityManager& EntityManager,
 		FEnvironment& Environment, FMassExecutionContext& ParentContext, Handle Query,
@@ -1283,4 +1302,26 @@ namespace UE::Editor::DataStorage
 			QueryData.NativeQuery.Clear();
 		}
 	}
+
+	bool FExtendedQueryStore::CheckCompatibility(const FQueryDescription& QueryDescription, const Queries::TQueryFunction<void> QueryCallback) const
+	{
+		for (const UScriptStruct* Column : QueryCallback.ConstColumnTypes)
+		{
+			if (int32 Index = QueryDescription.SelectionTypes.Find(Column);
+				Index == INDEX_NONE || QueryDescription.SelectionAccessTypes[Index] == EQueryAccessType::OptionalReadOnly)
+			{
+				return false;
+			}
+		}
+
+		for (const UScriptStruct* Column : QueryCallback.MutableColumnTypes)
+		{
+			if (int32 Index = QueryDescription.SelectionTypes.Find(Column);
+				Index == INDEX_NONE || QueryDescription.SelectionAccessTypes[Index] != EQueryAccessType::ReadWrite)
+			{
+				return false;
+			}
+		}
+		return true;
+	}
 } // namespace UE::Editor::DataStorage
diff --git a/Engine/Plugins/Experimental/EditorDataStorage/Source/TedsCore/Private/Queries/TypedElementExtendedQueryStore.h b/Engine/Plugins/Experimental/EditorDataStorage/Source/TedsCore/Private/Queries/TypedElementExtendedQueryStore.h
--- a/Engine/Plugins/Experimental/EditorDataStorage/Source/TedsCore/Private/Queries/TypedElementExtendedQueryStore.h
+++ b/Engine/Plugins/Experimental/EditorDataStorage/Source/TedsCore/Private/Queries/TypedElementExtendedQueryStore.h
@@ -126,6 +126,12 @@ namespace UE::Editor::DataStorage
 			Handle Query,
 			EDirectQueryExecutionFlags DirectExecutionFlags,
 			DirectQueryCallbackRef Callback);
+		FQueryResult RunQuery(
+			FMassEntityManager& EntityManager,
+			FEnvironment& Environment,
+			Handle Query,
+			ERunQueryFlags Flags,
+			const Queries::TQueryFunction<void>& Callback);
 		FQueryResult RunQuery(
 			FMassEntityManager& EntityManager,
 			FEnvironment& Environment,
@@ -209,6 +215,8 @@ namespace UE::Editor::DataStorage
 
 		void UnregisterQueryData(Handle Query, FExtendedQuery& QueryData, FMassEntityManager& EntityManager, FMassProcessingPhaseManager& PhaseManager);
 
+		bool CheckCompatibility(const FQueryDescription& QueryDescription, const Queries::TQueryFunction<void> QueryCallback) const;
+
 		static const FQueryDescription EmptyDescription;
 
 		QueryStore Queries;
diff --git a/Engine/Plugins/Experimental/EditorDataStorage/Source/TedsCore/Private/Queries/TypedElementQueryContext_MassForwarder.cpp b/Engine/Plugins/Experimental/EditorDataStorage/Source/TedsCore/Private/Queries/TypedElementQueryContext_MassForwarder.cpp
new file mode 100644
--- /dev/null
+++ b/Engine/Plugins/Experimental/EditorDataStorage/Source/TedsCore/Private/Queries/TypedElementQueryContext_MassForwarder.cpp
@@ -0,0 +1,106 @@
+// Copyright Epic Games, Inc. All Rights Reserved.
+
+#include "Queries/TypedElementQueryContext_MassForwarder.h"
+#include "MassExecutionContext.h"
+#include "TypedElementUtils.h"
+
+namespace UE::Editor::DataStorage::Queries
+{
+	const FName FQueryContext_MassForwarder::Name = "FQueryContext_MassForwarder";
+
+	FQueryContext_MassForwarder::FQueryContext_MassForwarder(FMassExecutionContext& Context)
+		: Context(Context)
+	{
+		TConstArrayView<RowHandle> Rows = MassEntitiesToRowsConversion(Context.GetEntities());
+		CurrentRow = Rows.GetData();
+		EndRow = Rows.GetData() + Rows.Num();
+	}
+
+	bool FQueryContext_MassForwarder::CurrentRowHasColumns(TConstArrayView<const UScriptStruct*> ColumnTypes) const
+	{
+		bool Result = true;
+		const UScriptStruct* const* It = ColumnTypes.GetData();
+		const UScriptStruct* const* End = ColumnTypes.GetData() + ColumnTypes.Num();
+
+		for (; Result && It != End; ++It)
+		{
+			const UScriptStruct* ColumnType = *It;
+			if (UE::Mass::IsA<FMassTag>(ColumnType))
+			{
+				Result = Context.DoesArchetypeHaveTag(*ColumnType);
+			}
+			else if (UE::Mass::IsA<FMassFragment>(ColumnType))
+			{
+				Result = Context.DoesArchetypeHaveFragment(*ColumnType);
+			}
+			else
+			{
+				const bool bIsTagOrFragment = false;
+				checkf(bIsTagOrFragment, TEXT("Attempting to check for a column type that is not a column or tag."));
+				return false;
+			}
+		}
+		return Result;
+	}
+
+	RowHandle FQueryContext_MassForwarder::GetCurrentRow() const
+	{
+		checkfSlow(CurrentRow != nullptr, TEXT("GetCurrentRow shouldn't be called if there are no rows."));
+		return *CurrentRow;
+	}
+
+	uint32 FQueryContext_MassForwarder::GetBatchRowCount() const
+	{
+		return static_cast<uint32>(Context.GetNumEntities());
+	}
+
+	TConstArrayView<RowHandle> FQueryContext_MassForwarder::GetBatchRowHandles() const
+	{
+		return MassEntitiesToRowsConversion(Context.GetEntities());
+	}
+
+	bool FQueryContext_MassForwarder::CurrentBatchTableHasColumns(TConstArrayView<const UScriptStruct*> ColumnTypes) const
+	{
+		return CurrentRowHasColumns(ColumnTypes);
+	}
+
+	bool FQueryContext_MassForwarder::NextBatch()
+	{
+		return false; // This only ever contains a single batch.
+	}
+
+	bool FQueryContext_MassForwarder::NextRow()
+	{
+		checkfSlow(CurrentRow != nullptr, TEXT("NextRow shouldn't be called if there are no rows."));
+		CurrentRow++;
+		return CurrentRow < EndRow;
+	}
+
+	void FQueryContext_MassForwarder::GetConstColumns(TArrayView<const void*> ColumnsData, TConstArrayView<const UScriptStruct*> ColumnTypes)
+	{
+		checkfSlow(ColumnsData.Num() >= ColumnTypes.Num(), TEXT("The list size of requested const column types (%i) doesn't match the "
+			"available space for the results (%i)"), ColumnTypes.Num(), ColumnsData.Num());
+
+		const void** ColumnsDataIt = ColumnsData.GetData();
+		for (const UScriptStruct* ColumnType : ColumnTypes)
+		{
+			checkfSlow(UE::Mass::IsA<FMassFragment>(ColumnType),
+				TEXT("Unable to retrieve const column data for a column type without data: %s"), *ColumnType->GetFullName());
+			*ColumnsDataIt++ = Context.GetFragmentView(ColumnType).GetData();
+		}
+	}
+
+	void FQueryContext_MassForwarder::GetMutableColumns(TArrayView<void*> ColumnsData, TConstArrayView<const UScriptStruct*> ColumnTypes)
+	{
+		checkfSlow(ColumnsData.Num() >= ColumnTypes.Num(), TEXT("The list size of requested mutable column types (%i) doesn't match the "
+			"available space for the results (%i)"), ColumnTypes.Num(), ColumnsData.Num());
+
+		void** ColumnsDataIt = ColumnsData.GetData();
+		for (const UScriptStruct* ColumnType : ColumnTypes)
+		{
+			checkfSlow(UE::Mass::IsA<FMassFragment>(ColumnType),
+				TEXT("Unable to retrieve mutable column data for a column type without data: %s"), *ColumnType->GetFullName());
+			*ColumnsDataIt++ = Context.GetMutableFragmentView(ColumnType).GetData();
+		}
+	}
+} // namespace UE::Editor::DataStorage::Queries
\ No newline at end of file
diff --git a/Engine/Plugins/Experimental/EditorDataStorage/Source/TedsCore/Private/Queries/TypedElementQueryContext_MassForwarder.h b/Engine/Plugins/Experimental/EditorDataStorage/Source/TedsCore/Private/Queries/TypedElementQueryContext_MassForwarder.h
new file mode 100644
--- /dev/null
+++ b/Engine/Plugins/Experimental/EditorDataStorage/Source/TedsCore/Private/Queries/TypedElementQueryContext_MassForwarder.h
@@ -0,0 +1,49 @@
+// Copyright Epic Games, Inc. All Rights Reserved.
+
+#pragma once
+
+#include "Elements/Framework/TypedElementQueryContextImplementation.h"
+#include "Elements/Framework/TypedElementQueryFunctions.h"
+
+struct FMassExecutionContext;
+
+namespace UE::Editor::DataStorage::Queries
+{
+	/** Query context implementation for single row queries that call into TEDS Direct API. */
+	class FQueryContext_MassForwarder final : public SingleRowInfo, public RowBatchInfo, public IQueryFunctionResponse
+	{
+	public:
+		static const FName Name;
+
+		explicit FQueryContext_MassForwarder(FMassExecutionContext& Context);
+
+		/**
+		 * @section ISingleRowInfo
+		 */
+
+		virtual bool CurrentRowHasColumns(TConstArrayView<const UScriptStruct*> ColumnTypes) const override;
+		virtual RowHandle GetCurrentRow() const override;
+
+		/**
+		 * @section IRowBatchInfo
+		 */
+		virtual uint32 GetBatchRowCount() const override;
+		virtual TConstArrayView<RowHandle> GetBatchRowHandles() const override;
+		virtual bool CurrentBatchTableHasColumns(TConstArrayView<const UScriptStruct*> ColumnTypes) const override;
+
+		/**
+		 * @section IQueryFunctionResponse
+		 */
+		virtual bool NextBatch() override;
+		virtual bool NextRow() override;
+		virtual void GetConstColumns(TArrayView<const void*> ColumnsData, TConstArrayView<const UScriptStruct*> ColumnTypes) override;
+		virtual void GetMutableColumns(TArrayView<void*> ColumnsData, TConstArrayView<const UScriptStruct*> ColumnTypes) override;
+
+	private:
+		FMassExecutionContext& Context;
+		const RowHandle* CurrentRow;
+		const RowHandle* EndRow;
+	};
+	
+	using QueryContext_MassForwarder = TQueryContextImpl<FQueryContext_MassForwarder>;
+} // namespace UE::Editor::DataStorage::Queries
diff --git a/Engine/Plugins/Experimental/EditorDataStorage/Source/TedsCore/Private/TypedElementDatabase.cpp b/Engine/Plugins/Experimental/EditorDataStorage/Source/TedsCore/Private/TypedElementDatabase.cpp
--- a/Engine/Plugins/Experimental/EditorDataStorage/Source/TedsCore/Private/TypedElementDatabase.cpp
+++ b/Engine/Plugins/Experimental/EditorDataStorage/Source/TedsCore/Private/TypedElementDatabase.cpp
@@ -1280,6 +1280,26 @@ UE::Editor::DataStorage::FQueryResult UEditorDataStorage::RunQuery(
 	}
 }
 
+UE::Editor::DataStorage::FQueryResult UEditorDataStorage::RunQuery(
+	UE::Editor::DataStorage::QueryHandle Query,
+	UE::Editor::DataStorage::ERunQueryFlags Flags,
+	const UE::Editor::DataStorage::Queries::TQueryFunction<void>& Callback)
+{
+	using namespace UE::Editor::DataStorage;
+
+	TRACE_CPUPROFILER_EVENT_SCOPE(TEDS.RunQuery);
+
+	if (ActiveEditorEntityManager)
+	{
+		return Environment->GetQueryStore().RunQuery(*ActiveEditorEntityManager, *Environment, 
+			FExtendedQueryStore::Handle(Query), Flags, Callback);
+	}
+	else
+	{
+		return FQueryResult();
+	}
+}
+
 void UEditorDataStorage::ActivateQueries(FName ActivationName)
 {
 	if (ActiveEditorEntityManager)
diff --git a/Engine/Plugins/Experimental/EditorDataStorage/Source/TedsCore/Private/TypedElementDatabase.h b/Engine/Plugins/Experimental/EditorDataStorage/Source/TedsCore/Private/TypedElementDatabase.h
--- a/Engine/Plugins/Experimental/EditorDataStorage/Source/TedsCore/Private/TypedElementDatabase.h
+++ b/Engine/Plugins/Experimental/EditorDataStorage/Source/TedsCore/Private/TypedElementDatabase.h
@@ -176,6 +176,8 @@ public:
 	UE_API virtual UE::Editor::DataStorage::FQueryResult RunQuery(UE::Editor::DataStorage::QueryHandle Query, UE::Editor::DataStorage::DirectQueryCallbackRef Callback) override;
 	UE_API virtual UE::Editor::DataStorage::FQueryResult RunQuery(UE::Editor::DataStorage::QueryHandle Query, UE::Editor::DataStorage::EDirectQueryExecutionFlags Flags,
 		UE::Editor::DataStorage::DirectQueryCallbackRef Callback) override;
+	UE_API virtual UE::Editor::DataStorage::FQueryResult RunQuery(UE::Editor::DataStorage::QueryHandle Query, UE::Editor::DataStorage::ERunQueryFlags Flags,
+		const UE::Editor::DataStorage::Queries::TQueryFunction<void>& Callback) override;
 	UE_API virtual void ActivateQueries(FName ActivationName) override;
 	
 	UE_API virtual UE::Editor::DataStorage::RowHandle LookupMappedRow(
diff --git a/Engine/Source/Runtime/TypedElementFramework/Public/DataStorage/Queries/Types.h b/Engine/Source/Runtime/TypedElementFramework/Public/DataStorage/Queries/Types.h
--- a/Engine/Source/Runtime/TypedElementFramework/Public/DataStorage/Queries/Types.h
+++ b/Engine/Source/Runtime/TypedElementFramework/Public/DataStorage/Queries/Types.h
@@ -126,6 +126,11 @@ namespace UE::Editor::DataStorage
 			Unsupported,
 			/** The provided query is no longer available. */
 			Unavailable,
+			/** 
+			 * The provided callback is not compatible with the query. This can be due to missing columns in the selection or the access
+			 * to a column is more restricted than the callback requests.
+			 */
+			Incompatible,
 			/** One or more dependencies declared on the query could not be retrieved. */
 			MissingDependency
 		};
diff --git a/Engine/Source/Runtime/TypedElementFramework/Public/Elements/Framework/TypedElementQueryFunctions.h b/Engine/Source/Runtime/TypedElementFramework/Public/Elements/Framework/TypedElementQueryFunctions.h
--- a/Engine/Source/Runtime/TypedElementFramework/Public/Elements/Framework/TypedElementQueryFunctions.h
+++ b/Engine/Source/Runtime/TypedElementFramework/Public/Elements/Framework/TypedElementQueryFunctions.h
@@ -2,6 +2,7 @@
 
 #pragma once
 
+#include <type_traits>
 #include "Containers/ContainersFwd.h"
 #include "Elements/Framework/TypedElementQueryCapabilities.h"
 #include "Elements/Framework/TypedElementQueryFunctionArgumentTypes.h"
@@ -34,6 +35,12 @@ namespace UE::Editor::DataStorage::Queries
 	};
 	ENUM_CLASS_FLAGS(EFunctionCallConfig);
 
+	template<typename T>
+	concept FunctionType =
+		requires { &std::remove_reference_t<T>::operator(); }			// non-generic lambdas & single-operator() functors
+		|| std::is_function_v<std::remove_pointer_t<std::decay_t<T>>>	// functions, function pointers & function references
+		|| std::is_member_function_pointer_v<std::decay_t<T>>;			// member function pointers
+
 	/** Storage for a function that can be used as part of a query.  */
 	template<typename ReturnType>
 	class TQueryFunctionBase
@@ -94,8 +101,12 @@ namespace UE::Editor::DataStorage::Queries
 		}
 	};
 
-	template<typename ReturnType, typename FunctionType>
-	TQueryFunction<ReturnType> BuildQueryFunction(FunctionType&& Function);
+	template<typename Return, FunctionType Function>
+	TQueryFunction<Return> BuildQueryFunction(Function&& Callback);
+
+	template<typename Return, FunctionType Function>
+	TQueryFunction<void> BuildQueryFunction(TResult<Return>& Result, Function&& Callback);
+
 } // namespace UE::Editor::DataStorage::Queries
 
 #include "Elements/Framework/TypedElementQueryFunctions.inl"
diff --git a/Engine/Source/Runtime/TypedElementFramework/Public/Elements/Framework/TypedElementQueryFunctions.inl b/Engine/Source/Runtime/TypedElementFramework/Public/Elements/Framework/TypedElementQueryFunctions.inl
--- a/Engine/Source/Runtime/TypedElementFramework/Public/Elements/Framework/TypedElementQueryFunctions.inl
+++ b/Engine/Source/Runtime/TypedElementFramework/Public/Elements/Framework/TypedElementQueryFunctions.inl
@@ -65,17 +65,17 @@ namespace UE::Editor::DataStorage::Queries
 		{
 			using Type = TFunctionInfoImpl<FunctionType>;
 		};
-		
+
 		template<typename FunctionType>
 		struct TFunctionInfo
 		{
 			using Base = TFunctionInfoSelection<FunctorType<FunctionType>, FunctionType>::Type;
-			
+
 			using ReturnType = typename Base::ReturnType;
 			using ArgumentInfo = typename Base::ArgumentInfo;
 			using ContextType = typename ArgumentInfo::ContextType;
 
-			static constexpr bool bIsValidQueryFunction = 
+			static constexpr bool bIsValidQueryFunction =
 				Base::bIsValidQueryFunction &&
 				(ArgumentInfo::bIsSingle || ArgumentInfo::bIsBatch);
 			static constexpr bool bIsSingleRowProcessor = ArgumentInfo::bIsSingle;
@@ -121,7 +121,7 @@ namespace UE::Editor::DataStorage::Queries
 						static_assert(ArgumentInfo::CountResults() == 1, "Only one TResult can be present.");
 						static_assert(std::is_same_v<typename ArgumentInfo::ResultType, RequestedReturnType>,
 							"The type used for TResult is not compatible with the result type expected for the query function.");
-						static_assert(std::is_same_v<ReturnType, void>, 
+						static_assert(std::is_same_v<ReturnType, void>,
 							"The return type of the provided query is expected to be void for batch processing functions.");
 					}
 				}
@@ -139,7 +139,70 @@ namespace UE::Editor::DataStorage::Queries
 				Result.bIsSingleRowProcessor = bIsSingleRowProcessor;
 			}
 		};
-	}
+
+		template<typename FunctionInfo, typename ArgumentInfo, typename ReturnType, typename FunctionType>
+		void CallBody(
+			TResult<ReturnType>& QueryResult,
+			IContextContract& Contract,
+			IQueryFunctionResponse& Response,
+			typename TQueryFunction<ReturnType>::FunctionSpecializationCallback Specialization,
+			FunctionType&& Callback)
+		{
+			using ArgumentList = typename ArgumentInfo::ArgumentList;
+
+			TConstArrayView<const UScriptStruct*> ConstColumnTypes = ArgumentInfo::ListConstColumns();
+			TConstArrayView<const UScriptStruct*> MutableColumnTypes = ArgumentInfo::ListMutableColumns();
+
+			const void** ConstColumnsData = static_cast<const void**>(FMemory_Alloca(sizeof(void*) * ConstColumnTypes.Num()));
+			void** MutableColumnsData = static_cast<void**>(FMemory_Alloca(sizeof(void*) * MutableColumnTypes.Num()));
+
+			TArrayView<const void*> ConstColumns(ConstColumnsData, ConstColumnTypes.Num());
+			TArrayView<void*> MutableColumns(MutableColumnsData, MutableColumnTypes.Num());
+
+			EFlowControl Flow = EFlowControl::Continue;
+
+			ArgumentList Arguments;
+			ArgumentInfo::SetResult(Arguments, QueryResult);
+			ArgumentInfo::SetContext(Arguments, Contract);
+			ArgumentInfo::SetFlowControl(Arguments, Flow);
+
+			do
+			{
+				Response.GetConstColumns(ConstColumns, ConstColumnTypes);
+				Response.GetMutableColumns(MutableColumns, MutableColumnTypes);
+
+				if (Specialization(Response, ConstColumns, MutableColumns))
+				{
+					ArgumentInfo::SetConstColumns(Arguments, ConstColumns);
+					ArgumentInfo::SetMutableColumns(Arguments, MutableColumns);
+
+					if constexpr (FunctionInfo::bIsSingleRowProcessor)
+					{
+						do
+						{
+							if constexpr (ArgumentInfo::ResultIndex >= 0 || std::is_same_v<ReturnType, void>)
+							{
+								Arguments.ApplyBefore(Callback);
+							}
+							else
+							{
+								QueryResult.Add(Arguments.ApplyBefore(Callback));
+							}
+							ArgumentInfo::IncrementColumns(Arguments);
+						} while (Flow == EFlowControl::Continue && Response.NextRow());
+					}
+					else
+					{
+						Arguments.ApplyBefore(Callback);
+					}
+				}
+				else
+				{
+					break;
+				}
+			} while (Flow == EFlowControl::Continue && Response.NextBatch());
+		}
+	} // namespace Private
 
 	template<typename ReturnType>
 	template<EFunctionCallConfig Config>
@@ -169,77 +232,49 @@ namespace UE::Editor::DataStorage::Queries
 		Function(Result, Contract, Response, Specialization);
 	}
 
-	template<typename ReturnType, typename FunctionType>
-	TQueryFunction<ReturnType> BuildQueryFunction(FunctionType&& Function)
+	template<typename Return, FunctionType Function>
+	TQueryFunction<Return> BuildQueryFunction(Function&& Callback)
 	{
-		using FunctionInfo = Private::TFunctionInfo<FunctionType>;
+		using FunctionInfo = Private::TFunctionInfo<Function>;
 		using ArgumentInfo = typename FunctionInfo::ArgumentInfo;
 		
-		TQueryFunction<ReturnType> Result;
-		FunctionInfo::template StaticValidation<ReturnType>();
+		TQueryFunction<Return> Result;
+		FunctionInfo::template StaticValidation<Return>();
 		FunctionInfo::SetupResult(Result);
 
 		Result.Function =
-			[Callback = Forward<FunctionType>(Function)](
-				TResult<ReturnType>& QueryResult, 
+			[LocalCallback = Forward<Function>(Callback)](
+				TResult<Return>& QueryResult, 
 				IContextContract& Contract, 
 				IQueryFunctionResponse& Response, 
-				TQueryFunction<ReturnType>::FunctionSpecializationCallback Specialization)
+				TQueryFunction<Return>::FunctionSpecializationCallback Specialization)
 			{
-				using ArgumentList = typename ArgumentInfo::ArgumentList;
-
-				TConstArrayView<const UScriptStruct*> ConstColumnTypes = ArgumentInfo::ListConstColumns();
-				TConstArrayView<const UScriptStruct*> MutableColumnTypes = ArgumentInfo::ListMutableColumns();
-
-				const void** ConstColumnsData = static_cast<const void**>(FMemory_Alloca(sizeof(void*) * ConstColumnTypes.Num()));
-				void** MutableColumnsData = static_cast<void**>(FMemory_Alloca(sizeof(void*) * MutableColumnTypes.Num()));
-
-				TArrayView<const void*> ConstColumns(ConstColumnsData, ConstColumnTypes.Num());
-				TArrayView<void*> MutableColumns(MutableColumnsData, MutableColumnTypes.Num());
-
-				EFlowControl Flow = EFlowControl::Continue;
-
-				ArgumentList Arguments;
-				ArgumentInfo::SetResult(Arguments, QueryResult);
-				ArgumentInfo::SetContext(Arguments, Contract);
-				ArgumentInfo::SetFlowControl(Arguments, Flow);
+				Private::CallBody<FunctionInfo, ArgumentInfo, Return>(QueryResult, Contract, Response, Specialization, LocalCallback);
+			};
+		return Result;
+	}
 
-				do
-				{
-					Response.GetConstColumns(ConstColumns, ConstColumnTypes);
-					Response.GetMutableColumns(MutableColumns, MutableColumnTypes);
+	template<typename Return, FunctionType Function>
+	TQueryFunction<void> BuildQueryFunction(TResult<Return>& Result, Function&& Callback)
+	{
+		using FunctionInfo = Private::TFunctionInfo<Function>;
+		using ArgumentInfo = typename FunctionInfo::ArgumentInfo;
 
-					if (Specialization(Response, ConstColumns, MutableColumns))
-					{
-						ArgumentInfo::SetConstColumns(Arguments, ConstColumns);
-						ArgumentInfo::SetMutableColumns(Arguments, MutableColumns);
+		TQueryFunction<Return> CompositedFunction;
+		FunctionInfo::template StaticValidation<Return>();
+		FunctionInfo::SetupResult(CompositedFunction);
 
-						if constexpr (FunctionInfo::bIsSingleRowProcessor)
-						{
-							do
-							{
-								if constexpr (ArgumentInfo::ResultIndex >= 0 || std::is_same_v<ReturnType, void>)
-								{
-									Arguments.ApplyBefore(Callback);
-								}
-								else
-								{
-									QueryResult.Add(Arguments.ApplyBefore(Callback));
-								}
-								ArgumentInfo::IncrementColumns(Arguments);
-							} while (Flow == EFlowControl::Continue && Response.NextRow());
-						}
-						else
-						{
-							Arguments.ApplyBefore(Callback);
-						}
-					}
-					else
-					{
-						break;
-					}
-				} while (Flow == EFlowControl::Continue && Response.NextBatch());
+		Result.Function =
+			[&Result, LocalCallback = Forward<Function>(Callback)](
+				TResult<Return>& QueryResult,
+				IContextContract& Contract,
+				IQueryFunctionResponse& Response,
+				TQueryFunction<Return>::FunctionSpecializationCallback Specialization)
+			{
+				// Ignore QueryResult (which will always be a dummy for TQueryFunction<void>) in favor of the captured TResult.
+				Private::CallBody<FunctionInfo, ArgumentInfo, Return>(Result, Contract, Response, Specialization, LocalCallback);
 			};
-		return Result;
+
+		return CompositedFunction;
 	}
 } // namespace UE::Editor::DataStorage::Queries
diff --git a/Engine/Source/Runtime/TypedElementFramework/Public/Elements/Interfaces/TypedElementDataStorageInterface.h b/Engine/Source/Runtime/TypedElementFramework/Public/Elements/Interfaces/TypedElementDataStorageInterface.h
--- a/Engine/Source/Runtime/TypedElementFramework/Public/Elements/Interfaces/TypedElementDataStorageInterface.h
+++ b/Engine/Source/Runtime/TypedElementFramework/Public/Elements/Interfaces/TypedElementDataStorageInterface.h
@@ -469,8 +469,25 @@ public:
 	 * Directly runs a query. The callback will be called for batches of matching rows. During a single call to RunQuery the callback
 	 * may be called multiple times. If the query handle is invalid or has been deleted nothing happens and the callback won't be called.
 	 */
-	virtual FQueryResult RunQuery(QueryHandle Query, EDirectQueryExecutionFlags Flags,
-		DirectQueryCallbackRef Callback) = 0;
+	virtual FQueryResult RunQuery(QueryHandle Query, EDirectQueryExecutionFlags Flags, DirectQueryCallbackRef Callback) = 0;
+	/**
+	 * Directly runs a query. The callback will be called for batches of matching rows. During a single call to RunQuery the callback
+	 * may be called multiple times. If the query handle is invalid or has been deleted nothing happens and the callback won't be called.
+	 */
+	virtual FQueryResult RunQuery(QueryHandle Query, ERunQueryFlags Flags, const Queries::TQueryFunction<void>& Callback) = 0;
+	/**
+	 * Directly runs a query. The callback will be called for batches of matching rows. During a single call to RunQuery the callback
+	 * may be called multiple times. If the query handle is invalid or has been deleted nothing happens and the callback won't be called.
+	 */
+	template<Queries::FunctionType Function>
+	FQueryResult RunQuery(QueryHandle Query, ERunQueryFlags Flags, Function&& Callback);
+	/**
+	 * Directly runs a query. The callback will be called for batches of matching rows. During a single call to RunQuery the callback
+	 * may be called multiple times. If the query handle is invalid or has been deleted nothing happens and the callback won't be called.
+	 */
+	template<typename ResultType, Queries::FunctionType Function>
+	FQueryResult RunQuery(QueryHandle Query, ERunQueryFlags Flags, Queries::TResult<ResultType>& Result, Function&& Callback);
+	
 	/**
 	 * Triggers all queries registered under the activation name to run for one update cycle. The activatable queries will be activated at
 	 * start of the cycle and disabled at the end of the cycle and act like regular queries for that cycle. This includes not running
@@ -875,4 +892,16 @@ void ICoreProvider::ForEachDynamicColumn(TFunctionRef<void(const UScriptStruct&
 	ForEachDynamicColumn(DynamicColumnTemplateType::StaticStruct(), Callback);
 }
 
+template<Queries::FunctionType Function>
+FQueryResult ICoreProvider::RunQuery(QueryHandle Query, ERunQueryFlags Flags, Function&& Callback)
+{
+	return RunQuery(Query, Flags, Queries::BuildQueryFunction<void>(Forward<Function>(Callback)));
+}
+
+template<typename ResultType, Queries::FunctionType Function>
+FQueryResult ICoreProvider::RunQuery(QueryHandle Query, ERunQueryFlags Flags, Queries::TResult<ResultType>& Result, Function&& Callback)
+{
+	return RunQuery(Query, Flags, Queries::BuildQueryFunction(Result, Forward<Function>(Callback)));
+}
+
 } // namespace UE::Editor::DataStorage
diff --git a/Engine/Source/Runtime/TypedElementFramework/Public/Elements/Interfaces/TypedElementQueryStorageInterfaces.h b/Engine/Source/Runtime/TypedElementFramework/Public/Elements/Interfaces/TypedElementQueryStorageInterfaces.h
--- a/Engine/Source/Runtime/TypedElementFramework/Public/Elements/Interfaces/TypedElementQueryStorageInterfaces.h
+++ b/Engine/Source/Runtime/TypedElementFramework/Public/Elements/Interfaces/TypedElementQueryStorageInterfaces.h
@@ -349,6 +349,23 @@ namespace UE::Editor::DataStorage
 	};
 	ENUM_CLASS_FLAGS(EDirectQueryExecutionFlags);
 
+	// For now this is an exact mirror of EDirectQueryExecutionFlags (except deprecated flags) in order to separate the old and the new 
+	// query callbacks/context. Please keep both in sync until the old query callbacks/context has been fully deprecated.
+	enum class ERunQueryFlags : uint32
+	{
+		Default = 0, //< No settings, use the default behavior.
+		ParallelizeChunks = 1 << 0, //< If set, each chunk is processed on a separate thread.
+		/*
+		 * Requires ParallelizeChunks to be set.
+		 * If set together with ParallelizeChunks, each chunk will be scheduled individually. This allows work to be better
+		 * distributed among threads, but comes at additional scheduling cost. Add this flag if the processing time per chunk varies.
+		 */
+		AutoBalanceParallelChunkProcessing = 1 << 1,
+		IgnoreActiveState = 1 << 2, //< If set, a direct call will not check activatable queries if they're set.
+		AllowBoundQueries = 1 << 3, //< Stops checking if queries are bound when set, otherwise bound queries will assert.
+	};
+	ENUM_CLASS_FLAGS(ERunQueryFlags);
+
 	/**
 	 * Interface to be provided to query callbacks that are directly called through RunQuery from outside a query callback.
 	 */
