TEDS: Remove sync queries and widgets for FTypedElementLocalTransformColumn since it causes various issues when sync'd to actors

For now the transform can be accessed from the Actor ptr in the UObjectColumn directly while the column is unused.

diff --git a/Engine/Plugins/Experimental/EditorDataStorageFeatures/Source/TedsActorCompatibility/Private/Compatibility/TedsActorTransformQueries.cpp b/Engine/Plugins/Experimental/EditorDataStorageFeatures/Source/TedsActorCompatibility/Private/Compatibility/TedsActorTransformQueries.cpp
deleted file mode 100644
diff --git a/Engine/Plugins/Experimental/EditorDataStorageFeatures/Source/TedsActorCompatibility/Private/Compatibility/TedsActorTransformQueries.h b/Engine/Plugins/Experimental/EditorDataStorageFeatures/Source/TedsActorCompatibility/Private/Compatibility/TedsActorTransformQueries.h
deleted file mode 100644
diff --git a/Engine/Plugins/Experimental/EditorDataStorageFeatures/Source/TedsActorCompatibility/Private/TedsActorCompatibilityFactory.cpp b/Engine/Plugins/Experimental/EditorDataStorageFeatures/Source/TedsActorCompatibility/Private/TedsActorCompatibilityFactory.cpp
--- a/Engine/Plugins/Experimental/EditorDataStorageFeatures/Source/TedsActorCompatibility/Private/TedsActorCompatibilityFactory.cpp
+++ b/Engine/Plugins/Experimental/EditorDataStorageFeatures/Source/TedsActorCompatibility/Private/TedsActorCompatibilityFactory.cpp
@@ -46,10 +46,6 @@ void UTedsActorCompatibilityFactory::RegisterTable(UE::Editor::DataStorage::ICor
 		FLevelColumn>(),
 		FName("Editor_StandardActorTable"));
 
-	StandardActorWithTransformTable = DataStorage.RegisterTable(StandardActorTable,
-		TTypedElementColumnTypeList<FTypedElementLocalTransformColumn>(),
-		FName("Editor_StandardActorWithTransformTable"));
-
 	DataStorageCompatibility->RegisterTypeTableAssociation(AActor::StaticClass(), StandardActorTable);
 }
 
diff --git a/Engine/Plugins/Experimental/EditorDataStorageFeatures/Source/TedsActorCompatibility/Private/TedsActorCompatibilityFactory.h b/Engine/Plugins/Experimental/EditorDataStorageFeatures/Source/TedsActorCompatibility/Private/TedsActorCompatibilityFactory.h
--- a/Engine/Plugins/Experimental/EditorDataStorageFeatures/Source/TedsActorCompatibility/Private/TedsActorCompatibilityFactory.h
+++ b/Engine/Plugins/Experimental/EditorDataStorageFeatures/Source/TedsActorCompatibility/Private/TedsActorCompatibilityFactory.h
@@ -29,7 +29,6 @@ private:
 	UE::Editor::DataStorage::ICompatibilityProvider* DataStorageCompatibility{nullptr};
 
 	UE::Editor::DataStorage::TableHandle StandardActorTable{ UE::Editor::DataStorage::InvalidTableHandle };
-	UE::Editor::DataStorage::TableHandle StandardActorWithTransformTable{ UE::Editor::DataStorage::InvalidTableHandle };
 
 	TMap<UWorld*, FDelegateHandle> ActorDestroyedDelegateHandles;
 	FDelegateHandle PostWorldInitializationDelegateHandle;
diff --git a/Engine/Plugins/Experimental/EditorDataStorageFeatures/Source/TedsOutliner/Private/Widgets/TransformHeadsUpWidget.cpp b/Engine/Plugins/Experimental/EditorDataStorageFeatures/Source/TedsOutliner/Private/Widgets/TransformHeadsUpWidget.cpp
deleted file mode 100644
diff --git a/Engine/Plugins/Experimental/EditorDataStorageFeatures/Source/TedsOutliner/Public/Widgets/TransformHeadsUpWidget.h b/Engine/Plugins/Experimental/EditorDataStorageFeatures/Source/TedsOutliner/Public/Widgets/TransformHeadsUpWidget.h
deleted file mode 100644
diff --git a/Engine/Source/Runtime/TypedElementFramework/Public/Elements/Columns/TypedElementTransformColumns.h b/Engine/Source/Runtime/TypedElementFramework/Public/Elements/Columns/TypedElementTransformColumns.h
--- a/Engine/Source/Runtime/TypedElementFramework/Public/Elements/Columns/TypedElementTransformColumns.h
+++ b/Engine/Source/Runtime/TypedElementFramework/Public/Elements/Columns/TypedElementTransformColumns.h
@@ -17,7 +17,7 @@ struct FTypedElementLocalTransformColumn final : public FEditorDataStorageColumn
 	GENERATED_BODY()
 
 	// Transform is not being initialized to avoid spending time on initialization when the
-	// Transform will be updated the first and following ticks after it's creation. If this
+	// Transform will be updated the first and following ticks afr it's creation. If this
 	// isn't initialized at the correct time, then the sync from source or the true initialization
 	// need to be moved to an earlier phase or group.
 	UPROPERTY(meta = (IgnoreForMemberInitializationTest))
