[TEDS] Fix ::Private Namespace overlap issues with Queries::Private in the Attribute Binder

diff --git a/Engine/Source/Runtime/TypedElementFramework/Public/Elements/Framework/TypedElementAttributeBinding.h b/Engine/Source/Runtime/TypedElementFramework/Public/Elements/Framework/TypedElementAttributeBinding.h
--- a/Engine/Source/Runtime/TypedElementFramework/Public/Elements/Framework/TypedElementAttributeBinding.h
+++ b/Engine/Source/Runtime/TypedElementFramework/Public/Elements/Framework/TypedElementAttributeBinding.h
@@ -110,7 +110,7 @@ namespace UE::Editor::DataStorage
 		 * @param InDefaultValue The default value to be used when the column isn't present on a row\
 		 */
 		template <typename DataType, TDataColumnType ColumnType, typename FunctionType>
-			requires Private::AttributeBinderInvocable<FunctionType, DataType>
+			requires DataStorage::Private::AttributeBinderInvocable<FunctionType, DataType>
 		auto BindData(DataType ColumnType::* InVariable, FunctionType&& InConverter, const DataType& InDefaultValue = DataType());
 
 		/**
@@ -123,7 +123,7 @@ namespace UE::Editor::DataStorage
 		 * @return A TAttribute bound to the row, column pair specified
 		 */
 		template <typename DataType, TDynamicColumnTemplate ColumnType, typename FunctionType>
-			requires Private::AttributeBinderInvocable<FunctionType, DataType>
+			requires DataStorage::Private::AttributeBinderInvocable<FunctionType, DataType>
 		auto BindData(const FName& InIdentifier, DataType ColumnType::* InVariable, FunctionType&& InConverter, const DataType& InDefaultValue = DataType());
 
 		/**
@@ -148,7 +148,7 @@ namespace UE::Editor::DataStorage
 		 * @return A TAttribute bound to the row, column pair specified
 		 */
 		template <TDataColumnType ColumnType, typename FunctionType>
-			requires Private::AttributeBinderColumnInvocable<FunctionType, ColumnType>
+			requires DataStorage::Private::AttributeBinderColumnInvocable<FunctionType, ColumnType>
 		auto BindColumn(FunctionType&& InConverter);
 		
 		/**
@@ -175,7 +175,7 @@ namespace UE::Editor::DataStorage
 		 * @return A TAttribute bound to the row, column pair specified
 		 */
 		template <typename FunctionType>
-			requires Private::AttributeBinderColumnDataInvocable<FunctionType>
+			requires DataStorage::Private::AttributeBinderColumnDataInvocable<FunctionType>
 		auto BindColumnData(const TWeakObjectPtr<const UScriptStruct>& InColumnType, FunctionType&& InConverter);
 
 		/**
diff --git a/Engine/Source/Runtime/TypedElementFramework/Public/Elements/Framework/TypedElementAttributeBinding.inl b/Engine/Source/Runtime/TypedElementFramework/Public/Elements/Framework/TypedElementAttributeBinding.inl
--- a/Engine/Source/Runtime/TypedElementFramework/Public/Elements/Framework/TypedElementAttributeBinding.inl
+++ b/Engine/Source/Runtime/TypedElementFramework/Public/Elements/Framework/TypedElementAttributeBinding.inl
@@ -14,129 +14,125 @@ namespace UE::Editor::DataStorage
 	template <typename AttributeType, TDataColumnType ColumnType>
 	TAttribute<AttributeType> FAttributeBinder::BindData(AttributeType ColumnType::* InVariable, const AttributeType& InDefaultValue)
 	{
-		using namespace Private;
+		using namespace DataStorage::Private;
 
-		if(!DataStorage)
+		if(DataStorage)
 		{
-			return TAttribute<AttributeType>();
-		}
-		
-		// Create a direct property and bind it to the given variable
-		TProperty<AttributeType> Prop;
-		Prop.Bind(InVariable);
-
-		// We don't want any references to this in the lambda because binders are designed to be used and destructed on the stack
-		return TAttribute<AttributeType>::CreateLambda([Property = MoveTemp(Prop), Storage = DataStorage, Row = TargetRow,
-			DefaultValue = InDefaultValue]()
-		{
-			// Get the column from the given row and use that to return the stored property
-			if(ColumnType* Column = Private::GetColumn<ColumnType>(Storage, Row))
+			// Create a direct property and bind it to the given variable
+			TProperty<AttributeType> Prop;
+			Prop.Bind(InVariable);
+	
+			// We don't want any references to this in the lambda because binders are designed to be used and destructed on the stack
+			return TAttribute<AttributeType>::CreateLambda([Property = MoveTemp(Prop), Storage = DataStorage, Row = TargetRow,
+				DefaultValue = InDefaultValue]()
 			{
-				return Property.Get(Column, ColumnType::StaticStruct());
-			}
-			return DefaultValue;
-		});
+				// Get the column from the given row and use that to return the stored property
+				if(ColumnType* Column = GetColumn<ColumnType>(Storage, Row))
+				{
+					return Property.Get(Column, ColumnType::StaticStruct());
+				}
+				return DefaultValue;
+			});
+		}
+		return TAttribute<AttributeType>();
 	}
 
 	template <typename AttributeType, TDynamicColumnTemplate ColumnType>
 	TAttribute<AttributeType> FAttributeBinder::BindData(const FName& InIdentifier, AttributeType ColumnType::* InVariable, const AttributeType& InDefaultValue)
 	{
-		using namespace Private;
+		using namespace DataStorage::Private;
 
-		if(!DataStorage)
+		if(DataStorage)
 		{
-			return TAttribute<AttributeType>();
-		}
-		
-		// Create a direct property and bind it to the given variable
-		TProperty<AttributeType> Prop;
-		Prop.Bind(InVariable);
+			// Create a direct property and bind it to the given variable
+			TProperty<AttributeType> Prop;
+			Prop.Bind(InVariable);
 
-		// We don't want any references to this in the lambda because binders are designed to be used and destructed on the stack
-		return TAttribute<AttributeType>::CreateLambda([Property = MoveTemp(Prop), Storage = DataStorage, Row = TargetRow,
-			DefaultValue = InDefaultValue, Identifier = InIdentifier]()
-		{
-			// Get the column from the given row and use that to return the stored property
-			if(ColumnType* Column = Private::GetColumn<ColumnType>(Storage, Row, Identifier))
+			// We don't want any references to this in the lambda because binders are designed to be used and destructed on the stack
+			return TAttribute<AttributeType>::CreateLambda([Property = MoveTemp(Prop), Storage = DataStorage, Row = TargetRow,
+				DefaultValue = InDefaultValue, Identifier = InIdentifier]()
 			{
-				return Property.Get(Column, ColumnType::StaticStruct());
-			}
-			return DefaultValue;
-		});
+				// Get the column from the given row and use that to return the stored property
+				if(ColumnType* Column = GetColumn<ColumnType>(Storage, Row, Identifier))
+				{
+					return Property.Get(Column, ColumnType::StaticStruct());
+				}
+				return DefaultValue;
+			});
+		}
+		return TAttribute<AttributeType>();
 	}
 
 	template <typename AttributeType, typename DataType, TDataColumnType ColumnType>
 	TAttribute<AttributeType> FAttributeBinder::BindData(DataType ColumnType::* InVariable, TFunction<AttributeType(const DataType&)> InConverter, const DataType& InDefaultValue)
 	{
-		using namespace Private;
+		using namespace DataStorage::Private;
 
-		if(!DataStorage)
+		if(DataStorage)
 		{
-			return TAttribute<AttributeType>();
-		}
-		
-		AttributeType ConvertedDefault = InConverter(InDefaultValue);
+			AttributeType ConvertedDefault = InConverter(InDefaultValue);
 
-		// Create a convertible property and bind it to the given variable
-		TProperty<AttributeType> Prop;
-		Prop.Bind(InVariable, MoveTemp(InConverter));
+			// Create a convertible property and bind it to the given variable
+			TProperty<AttributeType> Prop;
+			Prop.Bind(InVariable, MoveTemp(InConverter));
 	
-		return TAttribute<AttributeType>::CreateLambda(
-			[
-				Property = MoveTemp(Prop), 
-				Storage = DataStorage, 
-				Row = TargetRow, 
-				DefaultValue = InDefaultValue, 
-				ConvertedDefault = MoveTemp(ConvertedDefault)
-			]()
-			{
-				if(ColumnType* Column = Private::GetColumn<ColumnType>(Storage, Row))
+			return TAttribute<AttributeType>::CreateLambda(
+				[
+					Property = MoveTemp(Prop), 
+					Storage = DataStorage, 
+					Row = TargetRow, 
+					DefaultValue = InDefaultValue, 
+					ConvertedDefault = MoveTemp(ConvertedDefault)
+				]()
 				{
-					return Property.Get(Column, ColumnType::StaticStruct());
-				}
+					if(ColumnType* Column = GetColumn<ColumnType>(Storage, Row))
+					{
+						return Property.Get(Column, ColumnType::StaticStruct());
+					}
 
-				return ConvertedDefault;
-			});
+					return ConvertedDefault;
+				});
+		}
+		return TAttribute<AttributeType>();
 	}
 
 	template <typename AttributeType, typename DataType, TDynamicColumnTemplate ColumnType>
 	TAttribute<AttributeType> FAttributeBinder::BindData(const FName& InIdentifier, DataType ColumnType::* InVariable,
 		TFunction<AttributeType(const DataType&)> InConverter, const DataType& InDefaultValue)
 	{
-		using namespace Private;
+		using namespace DataStorage::Private;
 
-		if(!DataStorage)
+		if(DataStorage)
 		{
-			return TAttribute<AttributeType>();
-		}
-		
-		AttributeType ConvertedDefault = InConverter(InDefaultValue);
+			AttributeType ConvertedDefault = InConverter(InDefaultValue);
 
-		// Create a convertible property and bind it to the given variable
-		TProperty<AttributeType> Prop;
-		Prop.Bind(InVariable, MoveTemp(InConverter));
+			// Create a convertible property and bind it to the given variable
+			TProperty<AttributeType> Prop;
+			Prop.Bind(InVariable, MoveTemp(InConverter));
 	
-		return TAttribute<AttributeType>::CreateLambda(
-			[
-				Property = MoveTemp(Prop), 
-				Storage = DataStorage, 
-				Row = TargetRow, 
-				DefaultValue = InDefaultValue, 
-				ConvertedDefault = MoveTemp(ConvertedDefault),
-				Identifier = InIdentifier
-			]()
-			{
-				if(ColumnType* Column = Private::GetColumn<ColumnType>(Storage, Row, Identifier))
+			return TAttribute<AttributeType>::CreateLambda(
+				[
+					Property = MoveTemp(Prop), 
+					Storage = DataStorage, 
+					Row = TargetRow, 
+					DefaultValue = InDefaultValue, 
+					ConvertedDefault = MoveTemp(ConvertedDefault),
+					Identifier = InIdentifier
+				]()
 				{
-					return Property.Get(Column, ColumnType::StaticStruct());
-				}
+					if(ColumnType* Column = GetColumn<ColumnType>(Storage, Row, Identifier))
+					{
+						return Property.Get(Column, ColumnType::StaticStruct());
+					}
 
-				return ConvertedDefault;
-			});
+					return ConvertedDefault;
+				});
+		}
+		return TAttribute<AttributeType>();
 	}
 	
 	template <typename DataType, TDataColumnType ColumnType, typename FunctionType>
-		requires Private::AttributeBinderInvocable<FunctionType, DataType>
+		requires DataStorage::Private::AttributeBinderInvocable<FunctionType, DataType>
 	auto FAttributeBinder::BindData(DataType ColumnType::* InVariable, FunctionType&& InConverter, const DataType& InDefaultValue)
 	{
 		// Deduce the attribute type from the return value of the converter function
@@ -146,7 +142,7 @@ namespace UE::Editor::DataStorage
 	}
 	
 	template <typename DataType, TDynamicColumnTemplate ColumnType, typename FunctionType>
-		requires Private::AttributeBinderInvocable<FunctionType, DataType>
+		requires DataStorage::Private::AttributeBinderInvocable<FunctionType, DataType>
 	auto FAttributeBinder::BindData(const FName& InIdentifier, DataType ColumnType::* InVariable, FunctionType&& InConverter, const DataType& InDefaultValue)
 	{
 		// Deduce the attribute type from the return value of the converter function
@@ -158,24 +154,28 @@ namespace UE::Editor::DataStorage
 	template <typename AttributeType, TDataColumnType ColumnType>
 	TAttribute<AttributeType> FAttributeBinder::BindColumn(TFunction<AttributeType(const ColumnType&)> InConverter)
 	{
-		return TAttribute<AttributeType>::CreateLambda(
-			[
-				Convertor = InConverter, 
-				Storage = DataStorage, 
-				Row = TargetRow 
-			]()
-			{
-				if(ColumnType* Column = Storage->GetColumn<ColumnType>(Row))
+		if(DataStorage)
+		{
+			return TAttribute<AttributeType>::CreateLambda(
+				[
+					Convertor = InConverter, 
+					Storage = DataStorage, 
+					Row = TargetRow 
+				]()
 				{
-					return Convertor(*Column);
-				}
+					if(ColumnType* Column = Storage->GetColumn<ColumnType>(Row))
+					{
+						return Convertor(*Column);
+					}
 
-				return AttributeType();
-			});
+					return AttributeType();
+				});
+		}
+		return TAttribute<AttributeType>();
 	}
 		
 	template <TDataColumnType ColumnType, typename FunctionType>
-		requires Private::AttributeBinderColumnInvocable<FunctionType, ColumnType>
+		requires DataStorage::Private::AttributeBinderColumnInvocable<FunctionType, ColumnType>
 	auto FAttributeBinder::BindColumn(FunctionType&& InConverter)
 	{
 		// Deduce the attribute type from the return value of the converter function
@@ -188,25 +188,29 @@ namespace UE::Editor::DataStorage
 	TAttribute<AttributeType> FAttributeBinder::BindColumnData(const TWeakObjectPtr<const UScriptStruct>& InColumnType,
 		const TFunction<AttributeType(const TWeakObjectPtr<const UScriptStruct>&, const void*)>& InConverter)
 	{
-		return TAttribute<AttributeType>::CreateLambda(
-			[
-				ColumnType = InColumnType,
-				Convertor = InConverter, 
-				Storage = DataStorage, 
-				Row = TargetRow 
-			]()
-			{
-				if(const void* ColumnData = Storage->GetColumnData(Row, ColumnType.Get()))
+		if(DataStorage)
+		{
+			return TAttribute<AttributeType>::CreateLambda(
+				[
+					ColumnType = InColumnType,
+					Convertor = InConverter, 
+					Storage = DataStorage, 
+					Row = TargetRow 
+				]()
 				{
-					return Convertor(ColumnType, ColumnData);
-				}
+					if(const void* ColumnData = Storage->GetColumnData(Row, ColumnType.Get()))
+					{
+						return Convertor(ColumnType, ColumnData);
+					}
 
-				return AttributeType();
-			});
+					return AttributeType();
+				});
+		}
+		return TAttribute<AttributeType>();
 	}
 
 	template <typename FunctionType>
-			requires Private::AttributeBinderColumnDataInvocable<FunctionType>
+			requires DataStorage::Private::AttributeBinderColumnDataInvocable<FunctionType>
 	auto FAttributeBinder::BindColumnData(const TWeakObjectPtr<const UScriptStruct>& InColumnType, FunctionType&& InConverter)
 	{
 		// Deduce the attribute type from the return value of the converter function
@@ -219,7 +223,7 @@ namespace UE::Editor::DataStorage
 	template <typename InRetValType, typename... ParamTypes, TDataColumnType ColumnType>
 		TDelegate<InRetValType(ParamTypes...)> FAttributeBinder::BindEvent(TDelegate<InRetValType(ParamTypes...)> ColumnType::* InVariable)
 	{
-		using namespace Private;
+		using namespace DataStorage::Private;
 		
 		// Create a property for the delegate
 		TProperty<TDelegate<InRetValType(ParamTypes...)>> Prop;
@@ -228,7 +232,7 @@ namespace UE::Editor::DataStorage
 		return TDelegate<InRetValType(ParamTypes...)>::CreateLambda([Property = MoveTemp(Prop), Storage = DataStorage,
 			Row = TargetRow](ParamTypes&&... Params)
 		{
-			if(ColumnType* Column = Private::GetColumn<ColumnType>(Storage, Row))
+			if(ColumnType* Column = GetColumn<ColumnType>(Storage, Row))
 			{
 				// Get the delegate in the bound column for the specified row
 				TDelegate<InRetValType(ParamTypes...)> Delegate = Property.Get(Column, ColumnType::StaticStruct());
@@ -246,7 +250,7 @@ namespace UE::Editor::DataStorage
 	template <typename InRetValType, typename... ParamTypes, TDynamicColumnTemplate ColumnType>
 	TDelegate<InRetValType(ParamTypes...)> FAttributeBinder::BindEvent(const FName& InIdentifier, TDelegate<InRetValType(ParamTypes...)> ColumnType::* InVariable)
 	{
-		using namespace Private;
+		using namespace DataStorage::Private;
 		
 		// Create a property for the delegate
 		TProperty<TDelegate<InRetValType(ParamTypes...)>> Prop;
@@ -255,7 +259,7 @@ namespace UE::Editor::DataStorage
 		return TDelegate<InRetValType(ParamTypes...)>::CreateLambda([Property = MoveTemp(Prop), Storage = DataStorage,
 			Row = TargetRow, Identifier = InIdentifier](ParamTypes&&... Params)
 		{
-			if(ColumnType* Column = Private::GetColumn<ColumnType>(Storage, Row, Identifier))
+			if(ColumnType* Column = GetColumn<ColumnType>(Storage, Row, Identifier))
 			{
 				// Get the delegate in the bound column for the specified row
 				TDelegate<InRetValType(ParamTypes...)> Delegate = Property.Get(Column, ColumnType::StaticStruct());
diff --git a/Engine/Source/Runtime/TypedElementFramework/Public/Elements/Framework/TypedElementAttributeBindingText.h b/Engine/Source/Runtime/TypedElementFramework/Public/Elements/Framework/TypedElementAttributeBindingText.h
--- a/Engine/Source/Runtime/TypedElementFramework/Public/Elements/Framework/TypedElementAttributeBindingText.h
+++ b/Engine/Source/Runtime/TypedElementFramework/Public/Elements/Framework/TypedElementAttributeBindingText.h
@@ -44,7 +44,7 @@ namespace UE::Editor::DataStorage
 			FName ColumnIdentifier = NAME_None);
 			
 		template <typename DataType, TDataColumnType ColumnType, typename FunctionType> 
-			requires Private::AttributeBinderInvocable<FunctionType, DataType>
+			requires DataStorage::Private::AttributeBinderInvocable<FunctionType, DataType>
 		FTextAttributeFormatted& Arg(FString Name, DataType ColumnType::* Variable, FunctionType Converter, FFormatArgumentValue Default = {},
 			FName ColumnIdentifier = NAME_None);
 	
@@ -55,7 +55,7 @@ namespace UE::Editor::DataStorage
 
 		struct FPropertyInfo
 		{
-			Private::TProperty<FText> Property;
+			DataStorage::Private::TProperty<FText> Property;
 			FFormatArgumentValue Default;
 			FName DynamicColumnIdentifier;
 		};
@@ -72,7 +72,7 @@ namespace UE::Editor::DataStorage
 	FTextAttributeFormatted& FTextAttributeFormatted::Arg(
 		FString Name, FText ColumnType::* Variable, FFormatArgumentValue Default, FName ColumnIdentifier)
 	{
-		Private::TProperty<FText> Prop;
+		DataStorage::Private::TProperty<FText> Prop;
 		Prop.Bind(Variable);
 		NamedProperties.Add(MoveTemp(Name), FPropertyInfo{ .Property = MoveTemp(Prop), .Default = MoveTemp(Default),
 			.DynamicColumnIdentifier = ColumnIdentifier});
@@ -102,7 +102,7 @@ namespace UE::Editor::DataStorage
 		FFormatArgumentValue Default,
 		FName ColumnIdentifier)
 	{
-		Private::TProperty<FText> Prop;
+		DataStorage::Private::TProperty<FText> Prop;
 		Prop.Bind(Variable, Converter);
 		NamedProperties.Add(MoveTemp(Name), FPropertyInfo{ .Property = MoveTemp(Prop), .Default = MoveTemp(Default),
 			.DynamicColumnIdentifier = ColumnIdentifier });
@@ -111,7 +111,7 @@ namespace UE::Editor::DataStorage
 	}
 
 	template <typename DataType, TDataColumnType ColumnType, typename FunctionType>
-		requires Private::AttributeBinderInvocable<FunctionType, DataType>
+		requires DataStorage::Private::AttributeBinderInvocable<FunctionType, DataType>
 	FTextAttributeFormatted& FTextAttributeFormatted::Arg(
 		FString Name,
 		DataType ColumnType::* Variable,
