TTypedElementInternalDataStore<void> : Fix the array indices used in AddDataForElement and RemoveDataForElement

ElementIdToArrayIndex contains the mapping from ElementId to ArrayIndex for the index to use in the InternalDataArray.  These 2 methods were erroneously doing the array lookup using the ElementId directly.

diff --git a/Engine/Source/Runtime/TypedElementFramework/Public/Elements/Framework/TypedElementData.h b/Engine/Source/Runtime/TypedElementFramework/Public/Elements/Framework/TypedElementData.h
--- a/Engine/Source/Runtime/TypedElementFramework/Public/Elements/Framework/TypedElementData.h
+++ b/Engine/Source/Runtime/TypedElementFramework/Public/Elements/Framework/TypedElementData.h
@@ -641,7 +641,7 @@ public:
 
 		ElementIdToArrayIndex.Add(InOutElementId, InternalDataArrayIndex);
 
-		TTypedElementInternalData<void>& InternalData = InternalDataArray[InOutElementId];
+		TTypedElementInternalData<void>& InternalData = InternalDataArray[InternalDataArrayIndex];
 		InternalData.Initialize(InTypeId, InOutElementId);
 		return InternalData;
 	}
@@ -675,7 +675,7 @@ public:
 
 		checkSlow(InternalDataArray.IsValidIndex(InternalDataArrayIndex));
 
-		TTypedElementInternalData<void>& InternalData = InternalDataArray[InElementId];
+		TTypedElementInternalData<void>& InternalData = InternalDataArray[InternalDataArrayIndex];
 		checkf(InExpectedDataPtr == &InternalData, TEXT("Internal data pointer did not match the expected value! Does this handle belong to a different element registry?"));
 		InternalData.LogExternalReferencesOnDestruction();
 		InternalData.Reset();
