[TEDS] Add an explicit function to the factories to allow hierarchies to be registered.

diff --git a/Engine/Plugins/Experimental/EditorDataStorage/Source/TedsCore/Private/TypedElementsDataStorage.cpp b/Engine/Plugins/Experimental/EditorDataStorage/Source/TedsCore/Private/TypedElementsDataStorage.cpp
--- a/Engine/Plugins/Experimental/EditorDataStorage/Source/TedsCore/Private/TypedElementsDataStorage.cpp
+++ b/Engine/Plugins/Experimental/EditorDataStorage/Source/TedsCore/Private/TypedElementsDataStorage.cpp
@@ -127,6 +127,8 @@ void FEditorDataStorageModule::StartupModule()
 					Factory->RegisterTables(*DataStorage, *DataStorageCompatibility);
 					
 					Factory->RegisterTickGroups(*DataStorage);
+
+					Factory->RegisterHierarchies(*DataStorage);
 					
 					Factory->RegisterRegistrationFilters(*DataStorageCompatibility);
 					Factory->RegisterDealiaser(*DataStorageCompatibility);
diff --git a/Engine/Source/Runtime/TypedElementFramework/Public/Elements/Interfaces/TypedElementDataStorageFactory.h b/Engine/Source/Runtime/TypedElementFramework/Public/Elements/Interfaces/TypedElementDataStorageFactory.h
--- a/Engine/Source/Runtime/TypedElementFramework/Public/Elements/Interfaces/TypedElementDataStorageFactory.h
+++ b/Engine/Source/Runtime/TypedElementFramework/Public/Elements/Interfaces/TypedElementDataStorageFactory.h
@@ -38,6 +38,8 @@ public:
 	virtual void RegisterTables(UE::Editor::DataStorage::ICoreProvider& DataStorage, UE::Editor::DataStorage::ICompatibilityProvider& DataStorageCompatibility) {}
 	virtual void RegisterTickGroups(UE::Editor::DataStorage::ICoreProvider& DataStorage) const {}
 	virtual void RegisterQueries(UE::Editor::DataStorage::ICoreProvider& DataStorage) {}
+	
+	virtual void RegisterHierarchies(UE::Editor::DataStorage::ICoreProvider& DataStorage) {}
 
 	virtual void RegisterRegistrationFilters(UE::Editor::DataStorage::ICompatibilityProvider& DataStorageCompatibility) const {}
 	virtual void RegisterDealiaser(UE::Editor::DataStorage::ICompatibilityProvider& DataStorageCompatibility) const {}
