[TypedElementFramework] Fixed display name typo for level tag.

diff --git a/Engine/Source/Runtime/TypedElementFramework/Public/Elements/Columns/TypedElementCompatibilityColumns.h b/Engine/Source/Runtime/TypedElementFramework/Public/Elements/Columns/TypedElementCompatibilityColumns.h
--- a/Engine/Source/Runtime/TypedElementFramework/Public/Elements/Columns/TypedElementCompatibilityColumns.h
+++ b/Engine/Source/Runtime/TypedElementFramework/Public/Elements/Columns/TypedElementCompatibilityColumns.h
@@ -109,7 +109,7 @@ struct FTypedElementLevelColumn final : public FEditorDataStorageColumn
 * Tag to identify a row with a ULevel. If there's a FTypedElementUObjectColumn, the stored object will be
  * a ULevel if this tag is present.
  */
-USTRUCT(meta = (DisplayName = "World Tag"))
+USTRUCT(meta = (DisplayName = "Level Tag"))
 struct FEditorDataStorageLevelTag final : public FEditorDataStorageTag
 {
 	GENERATED_BODY()
