FDisplayNameColumn : Add Sortable metadata

Also clean up comments that no longer apply (these are not dynamic columns)

diff --git a/Engine/Source/Runtime/TypedElementFramework/Public/Elements/Columns/TypedElementSlateWidgetColumns.h b/Engine/Source/Runtime/TypedElementFramework/Public/Elements/Columns/TypedElementSlateWidgetColumns.h
--- a/Engine/Source/Runtime/TypedElementFramework/Public/Elements/Columns/TypedElementSlateWidgetColumns.h
+++ b/Engine/Source/Runtime/TypedElementFramework/Public/Elements/Columns/TypedElementSlateWidgetColumns.h
@@ -50,22 +50,18 @@ struct FTypedElementSlateWidgetReferenceDeletesRowTag final : public FEditorData
 
 /**
  * A localized display name for this row.
- * 
- * This can be used as a dynamic column to specify display names for multiple items in a row.
  */
 USTRUCT(meta = (DisplayName = "Display Name"))
 struct FDisplayNameColumn final : public FEditorDataStorageColumn
 {
 	GENERATED_BODY()
 
-	UPROPERTY(meta = (Searchable))
+	UPROPERTY(meta = (Searchable, Sortable))
 	FText DisplayName;
 };
 
 /**
  * A localized description for this row.
- * 
- * This can be used as a dynamic column to specify descriptions for multiple items in a row.
  */
 USTRUCT(meta = (DisplayName = "Description"))
 struct FDescriptionColumn final : public FEditorDataStorageColumn
