[Attribute Binding] Add a new TEDS attribute binding to work on whether the row has the column (allows tags to be used and bound easily to attributes)

diff --git a/Engine/Source/Runtime/TypedElementFramework/Public/Elements/Framework/TypedElementAttributeBinding.h b/Engine/Source/Runtime/TypedElementFramework/Public/Elements/Framework/TypedElementAttributeBinding.h
--- a/Engine/Source/Runtime/TypedElementFramework/Public/Elements/Framework/TypedElementAttributeBinding.h
+++ b/Engine/Source/Runtime/TypedElementFramework/Public/Elements/Framework/TypedElementAttributeBinding.h
@@ -151,6 +151,50 @@ namespace UE::Editor::DataStorage
 			requires DataStorage::Private::AttributeBinderColumnInvocable<FunctionType, ColumnType>
 		auto BindColumn(FunctionType&& InConverter);
 		
+		/**
+         * Bind the presence of a TEDS tag or column to a slate attribute (instead of a single member variable). This allows you to derive an
+         * attribute from the presence of a tag on a row
+         *
+         * @param InConverter Conversion function to convert from a bool (aka bHasColumn) -> AttributeType
+         * @return A TAttribute bound to the row, column pair specified
+         *
+         * Example:
+         *			const TAttribute TestHasIntAttribute(Binder.BindColumnsPresence<FTestHasIntTag>([](bool bHasColumn)
+         *			{
+         *				return bHasColumn ? Image1 : Image2;
+         *			}));
+         */
+        template <typename AttributeType, TColumnType... ColumnTypes>
+        TAttribute<AttributeType> BindColumnsPresence(TFunction<AttributeType(const bool)> InConverter);
+
+		/**
+		 * Overload for the BindColumnPresence to accept lambdas instead of TFunctions
+		 * @param InConverter Conversion function to convert from a bool (aka bHasColumn) -> AttributeType
+		 * @return A TAttribute bound to the row, column pair specified
+		 */
+		template <TColumnType... ColumnTypes, typename FunctionType>
+			requires DataStorage::Private::AttributeBinderColumnPresenceInvocable<FunctionType>
+		auto BindColumnsPresence(FunctionType&& InConverter);
+		
+		/**
+		 * Overload for the BindColumnPresence to accept FConditions instead of TColumnTypes
+		 * @param InConditions Conditions to check the row if it matches (can include dynamic columns)
+		 * @param InConverter Conversion function to convert from a bool (aka bHasColumn) -> AttributeType
+		 * @return A TAttribute bound to the row, column pair specified
+		 */
+		template <typename AttributeType>
+		TAttribute<AttributeType> BindColumnsPresence(Queries::FConditions InConditions, TFunction<AttributeType(const bool)> InConverter);
+
+		/**
+		 * Overload for the BindColumnPresence(FConditions) to accept lambdas instead of TFunctions
+		 * @param InConditions Conditions to check the row if it matches (can include dynamic columns)
+		 * @param InConverter Conversion function to convert from a bool (aka bHasColumn) -> AttributeType
+		 * @return A TAttribute bound to the row, column pair specified
+		 */
+		template <typename FunctionType>
+			requires DataStorage::Private::AttributeBinderColumnPresenceInvocable<FunctionType>
+		auto BindColumnsPresence(Queries::FConditions InConditions, FunctionType&& InConverter);
+		
 		/**
 		 * Bind a whole TEDS column to a slate attribute (instead of a single member variable) using the column's typeinfo.
 		 *
diff --git a/Engine/Source/Runtime/TypedElementFramework/Public/Elements/Framework/TypedElementAttributeBinding.inl b/Engine/Source/Runtime/TypedElementFramework/Public/Elements/Framework/TypedElementAttributeBinding.inl
--- a/Engine/Source/Runtime/TypedElementFramework/Public/Elements/Framework/TypedElementAttributeBinding.inl
+++ b/Engine/Source/Runtime/TypedElementFramework/Public/Elements/Framework/TypedElementAttributeBinding.inl
@@ -184,6 +184,67 @@ namespace UE::Editor::DataStorage
 		return BindColumn<AttributeType, ColumnType>(TFunction<AttributeType(const ColumnType&)>(InConverter));
 	}
 	
+	template <typename AttributeType, TColumnType... ColumnTypes>
+    TAttribute<AttributeType> FAttributeBinder::BindColumnsPresence(TFunction<AttributeType(const bool)> InConverter)
+    {
+    	if(DataStorage)
+    	{
+    		return TAttribute<AttributeType>::CreateLambda(
+    			[
+    				Convertor = InConverter, 
+    				Storage = DataStorage, 
+    				Row = TargetRow 
+    			]()
+    			{
+    				return Convertor(Storage->HasColumns<ColumnTypes...>(Row));
+    			});
+    	}
+    	return TAttribute<AttributeType>();
+    }
+    	
+	template <TColumnType... ColumnTypes, typename FunctionType>
+		requires DataStorage::Private::AttributeBinderColumnPresenceInvocable<FunctionType>
+	auto FAttributeBinder::BindColumnsPresence(FunctionType&& InConverter)
+	{
+		// Deduce the attribute type from the return value of the converter function
+		using AttributeType = decltype(InConverter(std::declval<const bool>()));
+    	
+		return BindColumnsPresence<AttributeType, ColumnTypes...>(TFunction<AttributeType(const bool)>(MoveTemp(InConverter)));
+	}
+
+	template <typename AttributeType>
+	TAttribute<AttributeType> FAttributeBinder::BindColumnsPresence(Queries::FConditions InConditions, TFunction<AttributeType(const bool)> InConverter)
+	{
+		if(DataStorage)
+		{
+			if (!InConditions.IsCompiled())
+			{
+				InConditions.Compile(Queries::FEditorStorageQueryConditionCompileContext(DataStorage));
+			}
+			return TAttribute<AttributeType>::CreateLambda(
+				[
+					Conditions = InConditions,
+					Convertor = InConverter, 
+					Storage = DataStorage, 
+					Row = TargetRow 
+				]()
+				{
+					return Convertor(Storage->MatchesColumns(Row, Conditions));
+				});
+		}
+		return TAttribute<AttributeType>();
+	}
+    	
+	template <typename FunctionType>
+		requires DataStorage::Private::AttributeBinderColumnPresenceInvocable<FunctionType>
+	auto FAttributeBinder::BindColumnsPresence(Queries::FConditions InConditions, FunctionType&& InConverter)
+	{
+		// Deduce the attribute type from the return value of the converter function
+		using AttributeType = decltype(InConverter(std::declval<const bool>()));
+    	
+		return BindColumnsPresence<AttributeType>(InConditions, TFunction<AttributeType(const bool)>(MoveTemp(InConverter)));
+	}
+	
 	template <typename AttributeType>
 	TAttribute<AttributeType> FAttributeBinder::BindColumnData(const TWeakObjectPtr<const UScriptStruct>& InColumnType,
 		const TFunction<AttributeType(const TWeakObjectPtr<const UScriptStruct>&, const void*)>& InConverter)
diff --git a/Engine/Source/Runtime/TypedElementFramework/Public/Elements/Framework/TypedElementAttributeBindingProperty.h b/Engine/Source/Runtime/TypedElementFramework/Public/Elements/Framework/TypedElementAttributeBindingProperty.h
--- a/Engine/Source/Runtime/TypedElementFramework/Public/Elements/Framework/TypedElementAttributeBindingProperty.h
+++ b/Engine/Source/Runtime/TypedElementFramework/Public/Elements/Framework/TypedElementAttributeBindingProperty.h
@@ -23,6 +23,11 @@ namespace UE::Editor::DataStorage
 		// Concept for a conversion function used by the attribute binder to bind a lambda to a TFunction accepting column data arguments
 		template <typename FunctionType>
 		concept AttributeBinderColumnDataInvocable = std::is_invocable_v<std::decay_t<FunctionType>, const TWeakObjectPtr<const UScriptStruct>&, const void*>;
+		
+		// Concept for a conversion function used by the attribute binder to bind a lambda to a TFunction accepting a bool argument representing if
+		// the given column is present on the row
+        template <typename FunctionType>
+        concept AttributeBinderColumnPresenceInvocable = std::is_invocable_v<std::decay_t<FunctionType>, const bool>;
 
 		template <TColumnType ColumnType>
 		ColumnType* GetColumn(ICoreProvider* DataStorage, RowHandle Row)
