[TEDS] Fixed a bug when sorting by name.

Sorting by names by id incorrectly use the optional number that can be added to the name instead of it's internal id.

diff --git a/Engine/Source/Runtime/TypedElementFramework/Public/Elements/Framework/TypedElementSorter.inl b/Engine/Source/Runtime/TypedElementFramework/Public/Elements/Framework/TypedElementSorter.inl
--- a/Engine/Source/Runtime/TypedElementFramework/Public/Elements/Framework/TypedElementSorter.inl
+++ b/Engine/Source/Runtime/TypedElementFramework/Public/Elements/Framework/TypedElementSorter.inl
@@ -391,7 +391,7 @@ namespace UE::Editor::DataStorage
 		{
 			if constexpr (bIsById)
 			{
-				Cache = View->GetNumber();
+				Cache = View->GetComparisonIndex().ToUnstableInt();
 			}
 			else
 			{
