[TEDS] Solved compile error with returning a pointer and the new query context

Fixed an issue where functions that return a pointer would cause a string of hard to decipher compiler errors about missing braces.

diff --git a/Engine/Source/Runtime/TypedElementFramework/Public/DataStorage/CommonTypes.h b/Engine/Source/Runtime/TypedElementFramework/Public/DataStorage/CommonTypes.h
--- a/Engine/Source/Runtime/TypedElementFramework/Public/DataStorage/CommonTypes.h
+++ b/Engine/Source/Runtime/TypedElementFramework/Public/DataStorage/CommonTypes.h
@@ -42,6 +42,12 @@ namespace UE
 	template<typename Derived, typename Base>
 	concept derived_from = std::is_base_of_v<Base, Derived> && std::is_convertible_v<const volatile Derived*, const volatile Base*>;
 
+	template<class T>
+	struct type_identity { using type = T; };
+
+	template< class T >
+	using type_identity_t = type_identity<T>::type;
+
 	namespace Editor::DataStorage
 	{
 		using FColumn = FEditorDataStorageColumn;
diff --git a/Engine/Source/Runtime/TypedElementFramework/Public/Elements/Framework/TypedElementQueryContextImplementation.inl b/Engine/Source/Runtime/TypedElementFramework/Public/Elements/Framework/TypedElementQueryContextImplementation.inl
--- a/Engine/Source/Runtime/TypedElementFramework/Public/Elements/Framework/TypedElementQueryContextImplementation.inl
+++ b/Engine/Source/Runtime/TypedElementFramework/Public/Elements/Framework/TypedElementQueryContextImplementation.inl
@@ -1,6 +1,7 @@
 // Copyright Epic Games, Inc. All Rights Reserved.
 
 #include <type_traits>
+#include "DataStorage/CommonTypes.h"
 
 namespace UE::Editor::DataStorage::Queries
 {
@@ -60,8 +61,8 @@ namespace UE::Editor::DataStorage::Queries
 	} \
 	else \
 	{ \
-		checkf(false, TEXT("Function '" #FunctionName "' in capability '" #Capability"' is not supported by the current query context implementation.")); \
-		return ReturnValue {} ; \
+		checkf(false, TEXT("Function '" #FunctionName "' in capability '" #Capability "' is not supported by the current query context implementation.")); \
+		return UE::type_identity_t< ReturnValue >{}; \
 	}
 
 #define ArgTypeName(Type, Name) Type Name
diff --git a/Engine/Source/Runtime/TypedElementFramework/Public/Elements/Framework/TypedElementQueryContract.h b/Engine/Source/Runtime/TypedElementFramework/Public/Elements/Framework/TypedElementQueryContract.h
--- a/Engine/Source/Runtime/TypedElementFramework/Public/Elements/Framework/TypedElementQueryContract.h
+++ b/Engine/Source/Runtime/TypedElementFramework/Public/Elements/Framework/TypedElementQueryContract.h
@@ -103,7 +103,7 @@ namespace UE::Editor::DataStorage::Queries
 	 * Contract between context and implementation to be able to communicate.
 	 * The contract contains all functions for all capabilities that are available. Based on the supported capabilities
 	 * an implementation may opt to only partially implement the contract, with the remaining functions asserting. On 
-	 * the opposite side a context may restrict what funtions on the contract can be called based on the requested
+	 * the opposite side a context may restrict what functions on the contract can be called based on the requested
 	 * capabilities. Through the query function the capabilities on both sides are kept aligned, resulting in no
 	 * function on the contract being callable if they're not implemented.
 	 */
diff --git a/Engine/Source/Runtime/TypedElementFramework/Tests/Elements/Framework/TypedElementQueryContextMock.h b/Engine/Source/Runtime/TypedElementFramework/Tests/Elements/Framework/TypedElementQueryContextMock.h
--- a/Engine/Source/Runtime/TypedElementFramework/Tests/Elements/Framework/TypedElementQueryContextMock.h
+++ b/Engine/Source/Runtime/TypedElementFramework/Tests/Elements/Framework/TypedElementQueryContextMock.h
@@ -2,6 +2,7 @@
 
 #pragma once
 
+#include "DataStorage/CommonTypes.h"
 #include "Elements/Framework/TypedElementQueryCapabilities.h"
 #include "Elements/Framework/TypedElementQueryContextImplementation.h"
 #include "Elements/Framework/TypedElementQueryFunctionArgumentTypes.h"
@@ -132,7 +133,7 @@ namespace UE::Editor::DataStorage::Queries
 #undef WithWrappers
 
 #define MockFunctionImplementation(Capability, Return, Function) \
-		checkf(false, TEXT( "Context function " #Function " in capability " #Capability " was unexpectedly called." )); return Return{};
+		checkf(false, TEXT( "Context function " #Function " in capability " #Capability " was unexpectedly called." )); return UE::type_identity_t<Return>{};
 #define Function0(Capability, Return, Function) \
 		MockFunction(Capability, Function) = []() { MockFunctionImplementation(Capability, Return, Function) };
 #define Function1(Capability, Return, Function, Arg1) \
