[TEDS] Fixing compile error in TypedElementQueryStorageInterfaces.h

diff --git a/Engine/Source/Runtime/TypedElementFramework/Public/Elements/Interfaces/TypedElementQueryStorageInterfaces.h b/Engine/Source/Runtime/TypedElementFramework/Public/Elements/Interfaces/TypedElementQueryStorageInterfaces.h
--- a/Engine/Source/Runtime/TypedElementFramework/Public/Elements/Interfaces/TypedElementQueryStorageInterfaces.h
+++ b/Engine/Source/Runtime/TypedElementFramework/Public/Elements/Interfaces/TypedElementQueryStorageInterfaces.h
@@ -493,7 +493,7 @@ namespace UE::Editor::DataStorage
 	const DynamicColumnTemplate* ICommonQueryWithEnvironmentContext::GetColumn(const FName& Identifier) const
 	{
 		using namespace UE::Editor::DataStorage;
-		const UScriptStruct* StructInfo = FindDynamicColumnType(DynamicColumnTemplate::StaticStruct(), Identifier);
+		const UScriptStruct* StructInfo = FindDynamicColumnType(*DynamicColumnTemplate::StaticStruct(), Identifier);
 		if (StructInfo)
 		{
 			const void* ColumnData = GetColumn(StructInfo);
@@ -506,7 +506,7 @@ namespace UE::Editor::DataStorage
 	DynamicColumnTemplate* ICommonQueryWithEnvironmentContext::GetMutableColumn(const FName& Identifier)
 	{
 		using namespace UE::Editor::DataStorage;
-		const UScriptStruct* StructInfo = FindDynamicColumnType(DynamicColumnTemplate::StaticStruct(), Identifier);
+		const UScriptStruct* StructInfo = FindDynamicColumnType(*DynamicColumnTemplate::StaticStruct(), Identifier);
 		if (StructInfo)
 		{
 			const void* ColumnData = GetMutableColumn(StructInfo);
@@ -519,7 +519,7 @@ namespace UE::Editor::DataStorage
 	bool ICommonQueryWithEnvironmentContext::HasColumn(const FName& Identifier) const
 	{
 		using namespace UE::Editor::DataStorage;
-		const UScriptStruct* StructInfo = FindDynamicColumnType(DynamicColumnTemplate::StaticStruct(), Identifier);
+		const UScriptStruct* StructInfo = FindDynamicColumnType(*DynamicColumnTemplate::StaticStruct(), Identifier);
 		if (StructInfo)
 		{
 			return ICommonQueryContext::HasColumn(StructInfo);
@@ -531,7 +531,7 @@ namespace UE::Editor::DataStorage
 	bool ICommonQueryWithEnvironmentContext::HasColumn(RowHandle Row, const FName& Identifier) const
 	{
 		using namespace UE::Editor::DataStorage;
-		const UScriptStruct* StructInfo = FindDynamicColumnType(DynamicColumnTemplate::StaticStruct(), Identifier);
+		const UScriptStruct* StructInfo = FindDynamicColumnType(*DynamicColumnTemplate::StaticStruct(), Identifier);
 		if (StructInfo)
 		{
 			return ICommonQueryContext::HasColumn(Row, StructInfo);
