[TEDS Sorting] Have the default ConstructColumnSorters function remove duplicate sorter constructors

diff --git a/Engine/Plugins/Experimental/EditorDataStorage/Source/TedsCore/Private/TypedElementDatabaseUI.cpp b/Engine/Plugins/Experimental/EditorDataStorage/Source/TedsCore/Private/TypedElementDatabaseUI.cpp
--- a/Engine/Plugins/Experimental/EditorDataStorage/Source/TedsCore/Private/TypedElementDatabaseUI.cpp
+++ b/Engine/Plugins/Experimental/EditorDataStorage/Source/TedsCore/Private/TypedElementDatabaseUI.cpp
@@ -755,6 +755,8 @@ void UEditorDataStorageUi::GeneratePropertySorters(TArray<TSharedPtr<const UE::E
 {
 	using namespace UE::Editor::DataStorage;
 
+	TArray<const PropertySorterConstructorCallback*, TInlineAllocator<8>> UsedConstructors;
+
 	for (TWeakObjectPtr<const UScriptStruct>& ColumnTypePointer : Columns)
 	{
 		if (const UScriptStruct* ColumnType = ColumnTypePointer.Get())
@@ -766,7 +768,11 @@ void UEditorDataStorageUi::GeneratePropertySorters(TArray<TSharedPtr<const UE::E
 				{
 					if (const PropertySorterConstructorCallback* Callback = PropertySorterConstructors.Find(Property->GetClass()))
 					{
-						Results.Add((*Callback)(ColumnTypePointer, *Property));
+						if (!UsedConstructors.Contains(Callback))
+						{
+							UsedConstructors.Add(Callback);
+							Results.Add((*Callback)(ColumnTypePointer, *Property));
+						}
 					}
 					else
 					{
diff --git a/Engine/Source/Runtime/TypedElementFramework/Private/Elements/Interfaces/TypedElementDataStorageUiInterface.cpp b/Engine/Source/Runtime/TypedElementFramework/Private/Elements/Interfaces/TypedElementDataStorageUiInterface.cpp
--- a/Engine/Source/Runtime/TypedElementFramework/Private/Elements/Interfaces/TypedElementDataStorageUiInterface.cpp
+++ b/Engine/Source/Runtime/TypedElementFramework/Private/Elements/Interfaces/TypedElementDataStorageUiInterface.cpp
@@ -212,12 +212,17 @@ TArray<TSharedPtr<const UE::Editor::DataStorage::FColumnSorterInterface>> FTyped
 	using namespace UE::Editor::DataStorage;
 
 	TArray<TSharedPtr<const FColumnSorterInterface>> Results;
-
-	DataStorageUi->GeneratePropertySorters(Results, MatchedColumnTypes);
 	
 	TArray<TWeakObjectPtr<const UScriptStruct>, TInlineAllocator<32>> DynamicKeys;
 	MatchedDynamicTemplates.GetKeys(DynamicKeys);
-	DataStorageUi->GeneratePropertySorters(Results, DynamicKeys);
+
+	// Combine both column types and base dynamic column types so they can run in the same GeneratePropertySorter Function & check
+	// for any duplicate constructors.
+	TArray<TWeakObjectPtr<const UScriptStruct>> AllColumnTypes;
+	AllColumnTypes.Reserve(MatchedColumnTypes.Num() + DynamicKeys.Num());
+	AllColumnTypes.Append(MatchedColumnTypes);
+	AllColumnTypes.Append(DynamicKeys);
+	DataStorageUi->GeneratePropertySorters(Results, AllColumnTypes);
 
 	return Results;
 }
