[TypedElementFramework] Silence PVS warning.

diff --git a/Engine/Source/Runtime/TypedElementFramework/Public/Elements/Framework/TypedElementQueryBuilder.inl b/Engine/Source/Runtime/TypedElementFramework/Public/Elements/Framework/TypedElementQueryBuilder.inl
--- a/Engine/Source/Runtime/TypedElementFramework/Public/Elements/Framework/TypedElementQueryBuilder.inl
+++ b/Engine/Source/Runtime/TypedElementFramework/Public/Elements/Framework/TypedElementQueryBuilder.inl
@@ -514,7 +514,7 @@ namespace UE::Editor::DataStorage::Queries
 					{
 						std::apply([this, Rows = Rows.GetData(), &Caller] AUTORTFM_ENABLE (auto&&... Column)
 							{
-								if constexpr (AUTORTFM_MODE_OF_CALL(Caller(this->ContextWrapper, Rows, Column...)) == autortfm_mode_disable)
+								if constexpr (AUTORTFM_MODE_OF_CALL(Caller(this->ContextWrapper, Rows, Column...)) == autortfm_mode_disable) //-V501
 								{
 									AutoRTFM::UnreachableIfClosed();
 								}
@@ -534,7 +534,7 @@ namespace UE::Editor::DataStorage::Queries
 						{
 							std::apply([this, Row, &Caller] AUTORTFM_ENABLE (auto&&... Column)
 								{
-									if constexpr (AUTORTFM_MODE_OF_CALL(Caller(this->ContextWrapper, Row, *Column...)) == autortfm_mode_disable)
+									if constexpr (AUTORTFM_MODE_OF_CALL(Caller(this->ContextWrapper, Row, *Column...)) == autortfm_mode_disable) //-V501
 									{
 										AutoRTFM::UnreachableIfClosed();
 									}
