[TEDS] Utility class to collect row handles concurrently from multiple threads.

The collector allows rows to be collected from multiple threads at the same. Each thread first collects all its rows locally in an array and passes it to the collector which will merge the arrays from the threads together. When all threads are done processing the final result can be retrieved from the collector to be stored in the final destination. The collector supports appending results together in an array or can behave like a set where the rows are sorted and guaranteed to be unique.

#Synced-CL 47889446

diff --git a/Engine/Source/Runtime/TypedElementFramework/Private/Elements/Framework/TypedElementConcurrentRowHandleCollector.cpp b/Engine/Source/Runtime/TypedElementFramework/Private/Elements/Framework/TypedElementConcurrentRowHandleCollector.cpp
new file mode 100644
--- /dev/null
+++ b/Engine/Source/Runtime/TypedElementFramework/Private/Elements/Framework/TypedElementConcurrentRowHandleCollector.cpp
@@ -0,0 +1,113 @@
+// Copyright Epic Games, Inc. All Rights Reserved.
+
+#include "Elements/Framework/TypedElementConcurrentRowHandleCollector.h"
+
+namespace UE::Editor::DataStorage
+{
+	template<EConcurrentRowHandleCollectorType Type>
+	FConcurrentRowHandleCollector<Type>::~FConcurrentRowHandleCollector()
+	{
+		checkf(WorkingSet.load() == nullptr, TEXT("Attempting to destroy a FConcurrentRowHandleCollector while there are still values. "
+			"This can indicate that there are still threads actively collecting results."));
+	}
+
+	template<EConcurrentRowHandleCollectorType Type>
+	void FConcurrentRowHandleCollector<Type>::Append(FRowHandleArray&& LocalArray)
+	{
+		if (!LocalArray.IsEmpty())
+		{
+			if constexpr (Type == EConcurrentRowHandleCollectorType::Set)
+			{
+				TRACE_CPUPROFILER_EVENT_SCOPE_STR(TEXT("[TEDS] FConcurrentRowHandleCollector - Sort"));
+				LocalArray.Sort();
+			}
+
+			{
+				TRACE_CPUPROFILER_EVENT_SCOPE_STR(TEXT("[TEDS] FConcurrentRowHandleCollector - Merge"));
+
+				// Grab the main array and merge the local results into it, or create a new main array if it doesn't exist yet.
+				FRowHandleArray* WorkingSetPtr = WorkingSet.load();
+				while (!WorkingSet.compare_exchange_strong(WorkingSetPtr, nullptr));
+				if (WorkingSetPtr == nullptr)
+				{
+					WorkingSetPtr = new FRowHandleArray(MoveTemp(LocalArray));
+				}
+				else
+				{
+					TRACE_CPUPROFILER_EVENT_SCOPE_STR(TEXT("[TEDS] FConcurrentRowHandleCollector - Union"));
+					if constexpr (Type == EConcurrentRowHandleCollectorType::Array)
+					{
+						WorkingSetPtr->Append(LocalArray.GetRows());
+					}
+					else
+					{
+						WorkingSetPtr->SortedMerge(MoveTemp(LocalArray));
+					}
+				}
+
+				// Promote the local array to become the main array. If there already is a main array, try to claim that array
+				// as well and merge the local array into it to form the new main array and try to promote it again.
+				FRowHandleArray* NewSetPtr = nullptr;
+				while (!WorkingSet.compare_exchange_strong(NewSetPtr, WorkingSetPtr))
+				{
+					// There was an existing list so try to claim it and merge the local array, then try to set again.
+					if (WorkingSet.compare_exchange_strong(NewSetPtr, nullptr))
+					{
+						TRACE_CPUPROFILER_EVENT_SCOPE_STR(TEXT("[TEDS] FConcurrentRowHandleCollector - Union (retry)"));
+						if constexpr (Type == EConcurrentRowHandleCollectorType::Array)
+						{
+							NewSetPtr->Append(WorkingSetPtr->GetRows());
+						}
+						else
+						{
+							NewSetPtr->SortedMerge(MoveTemp(*WorkingSetPtr));
+						}
+						delete WorkingSetPtr;
+						WorkingSetPtr = NewSetPtr;
+					}
+					NewSetPtr = nullptr;
+				}
+			}
+		}
+	}
+
+	template<EConcurrentRowHandleCollectorType Type>
+	bool FConcurrentRowHandleCollector<Type>::Collect(FRowHandleArray& AccumulatedCollection)
+	{
+		bool Result = false;
+
+		// Get the current list. Threads can continue write and this will merge what has been collected up until that point.
+		FRowHandleArray* FinalList = WorkingSet.load();
+		while (!WorkingSet.compare_exchange_strong(FinalList, nullptr));
+
+		// Merge the results with the provided array.
+		if (FinalList)
+		{
+			if (!FinalList->IsEmpty())
+			{
+				TRACE_CPUPROFILER_EVENT_SCOPE_STR(TEXT("[TEDS] FConcurrentRowHandleCollector Collect"));
+
+				if constexpr (Type == EConcurrentRowHandleCollectorType::Array)
+				{
+					AccumulatedCollection.Append(FinalList->GetRows());
+				}
+				else
+				{
+					AccumulatedCollection.SortedMerge(MoveTemp(*FinalList));
+					AccumulatedCollection.MakeUnique();
+				}
+				Result = true;
+			}
+			delete FinalList;
+		}
+		return Result;
+	}
+
+	template FConcurrentRowHandleCollector<EConcurrentRowHandleCollectorType::Array>::~FConcurrentRowHandleCollector();
+	template void FConcurrentRowHandleCollector<EConcurrentRowHandleCollectorType::Array>::Append(FRowHandleArray&&);
+	template bool FConcurrentRowHandleCollector<EConcurrentRowHandleCollectorType::Array>::Collect(FRowHandleArray&);
+	
+	template FConcurrentRowHandleCollector<EConcurrentRowHandleCollectorType::Set>::~FConcurrentRowHandleCollector();
+	template void FConcurrentRowHandleCollector<EConcurrentRowHandleCollectorType::Set>::Append(FRowHandleArray&&);
+	template bool FConcurrentRowHandleCollector<EConcurrentRowHandleCollectorType::Set>::Collect(FRowHandleArray&);
+} // namespace UE::Editor::DataStorage
\ No newline at end of file
diff --git a/Engine/Source/Runtime/TypedElementFramework/Public/Elements/Framework/TypedElementConcurrentRowHandleCollector.h b/Engine/Source/Runtime/TypedElementFramework/Public/Elements/Framework/TypedElementConcurrentRowHandleCollector.h
new file mode 100644
--- /dev/null
+++ b/Engine/Source/Runtime/TypedElementFramework/Public/Elements/Framework/TypedElementConcurrentRowHandleCollector.h
@@ -0,0 +1,81 @@
+// Copyright Epic Games, Inc. All Rights Reserved.
+
+#pragma once
+
+#include <atomic>
+#include "Elements/Framework/TypedElementRowHandleArray.h"
+
+namespace UE::Editor::DataStorage
+{
+	enum class EConcurrentRowHandleCollectorType
+	{
+		Array, // Collects all row handles in random order. Is allowed to have duplicates.
+		Set // Collects only a unique instance of a row handles and orders them numerically.
+	};
+
+	/**
+	 * A lock-free utility class to collect rows from multiple threads at the same time.
+	 * Two options are provided:
+	 *		- Array: Collects all provided row handles without order and can have duplicates.
+	 *				This is usually faster than a set.
+	 *		- Set: Merges the provided row handles in sorted order and removes duplicates. This is
+	 *				generally slower than an array.
+	 * Each thread needs to first fill a local row handle array. That list is then passed on to the
+	 * collector that will combine it with the final list collected from multiple threads. The cost of
+	 * merging is spread across multiple threads, but as more list contribute some threads will have to
+	 * merge larger lists than other threads causing an imbalance in thread runtime. The worse case will
+	 * be the last thread as that will need to merge two arrays that are roughly both half the size of
+	 * the final array. It's therefor recommended to run with "AutoBalanceParallelChunkProcessing".
+	 */
+	template<EConcurrentRowHandleCollectorType Type>
+	class FConcurrentRowHandleCollector
+	{
+	public:
+		TYPEDELEMENTFRAMEWORK_API ~FConcurrentRowHandleCollector();
+
+		/**
+		 * Moves the provided array into the collected list by either appending (Array) or uniquely merging (Set).
+		 * This function can be repeatedly collect from multiple threads at the same time.
+		 * 
+		 * The approach for writing results back is a two step process. In the first step it tries to grab the 
+		 * temp local array that's shared by all threads. If that returns a nullptr it means there's no array yet or
+		 * another thread has already claimed it, in which case a new array is created. If it does manage to get a
+		 * hold of the array it merges its own local changes into it.
+		 * In the next phase it tries to place the locally update version of the array back so other threads can
+		 * add to it. If the shared array is a nullptr than it can move the local version into it, but if another
+		 * thread has already set it, it needs to try to grab that version and merge the local copy into it. This
+		 * local copy becomes the new target to place in the shared location.
+		 *
+		 * In practice there will be multiple merges collections present for the majority of processing and separate
+		 * threads will merge them together into a single final collection over time. This does result in later threads
+		 * having to merge increasingly large list, with one thread have to merge the final largest list. As a result this
+		 * leads to threads starting fast but taking increasingly long to run. It's therefore recommended to run queries
+		 * with `AutoBalanceParallelChunkProcessing`.
+		 */
+		TYPEDELEMENTFRAMEWORK_API void Append(FRowHandleArray&& LocalArray);
+		
+		/**
+		 * Collect the final results and merge with the provided list.
+		 * Returns true if additions were mode, otherwise false.
+		 * 
+		 * This function should be called once from a single thread to collect the final list. This can be called if threads 
+		 * are still adding new values to the collector, in which case this will remove the row handles merged into "Result" 
+		 * and will merge the rest when this is called again.
+		 */
+		TYPEDELEMENTFRAMEWORK_API bool Collect(FRowHandleArray& AccumulatedCollection);
+
+	private:
+		std::atomic<FRowHandleArray*> WorkingSet = nullptr;
+	};
+
+	using FConcurrentRowHandleArrayCollector = FConcurrentRowHandleCollector<EConcurrentRowHandleCollectorType::Array>;
+	using FConcurrentRowHandleSetCollector = FConcurrentRowHandleCollector<EConcurrentRowHandleCollectorType::Set>;
+	
+	extern template TYPEDELEMENTFRAMEWORK_API FConcurrentRowHandleCollector<EConcurrentRowHandleCollectorType::Array>::~FConcurrentRowHandleCollector();
+	extern template TYPEDELEMENTFRAMEWORK_API void FConcurrentRowHandleCollector<EConcurrentRowHandleCollectorType::Array>::Append(FRowHandleArray&&);
+	extern template TYPEDELEMENTFRAMEWORK_API bool FConcurrentRowHandleCollector<EConcurrentRowHandleCollectorType::Array>::Collect(FRowHandleArray&);
+
+	extern template TYPEDELEMENTFRAMEWORK_API FConcurrentRowHandleCollector<EConcurrentRowHandleCollectorType::Set>::~FConcurrentRowHandleCollector();
+	extern template TYPEDELEMENTFRAMEWORK_API void FConcurrentRowHandleCollector<EConcurrentRowHandleCollectorType::Set>::Append(FRowHandleArray&&);
+	extern template TYPEDELEMENTFRAMEWORK_API bool FConcurrentRowHandleCollector<EConcurrentRowHandleCollectorType::Set>::Collect(FRowHandleArray&);
+} // namespace UE::Editor::DataStorage
\ No newline at end of file
