[TEDS] Option to calculate a hash of the topology of a table.

New functionality was added to the Editor Data Storage that can calculate a hash for the topology of a table and its variants. The hash will change whenever a row is added, removed or internal reordering happens. This can be used to quickly check if a table needs to be evaluated for changes to the row count, reducing the need for observers to keep track of this instead.

#Synced-CL 48364666

diff --git a/Engine/Plugins/Experimental/EditorDataStorage/Source/TedsCore/Private/Queries/TypedElementExtendedQueryStore.cpp b/Engine/Plugins/Experimental/EditorDataStorage/Source/TedsCore/Private/Queries/TypedElementExtendedQueryStore.cpp
--- a/Engine/Plugins/Experimental/EditorDataStorage/Source/TedsCore/Private/Queries/TypedElementExtendedQueryStore.cpp
+++ b/Engine/Plugins/Experimental/EditorDataStorage/Source/TedsCore/Private/Queries/TypedElementExtendedQueryStore.cpp
@@ -274,6 +274,17 @@ namespace UE::Editor::DataStorage
 		return QueryData ? QueryData->Description : EmptyDescription;
 	}
 
+	uint64 FExtendedQueryStore::CalculateQueryTablesTopologyHash(Handle Query) const
+	{
+		const FExtendedQuery* QueryData = Get(Query);
+		// Const cast used because CalculateArchetypeEntityOrderHash is not const as it may need to redetermine what
+		// archetypes match with the native query. This is however conceptually const from TEDS' point of view.
+		return
+			QueryData
+			? const_cast<FExtendedQuery*>(QueryData)->NativeQuery.CalculateArchetypeEntityOrderHash()
+			: 0;
+	}
+
 	bool FExtendedQueryStore::IsAlive(Handle Entry) const
 	{
 		return Queries.IsAlive(Entry);
diff --git a/Engine/Plugins/Experimental/EditorDataStorage/Source/TedsCore/Private/Queries/TypedElementExtendedQueryStore.h b/Engine/Plugins/Experimental/EditorDataStorage/Source/TedsCore/Private/Queries/TypedElementExtendedQueryStore.h
--- a/Engine/Plugins/Experimental/EditorDataStorage/Source/TedsCore/Private/Queries/TypedElementExtendedQueryStore.h
+++ b/Engine/Plugins/Experimental/EditorDataStorage/Source/TedsCore/Private/Queries/TypedElementExtendedQueryStore.h
@@ -98,6 +98,12 @@ namespace UE::Editor::DataStorage
 		/** Gets the original description used to create an extended query or an empty default if the provided query isn't alive. */
 		const FQueryDescription& GetQueryDescription(Handle Query) const;
 
+		/**
+		 * Calculates a hash based on the topology revision id of the tables associated with the provided query. If this value changes
+		 * it indicates that one or more tables has new rows added, had rows removed and/or rows have been reordered.
+		 */
+		uint64 CalculateQueryTablesTopologyHash(Handle Query) const;
+
 		/** Checks to see if a query is still available or has been removed. */
 		bool IsAlive(Handle Entry) const;
 
diff --git a/Engine/Plugins/Experimental/EditorDataStorage/Source/TedsCore/Private/TypedElementDatabase.cpp b/Engine/Plugins/Experimental/EditorDataStorage/Source/TedsCore/Private/TypedElementDatabase.cpp
--- a/Engine/Plugins/Experimental/EditorDataStorage/Source/TedsCore/Private/TypedElementDatabase.cpp
+++ b/Engine/Plugins/Experimental/EditorDataStorage/Source/TedsCore/Private/TypedElementDatabase.cpp
@@ -1233,6 +1233,12 @@ const UE::Editor::DataStorage::FQueryDescription& UEditorDataStorage::GetQueryDe
 	return Environment->GetQueryStore().GetQueryDescription(StorageHandle);
 }
 
+uint64 UEditorDataStorage::CalculateQueryTablesTopologyHash(UE::Editor::DataStorage::QueryHandle Query) const
+{
+	const UE::Editor::DataStorage::FExtendedQueryStore::Handle StorageHandle(Query);
+	return Environment->GetQueryStore().CalculateQueryTablesTopologyHash(StorageHandle);
+}
+
 FName UEditorDataStorage::GetQueryTickGroupName(UE::Editor::DataStorage::EQueryTickGroups Group) const
 {
 	using namespace UE::Editor::DataStorage;
diff --git a/Engine/Plugins/Experimental/EditorDataStorage/Source/TedsCore/Private/TypedElementDatabase.h b/Engine/Plugins/Experimental/EditorDataStorage/Source/TedsCore/Private/TypedElementDatabase.h
--- a/Engine/Plugins/Experimental/EditorDataStorage/Source/TedsCore/Private/TypedElementDatabase.h
+++ b/Engine/Plugins/Experimental/EditorDataStorage/Source/TedsCore/Private/TypedElementDatabase.h
@@ -177,6 +177,7 @@ public:
 	UE_API UE::Editor::DataStorage::QueryHandle RegisterQuery(UE::Editor::DataStorage::FQueryDescription&& Query) override;
 	UE_API virtual void UnregisterQuery(UE::Editor::DataStorage::QueryHandle Query) override;
 	UE_API virtual const UE::Editor::DataStorage::FQueryDescription& GetQueryDescription(UE::Editor::DataStorage::QueryHandle Query) const override;
+	UE_API virtual uint64 CalculateQueryTablesTopologyHash(UE::Editor::DataStorage::QueryHandle Query) const override;
 	UE_API virtual FName GetQueryTickGroupName(UE::Editor::DataStorage::EQueryTickGroups Group) const override;
 	UE_API virtual UE::Editor::DataStorage::FQueryResult RunQuery(UE::Editor::DataStorage::QueryHandle Query) override;
 	UE_API virtual UE::Editor::DataStorage::FQueryResult RunQuery(UE::Editor::DataStorage::QueryHandle Query, UE::Editor::DataStorage::DirectQueryCallbackRef Callback) override;
diff --git a/Engine/Source/Runtime/TypedElementFramework/Public/Elements/Interfaces/TypedElementDataStorageInterface.h b/Engine/Source/Runtime/TypedElementFramework/Public/Elements/Interfaces/TypedElementDataStorageInterface.h
--- a/Engine/Source/Runtime/TypedElementFramework/Public/Elements/Interfaces/TypedElementDataStorageInterface.h
+++ b/Engine/Source/Runtime/TypedElementFramework/Public/Elements/Interfaces/TypedElementDataStorageInterface.h
@@ -456,6 +456,11 @@ public:
 	virtual void UnregisterQuery(QueryHandle Query) = 0;
 	/** Returns the description of a previously registered query. If the query no longer exists an empty description will be returned. */
 	virtual const FQueryDescription& GetQueryDescription(QueryHandle Query) const = 0;
+	/**
+	 * Calculates a hash based on the topology revision id of the tables associated with the provided query. If this value changes
+	 * it indicates that one or more tables has new rows added, had rows removed and/or rows have been reordered.
+	 */
+	virtual uint64 CalculateQueryTablesTopologyHash(QueryHandle Query) const = 0;
 	/**
 	 * Tick groups for queries can be given any name and the Data Storage will figure out the order of execution based on found
 	 * dependencies. However keeping processors within the same query group can help promote better performance through parallelization.
@@ -491,7 +496,6 @@ public:
 	 */
 	template<typename ResultType, Queries::FunctionType Function>
 	FQueryResult RunQuery(QueryHandle Query, ERunQueryFlags Flags, Queries::TResult<ResultType>& Result, Function&& Callback);
-	
 	/**
 	 * Triggers all queries registered under the activation name to run for one update cycle. The activatable queries will be activated at
 	 * start of the cycle and disabled at the end of the cycle and act like regular queries for that cycle. This includes not running
