[TEDS] Added ability to retrieve raw addressed for batches in the new query context.

#Synced-CL 48813241

diff --git a/Engine/Plugins/Experimental/EditorDataStorage/Source/TedsCore/Private/Queries/TypedElementQueryContext_DirectSingleApi.cpp b/Engine/Plugins/Experimental/EditorDataStorage/Source/TedsCore/Private/Queries/TypedElementQueryContext_DirectSingleApi.cpp
--- a/Engine/Plugins/Experimental/EditorDataStorage/Source/TedsCore/Private/Queries/TypedElementQueryContext_DirectSingleApi.cpp
+++ b/Engine/Plugins/Experimental/EditorDataStorage/Source/TedsCore/Private/Queries/TypedElementQueryContext_DirectSingleApi.cpp
@@ -72,6 +72,16 @@ namespace UE::Editor::DataStorage::Queries
 		return CurrentRowHasColumns(Columns);
 	}
 
+	const void* FQueryContext_DirectSingleApi::GetColumnBatchAddress(const UScriptStruct* ColumnType) const
+	{
+		return DataStorage.GetColumnData(CurrentRow, ColumnType);
+	}
+
+	void* FQueryContext_DirectSingleApi::GetMutableColumnBatchAddress(const UScriptStruct* ColumnType)
+	{
+		return DataStorage.GetColumnData(CurrentRow, ColumnType);
+	}
+
 	FEnvironment& FQueryContext_DirectSingleApi::GetEnvironment()
 	{
 		return Environment;
diff --git a/Engine/Plugins/Experimental/EditorDataStorage/Source/TedsCore/Private/Queries/TypedElementQueryContext_DirectSingleApi.h b/Engine/Plugins/Experimental/EditorDataStorage/Source/TedsCore/Private/Queries/TypedElementQueryContext_DirectSingleApi.h
--- a/Engine/Plugins/Experimental/EditorDataStorage/Source/TedsCore/Private/Queries/TypedElementQueryContext_DirectSingleApi.h
+++ b/Engine/Plugins/Experimental/EditorDataStorage/Source/TedsCore/Private/Queries/TypedElementQueryContext_DirectSingleApi.h
@@ -41,6 +41,8 @@ namespace UE::Editor::DataStorage::Queries
 		virtual uint32 GetBatchRowCount() const override;
 		virtual TConstArrayView<RowHandle> GetBatchRowHandles() const override;
 		virtual bool CurrentBatchTableHasColumns(TConstArrayView<const UScriptStruct*> Columns) const override;
+		virtual const void* GetColumnBatchAddress(const UScriptStruct* ColumnType) const override;
+		virtual void* GetMutableColumnBatchAddress(const UScriptStruct* ColumnType) override;
 
 		/**
 		 * @section IQueryFunctionResponse
diff --git a/Engine/Plugins/Experimental/EditorDataStorage/Source/TedsCore/Private/Queries/TypedElementQueryContext_MassForwarder.cpp b/Engine/Plugins/Experimental/EditorDataStorage/Source/TedsCore/Private/Queries/TypedElementQueryContext_MassForwarder.cpp
--- a/Engine/Plugins/Experimental/EditorDataStorage/Source/TedsCore/Private/Queries/TypedElementQueryContext_MassForwarder.cpp
+++ b/Engine/Plugins/Experimental/EditorDataStorage/Source/TedsCore/Private/Queries/TypedElementQueryContext_MassForwarder.cpp
@@ -65,6 +65,20 @@ namespace UE::Editor::DataStorage::Queries
 		return CurrentRowHasColumns(ColumnTypes);
 	}
 
+	const void* FQueryContext_MassForwarder::GetColumnBatchAddress(const UScriptStruct* ColumnType) const
+	{
+		checkf(UE::Mass::IsA<FMassFragment>(ColumnType), TEXT("Trying to retrieve a data pointer for column '%s' which doesn't have data."),
+			ColumnType ? *ColumnType->GetName() : TEXT("<null>"));
+		return Context.GetFragmentView(ColumnType).GetData();
+	}
+
+	void* FQueryContext_MassForwarder::GetMutableColumnBatchAddress(const UScriptStruct* ColumnType)
+	{
+		checkf(UE::Mass::IsA<FMassFragment>(ColumnType), TEXT("Trying to retrieve a data pointer for column '%s' which doesn't have data."),
+			ColumnType ? *ColumnType->GetName() : TEXT("<null>"));
+		return Context.GetMutableFragmentView(ColumnType).GetData();
+	}
+
 	bool FQueryContext_MassForwarder::NextBatch()
 	{
 		return false; // This only ever contains a single batch.
diff --git a/Engine/Plugins/Experimental/EditorDataStorage/Source/TedsCore/Private/Queries/TypedElementQueryContext_MassForwarder.h b/Engine/Plugins/Experimental/EditorDataStorage/Source/TedsCore/Private/Queries/TypedElementQueryContext_MassForwarder.h
--- a/Engine/Plugins/Experimental/EditorDataStorage/Source/TedsCore/Private/Queries/TypedElementQueryContext_MassForwarder.h
+++ b/Engine/Plugins/Experimental/EditorDataStorage/Source/TedsCore/Private/Queries/TypedElementQueryContext_MassForwarder.h
@@ -36,6 +36,8 @@ namespace UE::Editor::DataStorage::Queries
 		virtual uint32 GetBatchRowCount() const override;
 		virtual TConstArrayView<RowHandle> GetBatchRowHandles() const override;
 		virtual bool CurrentBatchTableHasColumns(TConstArrayView<const UScriptStruct*> ColumnTypes) const override;
+		virtual const void* GetColumnBatchAddress(const UScriptStruct* ColumnType) const override;
+		virtual void* GetMutableColumnBatchAddress(const UScriptStruct* ColumnType) override;
 
 		/**
 		 * @section IQueryFunctionResponse
diff --git a/Engine/Source/Runtime/TypedElementFramework/Public/Elements/Framework/TypedElementQueryCapabilities.inl b/Engine/Source/Runtime/TypedElementFramework/Public/Elements/Framework/TypedElementQueryCapabilities.inl
--- a/Engine/Source/Runtime/TypedElementFramework/Public/Elements/Framework/TypedElementQueryCapabilities.inl
+++ b/Engine/Source/Runtime/TypedElementFramework/Public/Elements/Framework/TypedElementQueryCapabilities.inl
@@ -40,6 +40,10 @@ CapabilityStart(RowBatchInfo, EContextCapabilityFlags::SupportsBatch)
 	ConstFunction0(RowBatchInfo, TConstArrayView<RowHandle>, GetBatchRowHandles)
 	/** Checks if the rows in the current batch have the requested columns. */
 	ConstFunction1(RowBatchInfo, bool, CurrentBatchTableHasColumns, (TConstArrayView<const UScriptStruct*>, Columns))
+	/** Direct access to the const addresses of a column or null if not found in this query. */
+	ConstFunction1(RowBatchInfo, const void*, GetColumnBatchAddress, (const UScriptStruct*, ColumnType))
+	/** Direct access to the mutable addresses of a column or null if not found in this query. */
+	Function1(RowBatchInfo, void*, GetMutableColumnBatchAddress, (const UScriptStruct*, ColumnType))
 
 #if defined(WithWrappers)
 	/** Checks if the rows in the current batch have the requested columns. */
@@ -49,12 +53,36 @@ CapabilityStart(RowBatchInfo, EContextCapabilityFlags::SupportsBatch)
 		return this->CurrentBatchTableHasColumns(TConstArrayView<const UScriptStruct*>({ ColumnTypes::StaticStruct()... }));
 	}
 
+	/**
+	 * Utility function to iterate over each individual cell in the rows in this batch for the requested columns.
+	 * Example: ForEachRow([](const ColumnA&, const ColumnB&), ArgA, ArgB); 
+	 *		where ArgA is a TBatch<ColumnA> or TConstBatch<ColumnA> (a.k.a. TBatch<const ColumnA>) and 
+	 *		ArgB is a TBatch<ColumnB> or TConstBatch<ColumnB> (a.k.a. TBatch<const ColumnB>) argument in the query callback function.
+	 */
+	template<typename CallbackType, typename... Args>
+	void ForEachRow(CallbackType&& Callback, TBatch<Args>... Batches) const
+	{
+		TTuple<Args*...> AddressTuple = MakeTuple(Batches.GetData()...);
+		for (RowHandle Row : GetBatchRowHandles())
+		{
+			AddressTuple.ApplyBefore([&](std::add_const_t<Args>*&... Pointers)
+				{
+					Callback((*Pointers++)...);
+				});
+		}
+	}
+
+	/**
+	 * Utility function to iterate over each individual cell in the rows in this batch for the requested columns.
+	 * Example: ForEachRow([](ColumnA&, const ColumnB&), ArgA, ArgB);
+	 *		where ArgA is a TBatch<ColumnA> and ArgB is a TConstBatch<ColumnB> (a.k.a. TBatch<const ColumnB>) 
+	 *		or TBatch<ColumnB> argument in the query callback function.
+	 */
 	template<typename CallbackType, typename... Args>
 	void ForEachRow(CallbackType&& Callback, TBatch<Args>... Batches)
 	{
-		uint32 RowCount = this->GetBatchRowCount();
 		TTuple<Args*...> AddressTuple = MakeTuple(Batches.GetData()...);
-		for (uint32 Index = 0; Index < RowCount; ++Index)
+		for (RowHandle Row : GetBatchRowHandles())
 		{
 			AddressTuple.ApplyBefore([&](Args*&... Pointers)
 				{
