[TEDS] Fixed failing query test.

#rnx
#Synced-CL 49045677

diff --git a/Engine/Source/Runtime/TypedElementFramework/Tests/Elements/Framework/TypedElementQueryFunctionsTests.cpp b/Engine/Source/Runtime/TypedElementFramework/Tests/Elements/Framework/TypedElementQueryFunctionsTests.cpp
--- a/Engine/Source/Runtime/TypedElementFramework/Tests/Elements/Framework/TypedElementQueryFunctionsTests.cpp
+++ b/Engine/Source/Runtime/TypedElementFramework/Tests/Elements/Framework/TypedElementQueryFunctionsTests.cpp
@@ -2,6 +2,7 @@
 
 #if WITH_TESTS
 
+#include "Containers/StaticArray.h"
 #include "Elements/Framework/TypedElementQueryContext.h"
 #include "Elements/Framework/TypedElementQueryContextMock.h"
 #include "Elements/Framework/TypedElementQueryFunctions.h"
@@ -297,9 +298,12 @@ namespace UE::Editor::DataStorage::Queries::Tests
 				});
 
 			static constexpr uint32 RowCount = 2;
+			TStaticArray<RowHandle, 2> Rows = { 0, 1 };
+			TConstArrayView<RowHandle> RowsView(Rows.GetData(), Rows.Num());
 
 			QueryContextMock Context;
-			Context.GetContextImplementation().RowBatchInfo_GetBatchRowCount_Mock = []() { return RowCount; };
+			Context.GetContextImplementation().RowBatchInfo_GetBatchRowCount_Mock = [](){ return RowCount; };
+			Context.GetContextImplementation().RowBatchInfo_GetBatchRowHandles_Mock = [RowsView](){ return RowsView; };
 			FTestColumnString ColumnsA[RowCount];
 			FTestColumnInt ColumnsB[] =
 			{
@@ -437,9 +441,12 @@ namespace UE::Editor::DataStorage::Queries::Tests
 				});
 
 			static constexpr int32 RowCount = 4;
+			TStaticArray<RowHandle, 4> Rows = { 0, 1, 2, 3 };
+			TConstArrayView<RowHandle> RowsView(Rows.GetData(), Rows.Num());
 
 			QueryContextMock Context;
 			Context.GetContextImplementation().RowBatchInfo_GetBatchRowCount_Mock = []() { return RowCount; };
+			Context.GetContextImplementation().RowBatchInfo_GetBatchRowHandles_Mock = [RowsView](){ return RowsView; };
 
 			FTestColumnInt Columns[] =
 			{
