Add FHighlightTextColumn

#Synced-CL 49621724

diff --git a/Engine/Source/Runtime/TypedElementFramework/Public/Elements/Columns/TypedElementSlateWidgetColumns.h b/Engine/Source/Runtime/TypedElementFramework/Public/Elements/Columns/TypedElementSlateWidgetColumns.h
--- a/Engine/Source/Runtime/TypedElementFramework/Public/Elements/Columns/TypedElementSlateWidgetColumns.h
+++ b/Engine/Source/Runtime/TypedElementFramework/Public/Elements/Columns/TypedElementSlateWidgetColumns.h
@@ -72,6 +72,18 @@ struct FDescriptionColumn final : public FEditorDataStorageColumn
 	FText Description;
 };
 
+/**
+ * A highlight text for this row.  Can be used by widget rows to set the corresponding HighlightText on a widget (e.g. STextBlock)
+ */
+USTRUCT(meta = (DisplayName = "Highlight Text"))
+struct FHighlightTextColumn final : public FEditorDataStorageColumn
+{
+	GENERATED_BODY()
+
+	UPROPERTY()
+	FText HighlightText;
+};
+
 /**
  * A color for this row. Can be used by widget rows to determine widget color or non-widget rows to attach a logical color to their data
  * (e.g asset colors)
