[Feature] Adding PostRegister for TEDS Factories
[Added] PostRegister virtual function to let factories know the registering phase is finished (A global delegate might be a bit more performant here, but the virtual function is more clear and might be more useful)

#Synced-CL 49621724

diff --git a/Engine/Plugins/Experimental/EditorDataStorage/Source/TedsCore/Private/TypedElementsDataStorage.cpp b/Engine/Plugins/Experimental/EditorDataStorage/Source/TedsCore/Private/TypedElementsDataStorage.cpp
--- a/Engine/Plugins/Experimental/EditorDataStorage/Source/TedsCore/Private/TypedElementsDataStorage.cpp
+++ b/Engine/Plugins/Experimental/EditorDataStorage/Source/TedsCore/Private/TypedElementsDataStorage.cpp
@@ -145,7 +145,15 @@ void FEditorDataStorageModule::StartupModule()
 					Factory->RegisterQueries(*DataStorage);
 					Factory->RegisterWidgetConstructors(*DataStorage, *DataStorageUi);
 				}
-				
+
+				// Third pass to call post register.
+				for (UEditorDataStorage::FactoryIterator Iterator = DataStorage->CreateFactoryIterator(); Iterator; ++Iterator)
+				{
+					UEditorDataStorageFactory* Factory = *Iterator;
+
+					Factory->PostRegister(*DataStorage);
+				}
+
 				UE_LOG(LogEditorDataStorage, Log, TEXT("Initialized"));
 
 				bInitialized = true;
diff --git a/Engine/Source/Runtime/TypedElementFramework/Public/Elements/Interfaces/TypedElementDataStorageFactory.h b/Engine/Source/Runtime/TypedElementFramework/Public/Elements/Interfaces/TypedElementDataStorageFactory.h
--- a/Engine/Source/Runtime/TypedElementFramework/Public/Elements/Interfaces/TypedElementDataStorageFactory.h
+++ b/Engine/Source/Runtime/TypedElementFramework/Public/Elements/Interfaces/TypedElementDataStorageFactory.h
@@ -34,6 +34,11 @@ public:
 	 */
 	virtual void PreRegister(UE::Editor::DataStorage::ICoreProvider& DataStorage) {}
 
+	/**
+	 * All factories will have this called after every Register functions on all factories are called
+	 */
+	virtual void PostRegister(UE::Editor::DataStorage::ICoreProvider& DataStorage) {}
+
 	virtual void RegisterTables(UE::Editor::DataStorage::ICoreProvider& DataStorage) {}
 	virtual void RegisterTables(UE::Editor::DataStorage::ICoreProvider& DataStorage, UE::Editor::DataStorage::ICompatibilityProvider& DataStorageCompatibility) {}
 	virtual void RegisterTickGroups(UE::Editor::DataStorage::ICoreProvider& DataStorage) const {}
