Teds Editor Visibility Column: Make the visibility member in the column a UProperty (mainly to help with debugging/UI display)

#Synced-CL 49860896

diff --git a/Engine/Source/Runtime/TypedElementFramework/Public/Elements/Columns/TypedElementVisibilityColumns.h b/Engine/Source/Runtime/TypedElementFramework/Public/Elements/Columns/TypedElementVisibilityColumns.h
--- a/Engine/Source/Runtime/TypedElementFramework/Public/Elements/Columns/TypedElementVisibilityColumns.h
+++ b/Engine/Source/Runtime/TypedElementFramework/Public/Elements/Columns/TypedElementVisibilityColumns.h
@@ -15,6 +15,6 @@ struct FVisibleInEditorColumn final : public FEditorDataStorageColumn
 	GENERATED_BODY()
 
 public:
-
+	UPROPERTY()
 	bool bIsVisibleInEditor = true;
 };
\ No newline at end of file
