[UE-355612] Deprecate UI capabilities from TEDS UI

[Deprecate] ITypedElementUiTooltipCapability
[Deprecate] ITypedElementUiCapability

#Synced-CL 50045910

diff --git a/Engine/Source/Runtime/TypedElementFramework/Public/Elements/Interfaces/Capabilities/TypedElementUiStyleOverrideCapability.h b/Engine/Source/Runtime/TypedElementFramework/Public/Elements/Interfaces/Capabilities/TypedElementUiStyleOverrideCapability.h
--- a/Engine/Source/Runtime/TypedElementFramework/Public/Elements/Interfaces/Capabilities/TypedElementUiStyleOverrideCapability.h
+++ b/Engine/Source/Runtime/TypedElementFramework/Public/Elements/Interfaces/Capabilities/TypedElementUiStyleOverrideCapability.h
@@ -7,6 +7,7 @@
 
 struct FSlateColor;
 
+PRAGMA_DISABLE_DEPRECATION_WARNINGS
 /**
  * Interface to provide access to widgets that support working with style overrides.
  */
@@ -34,4 +35,5 @@ public:
 
 private:
 	WidgetType& Widget;
-};
\ No newline at end of file
+};
+PRAGMA_ENABLE_DEPRECATION_WARNINGS
\ No newline at end of file
diff --git a/Engine/Source/Runtime/TypedElementFramework/Public/Elements/Interfaces/Capabilities/TypedElementUiTextCapability.h b/Engine/Source/Runtime/TypedElementFramework/Public/Elements/Interfaces/Capabilities/TypedElementUiTextCapability.h
--- a/Engine/Source/Runtime/TypedElementFramework/Public/Elements/Interfaces/Capabilities/TypedElementUiTextCapability.h
+++ b/Engine/Source/Runtime/TypedElementFramework/Public/Elements/Interfaces/Capabilities/TypedElementUiTextCapability.h
@@ -6,6 +6,7 @@
 #include "Internationalization/Text.h"
 #include "Misc/Attribute.h"
 
+PRAGMA_DISABLE_DEPRECATION_WARNINGS
 /**
  * Interface to provide access to widgets that support working with text.
  */
@@ -39,4 +40,5 @@ public:
 
 private:
 	WidgetType& Widget;
-};
\ No newline at end of file
+};
+PRAGMA_ENABLE_DEPRECATION_WARNINGS
\ No newline at end of file
diff --git a/Engine/Source/Runtime/TypedElementFramework/Public/Elements/Interfaces/Capabilities/TypedElementUiTooltipCapability.h b/Engine/Source/Runtime/TypedElementFramework/Public/Elements/Interfaces/Capabilities/TypedElementUiTooltipCapability.h
--- a/Engine/Source/Runtime/TypedElementFramework/Public/Elements/Interfaces/Capabilities/TypedElementUiTooltipCapability.h
+++ b/Engine/Source/Runtime/TypedElementFramework/Public/Elements/Interfaces/Capabilities/TypedElementUiTooltipCapability.h
@@ -10,7 +10,8 @@
 /**
  * Interface to provide access to tool tips on widgets.
  */
-class ITypedElementUiTooltipCapability : public ITypedElementUiCapability
+PRAGMA_DISABLE_DEPRECATION_WARNINGS
+class UE_DEPRECATED(5.8, "Please use FTypedElementLabelColumn and FAttributeBinder to bind the tooltip attribute to the widget instead") ITypedElementUiTooltipCapability : public ITypedElementUiCapability
 {
 public:
 	SLATE_METADATA_TYPE(ITypedElementUiTooltipCapability, ITypedElementUiCapability)
@@ -27,7 +28,7 @@ public:
 };
 
 template<typename WidgetType>
-class TTypedElementUiTooltipCapability : public ITypedElementUiTooltipCapability
+class UE_DEPRECATED(5.8, "Please use FTypedElementLabelColumn and FAttributeBinder to bind the tooltip attribute to the widget instead") TTypedElementUiTooltipCapability : public ITypedElementUiTooltipCapability
 {
 public:
 	explicit TTypedElementUiTooltipCapability(WidgetType& InWidget) : Widget(InWidget){}
@@ -64,4 +65,5 @@ public:
 
 private:
 	WidgetType& Widget;
-};
\ No newline at end of file
+};
+PRAGMA_ENABLE_DEPRECATION_WARNINGS
\ No newline at end of file
diff --git a/Engine/Source/Runtime/TypedElementFramework/Public/Elements/Interfaces/TypedElementUiCapabilities.h b/Engine/Source/Runtime/TypedElementFramework/Public/Elements/Interfaces/TypedElementUiCapabilities.h
--- a/Engine/Source/Runtime/TypedElementFramework/Public/Elements/Interfaces/TypedElementUiCapabilities.h
+++ b/Engine/Source/Runtime/TypedElementFramework/Public/Elements/Interfaces/TypedElementUiCapabilities.h
@@ -9,7 +9,7 @@
  * describe the functionality offered by the widget. This allows arbitrary widgets to describe themselves
  * and allows users of TEDS UI to call into or register with the widget.
  */
-class ITypedElementUiCapability : public ISlateMetaData
+class UE_DEPRECATED(5.8, "Please store metadata in columns instead of using UI capabilities") ITypedElementUiCapability : public ISlateMetaData
 {
 public:
 	SLATE_METADATA_TYPE(ITypedElementUiCapability, ISlateMetaData)
