[TEDS] Added the ability to put the compatibility layer in immediate mode.

#rnx

#Synced-CL 49834081

diff --git a/Engine/Plugins/Experimental/EditorDataStorage/Source/TedsCore/Private/Commands/EditorDataStorageCompatibilityCommands.cpp b/Engine/Plugins/Experimental/EditorDataStorage/Source/TedsCore/Private/Commands/EditorDataStorageCompatibilityCommands.cpp
--- a/Engine/Plugins/Experimental/EditorDataStorage/Source/TedsCore/Private/Commands/EditorDataStorageCompatibilityCommands.cpp
+++ b/Engine/Plugins/Experimental/EditorDataStorage/Source/TedsCore/Private/Commands/EditorDataStorageCompatibilityCommands.cpp
@@ -493,7 +493,8 @@ namespace UE::Editor::DataStorage
 		, MementoSystem(InStorageCompatibility.Environment->GetMementoSystem())
 	{}
 
-	void FCommandProcessor::SetupRow(RowHandle Row, UObject* Object)
+	void FCommandProcessor::SetupRow(ICoreProvider& Storage, UEditorDataStorageCompatibility& StorageCompatibility, 
+		RowHandle Row, UObject* Object)
 	{
 		Storage.AddColumn(Row, FTypedElementUObjectColumn{ .Object = Object });
 		Storage.AddColumn(Row, FTypedElementUObjectIdColumn
@@ -511,7 +512,8 @@ namespace UE::Editor::DataStorage
 		StorageCompatibility.TriggerOnObjectAdded(Object, Object->GetClass(), Row);
 	}
 
-	void FCommandProcessor::SetupRow(RowHandle Row, void* Object, TWeakObjectPtr<const UScriptStruct> TypeInfo)
+	void FCommandProcessor::SetupRow(ICoreProvider& Storage, UEditorDataStorageCompatibility& StorageCompatibility, 
+		RowHandle Row, void* Object, TWeakObjectPtr<const UScriptStruct> TypeInfo)
 	{
 		Storage.AddColumn(Row, FTypedElementExternalObjectColumn{ .Object = Object });
 		Storage.AddColumn(Row, FTypedElementScriptStructTypeInfoColumn{ .TypeInfo = TypeInfo });
@@ -568,7 +570,7 @@ namespace UE::Editor::DataStorage
 		Storage.AddRow(Object.Row, Object.Table, 
 			[ObjectPtr, this](RowHandle Row)
 			{
-				SetupRow(Row, ObjectPtr);
+				SetupRow(Storage, StorageCompatibility, Row, ObjectPtr);
 			});
 	}
 
@@ -584,7 +586,7 @@ namespace UE::Editor::DataStorage
 				UObject* ObjectPtr = ObjectIterator->Get();
 				checkf(ObjectPtr, TEXT(
 					"Expected a valid object pointer. If there isn't one here then the filter pass did not correctly clean up this command."));
-				SetupRow(Row, ObjectPtr);
+				SetupRow(Storage, StorageCompatibility, Row, ObjectPtr);
 				ObjectIterator++;
 			});
 	}
@@ -593,7 +595,7 @@ namespace UE::Editor::DataStorage
 	{
 		Storage.AddRow(Object.Row, Object.Table, [&Object, this](RowHandle Row)
 			{
-				SetupRow(Object.Row, Object.Object, Object.TypeInfo);
+				SetupRow(Storage, StorageCompatibility, Object.Row, Object.Object, Object.TypeInfo);
 			});
 	}
 
@@ -607,7 +609,7 @@ namespace UE::Editor::DataStorage
 		Storage.BatchAddRow(Batch.Table, Rows, [&RowIterator, &ObjectIterator, &TypeInfoIterator, this](RowHandle Row)
 			{
 				checkf(Row == *RowIterator++, TEXT("Expecting the same sequence of rows when batch adding object to TEDS Compatibility."));
-				SetupRow(Row, *ObjectIterator++, *TypeInfoIterator++);
+				SetupRow(Storage, StorageCompatibility, Row, *ObjectIterator++, *TypeInfoIterator++);
 			});
 	}
 
diff --git a/Engine/Plugins/Experimental/EditorDataStorage/Source/TedsCore/Private/Commands/EditorDataStorageCompatibilityCommands.h b/Engine/Plugins/Experimental/EditorDataStorage/Source/TedsCore/Private/Commands/EditorDataStorageCompatibilityCommands.h
--- a/Engine/Plugins/Experimental/EditorDataStorage/Source/TedsCore/Private/Commands/EditorDataStorageCompatibilityCommands.h
+++ b/Engine/Plugins/Experimental/EditorDataStorage/Source/TedsCore/Private/Commands/EditorDataStorageCompatibilityCommands.h
@@ -301,8 +301,10 @@ namespace UE::Editor::DataStorage
 	{
 		FCommandProcessor(ICoreProvider& InStorage, UEditorDataStorageCompatibility& InStorageCompatibility);
 
-		void SetupRow(RowHandle Row, UObject* Object);
-		void SetupRow(RowHandle Row, void* Object, TWeakObjectPtr<const UScriptStruct> TypeInfo);
+		static void SetupRow(ICoreProvider& Storage, UEditorDataStorageCompatibility& StorageCompatibility, 
+			RowHandle Row, UObject* Object);
+		static void SetupRow(ICoreProvider& Storage, UEditorDataStorageCompatibility& StorageCompatibility, 
+			RowHandle Row, void* Object, TWeakObjectPtr<const UScriptStruct> TypeInfo);
 		
 		void operator()(FNopCommand&) {}
 		void operator()(FTypeInfoReinstanced&) {}
diff --git a/Engine/Plugins/Experimental/EditorDataStorage/Source/TedsCore/Private/TypedElementDatabaseCompatibility.cpp b/Engine/Plugins/Experimental/EditorDataStorage/Source/TedsCore/Private/TypedElementDatabaseCompatibility.cpp
--- a/Engine/Plugins/Experimental/EditorDataStorage/Source/TedsCore/Private/TypedElementDatabaseCompatibility.cpp
+++ b/Engine/Plugins/Experimental/EditorDataStorage/Source/TedsCore/Private/TypedElementDatabaseCompatibility.cpp
@@ -5,6 +5,7 @@
 #include <utility>
 #include "Algo/Unique.h"
 #include "Async/UniqueLock.h"
+#include "DataStorage/Debug/Log.h"
 #include "Editor.h"
 #include "Editor/TransBuffer.h"
 #include "Elements/Columns/TypedElementCompatibilityColumns.h"
@@ -95,6 +96,8 @@ namespace UE::Editor::DataStorage::Private
 	static const FName CompatibilityUsesCommandBufferExtensionName(TEXT("CompatiblityUsesCommandBuffer"));
 } // namespace UE::Editor::DataStorage::Private
 
+thread_local UE::Editor::DataStorage::FObjectCompatibilityImmediateScope* UEditorDataStorageCompatibility::ImmediateModeTracker = nullptr;
+
 void UEditorDataStorageCompatibility::Initialize(UEditorDataStorage* InStorage)
 {
 	using namespace UE::Editor::DataStorage;
@@ -128,7 +131,7 @@ void UEditorDataStorageCompatibility::Deinitialize()
 	using namespace UE::Editor::DataStorage;
 
 	FScopedExclusiveLock Lock(EGlobalLockScope::Public);
-	
+
 	FCoreUObjectDelegates::PostReachabilityAnalysis.Remove(PostGcUnreachableAnalysisHandle);
 
 	FCoreUObjectDelegates::OnObjectsReinstanced.Remove(ObjectReinstancedDelegateHandle);
@@ -243,7 +246,7 @@ UE::Editor::DataStorage::RowHandle UEditorDataStorageCompatibility::AddCompatibl
 	// Because AddCompatibleObjectExplicitTransactionable needs a finer grained control over the lock, there's no higher up lock here.
 
 	bool bCanAddObject =
-		ensureMsgf(Storage, TEXT("Trying to add a UObject to Typed Element's Data Storage before the storage is available.")) &&
+		ensureMsgf(Storage, TEXT("Trying to add a UObject to the editor's Data Storage before it's available.")) &&
 		Object != nullptr &&
 		ShouldAddObject(Object);
 	return bCanAddObject ? AddCompatibleObjectExplicitTransactionable<true>(Object) : UE::Editor::DataStorage::InvalidRowHandle;
@@ -259,21 +262,28 @@ UE::Editor::DataStorage::RowHandle UEditorDataStorageCompatibility::AddCompatibl
 		return InvalidRowHandle;
 	}
 	
-	checkf(Storage, TEXT("Trying to add an object to Typed Element's Data Storage before the storage is available."));
+	checkf(Storage, TEXT("Trying to add an object to the editor's Data Storage before it's available."));
 	FScopedExclusiveLock Lock(EGlobalLockScope::Public);
 	
 	RowHandle Result = FindRowWithCompatibleObjectExplicit(Object);
 	if (!Storage->IsRowAvailable(Result))
 	{
-		Result = Storage->ReserveRow();
-		Storage->MapRow(ICompatibilityProvider::ObjectMappingDomain, FMapKey(Object), Result);
-		if (Private::bUseCommandBuffer)
+		if (const UScriptStruct* TypeInfoPtr = TypeInfo.Get(true))
 		{
-			QueuedCommands.AddCommand(FAddCompatibleExternalObject{ .Object = Object, .TypeInfo = TypeInfo, .Row = Result });
-		}
-		else
-		{
-			ExternalObjectsPendingRegistration.Add(Result, ExternalObjectRegistration{ .Object = Object, .TypeInfo = TypeInfo });
+			Result = Storage->ReserveRow();
+			Storage->MapRow(ICompatibilityProvider::ObjectMappingDomain, FMapKey(Object), Result);
+			if (ImmediateModeTracker)
+			{
+				FObjectCompatibilityImmediateScopeWriteAccess::RegisterObject(*ImmediateModeTracker, Object, TypeInfo.Get(true), Result);
+			}
+			else if (Private::bUseCommandBuffer)
+			{
+				QueuedCommands.AddCommand(FAddCompatibleExternalObject{ .Object = Object, .TypeInfo = TypeInfo, .Row = Result });
+			}
+			else
+			{
+				ExternalObjectsPendingRegistration.Add(Result, ExternalObjectRegistration{ .Object = Object, .TypeInfo = TypeInfo });
+			}
 		}
 	}
 	return Result;
@@ -372,6 +382,55 @@ void UEditorDataStorageCompatibility::ListExtensions(TFunctionRef<void(FName)> C
 	}
 }
 
+void UEditorDataStorageCompatibility::EnterImmediateMode(UE::Editor::DataStorage::FObjectCompatibilityImmediateScope& Scope)
+{
+	FObjectCompatibilityImmediateScopeWriteAccess::SetPrevious(Scope, ImmediateModeTracker);
+	ImmediateModeTracker = &Scope;
+}
+
+void UEditorDataStorageCompatibility::LeaveImmediateMode(UE::Editor::DataStorage::FObjectCompatibilityImmediateScope& Scope)
+{
+	using namespace UE::Editor::DataStorage;
+
+	FObjectCompatibilityImmediateScope* CurrentTracker = ImmediateModeTracker;
+	checkf(CurrentTracker, 
+		TEXT("Attempting to leave immediate mode in the editor data storage's compatibility layer, but immediate was never entered."));
+	checkf(CurrentTracker == &Scope, TEXT("Attempting to leave immediate mode in the editor data storage's compatibility ")
+		TEXT("layer, but the active scope doesn't match the provided scope."));
+
+	if (Scope.GetObjectCount() > 0)
+	{
+		FScopedExclusiveLock Lock(EGlobalLockScope::Public);
+
+		FObjectCompatibilityImmediateScopeWriteAccess::ListUObjects(*CurrentTracker,
+			[this](RowHandle Row, UObject* Object, const UClass* TypeInfo)
+			{
+				TableHandle Table = FindBestMatchingTable(TypeInfo);
+				checkf(Table != InvalidTableHandle,
+					TEXT("Unable to find a matching table in the compatibility layer for the editor's data storage for type '%s'"),
+					*Object->GetClass()->GetName());
+
+				Storage->AddRow(Row, Table, [this, Object](RowHandle Row)
+					{
+						FCommandProcessor::SetupRow(*Storage, *this, Row, Object);
+					});
+			});
+
+		FObjectCompatibilityImmediateScopeWriteAccess::ListExternalObjects(*CurrentTracker,
+			[this](RowHandle Row, void* Object, const UScriptStruct* TypeInfo)
+			{
+				TableHandle Table = FindBestMatchingTable(TypeInfo);
+				Table = (Table != InvalidTableHandle) ? Table : StandardExternalObjectTable;
+				Storage->AddRow(Row, Table, [this, Object, TypeInfo](RowHandle Row)
+					{
+						FCommandProcessor::SetupRow(*Storage, *this, Row, Object, TypeInfo);
+					});
+			});
+	}
+
+	ImmediateModeTracker = FObjectCompatibilityImmediateScopeWriteAccess::GetPrevious(*CurrentTracker);
+}
+
 void UEditorDataStorageCompatibility::Prepare()
 {
 	// Thread-safe as this is only called from a function that has an exclusive lock.
@@ -464,25 +523,39 @@ UE::Editor::DataStorage::RowHandle UEditorDataStorageCompatibility::AddCompatibl
 
 	RowHandle Result = FindRowWithCompatibleObjectExplicit(Object);
 	if (!Storage->IsRowAvailable(Result))
-	{
-		if (Private::GetAsyncLoadStrategy() == Private::EAsyncUObjectLoadStrategy::Assert)
+	{	
+		if (ImmediateModeTracker)
 		{
-			EInternalObjectFlags Flags = Object->GetInternalFlags();
-			ensureMsgf(
-				!(Flags & EInternalObjectFlags_AsyncLoading),
-				TEXT("Cannot add object that is being asynchronously loaded to TEDS"));
-			return InvalidRowHandle;
-		}
-		
-		Result = Storage->ReserveRow();
-		Storage->MapRow(ICompatibilityProvider::ObjectMappingDomain, FMapKey(Object), Result);
-		if (Private::bUseCommandBuffer)
-		{
-			QueuedCommands.AddCommand(FAddCompatibleUObject{ .Object = Object, .Row = Result });
+			checkf(
+				!(Object->GetInternalFlags() & EInternalObjectFlags_AsyncLoading),
+				TEXT("Cannot add object '%s' as it's being asynchronously loaded to the editor's data storage while in immediate mode."),
+				*Object->GetName());
+
+			Result = Storage->ReserveRow();
+			Storage->MapRow(ICompatibilityProvider::ObjectMappingDomain, FMapKey(Object), Result);
+			FObjectCompatibilityImmediateScopeWriteAccess::RegisterObject(*ImmediateModeTracker, Object, nullptr, Result);
 		}
 		else
 		{
-			UObjectsPendingRegistration.Add(Result, Object);
+			if (Private::GetAsyncLoadStrategy() == Private::EAsyncUObjectLoadStrategy::Assert)
+			{
+				EInternalObjectFlags Flags = Object->GetInternalFlags();
+				ensureMsgf(
+					!(Flags & EInternalObjectFlags_AsyncLoading),
+					TEXT("Cannot add object that is being asynchronously loaded to TEDS"));
+				return InvalidRowHandle;
+			}
+			Result = Storage->ReserveRow();
+			Storage->MapRow(ICompatibilityProvider::ObjectMappingDomain, FMapKey(Object), Result);
+
+			if (Private::bUseCommandBuffer)
+			{
+				QueuedCommands.AddCommand(FAddCompatibleUObject{ .Object = Object, .Row = Result });
+			}
+			else
+			{
+				UObjectsPendingRegistration.Add(Result, Object);
+			}
 		}
 		
 		if constexpr (bEnableTransactions)
@@ -876,7 +949,7 @@ void UEditorDataStorageCompatibility::PendingRegistration<AddressType>::ProcessE
 		FEntry* TableFront = Current;
 
 		// It is expected that InvalidTableHandle should always be sorted at the end
-		// This assertion ensures that we catch a change in this behaviour should the constant
+		// This assertion ensures that we catch a change in this behavior should the constant
 		// be modified. 
 		static_assert(InvalidTableHandle == TNumericLimits<TableHandle>::Max()); 
 		if (Entries[0].Table == InvalidTableHandle)
@@ -985,34 +1058,24 @@ void UEditorDataStorageCompatibility::TickPendingCommands()
 
 void UEditorDataStorageCompatibility::TickPendingUObjectRegistration()
 {
-	// Thread safe because it's only called from functions that already lock.
+	using namespace UE::Editor::DataStorage;
 
+	// Thread safe because it's only called from functions that already lock.
+	
 	if (!UObjectsPendingRegistration.IsEmpty())
 	{
 		UObjectsPendingRegistration.ProcessEntries(*Storage, *this,
 			[this](UE::Editor::DataStorage::RowHandle Row, const TWeakObjectPtr<UObject>& Object)
 			{
-				UE::Editor::DataStorage::ICoreProvider* Interface = Storage;
-				Interface->AddColumn(Row, FTypedElementUObjectColumn{ .Object = Object });
-				Interface->AddColumn(Row, FTypedElementUObjectIdColumn
-					{ 
-						.Id = Object->GetUniqueID(), 
-						.SerialNumber = GUObjectArray.GetSerialNumber(Object->GetUniqueID())
-					});
-				Interface->AddColumn(Row, FTypedElementClassTypeInfoColumn{ .TypeInfo = Object->GetClass() });
-				if (Object->HasAnyFlags(RF_ClassDefaultObject))
-				{
-					Interface->AddColumn<FTypedElementClassDefaultObjectTag>(Row);
-				}
-				// Make sure the new row is tagged for update.
-				Interface->AddColumn<FTypedElementSyncFromWorldTag>(Row);
-				TriggerOnObjectAdded(Object.Get(), Object->GetClass(), Row);
+				FCommandProcessor::SetupRow(*Storage, *this, Row, Object.Get());
 			});
 	}
 }
 
 void UEditorDataStorageCompatibility::TickPendingExternalObjectRegistration()
 {
+	using namespace UE::Editor::DataStorage;
+
 	// Thread safe because it's only called from functions that already lock.
 
 	if (!ExternalObjectsPendingRegistration.IsEmpty())
@@ -1020,12 +1083,7 @@ void UEditorDataStorageCompatibility::TickPendingExternalObjectRegistration()
 		ExternalObjectsPendingRegistration.ProcessEntries(*Storage, *this,
 			[this](UE::Editor::DataStorage::RowHandle Row, const ExternalObjectRegistration& Object)
 			{
-				UE::Editor::DataStorage::ICoreProvider* Interface = Storage;
-				Interface->AddColumn(Row, FTypedElementExternalObjectColumn{ .Object = Object.Object });
-				Interface->AddColumn(Row, FTypedElementScriptStructTypeInfoColumn{ .TypeInfo = Object.TypeInfo });
-				// Make sure the new row is tagged for update.
-				Interface->AddColumn<FTypedElementSyncFromWorldTag>(Row);
-				TriggerOnObjectAdded(Object.Object, Object.TypeInfo.Get(), Row);
+				FCommandProcessor::SetupRow(*Storage, *this, Row, Object.Object, Object.TypeInfo);
 			});
 	}
 }
diff --git a/Engine/Plugins/Experimental/EditorDataStorage/Source/TedsCore/Private/TypedElementDatabaseCompatibility.h b/Engine/Plugins/Experimental/EditorDataStorage/Source/TedsCore/Private/TypedElementDatabaseCompatibility.h
--- a/Engine/Plugins/Experimental/EditorDataStorage/Source/TedsCore/Private/TypedElementDatabaseCompatibility.h
+++ b/Engine/Plugins/Experimental/EditorDataStorage/Source/TedsCore/Private/TypedElementDatabaseCompatibility.h
@@ -66,6 +66,10 @@ public:
 
 	UE_API bool SupportsExtension(FName Extension) const override;
 	UE_API void ListExtensions(TFunctionRef<void(FName)> Callback) const override;
+
+protected:
+	UE_API void EnterImmediateMode(UE::Editor::DataStorage::FObjectCompatibilityImmediateScope& Scope) override;
+	UE_API void LeaveImmediateMode(UE::Editor::DataStorage::FObjectCompatibilityImmediateScope& Scope) override;
 	
 private:
 	// The below changes expect UEditorDataStorageCompatibility to be the object passed in to StoreUndo.
@@ -200,6 +204,8 @@ private:
 	TArray<TPair<UE::Editor::DataStorage::ObjectAddedCallback, FDelegateHandle>> ObjectAddedCallbackList;
 	TArray<TPair<UE::Editor::DataStorage::ObjectRemovedCallback, FDelegateHandle>> PreObjectRemovedCallbackList;
 
+	static thread_local UE::Editor::DataStorage::FObjectCompatibilityImmediateScope* ImmediateModeTracker;
+	
 	UE::Editor::DataStorage::TableHandle StandardUObjectTable{ UE::Editor::DataStorage::InvalidTableHandle };
 	UE::Editor::DataStorage::TableHandle StandardExternalObjectTable{ UE::Editor::DataStorage::InvalidTableHandle };
 	UE::Editor::DataStorage::ICoreProvider* Storage{ nullptr };
diff --git a/Engine/Source/Runtime/TypedElementFramework/Public/Elements/Interfaces/TypedElementDataStorageCompatibilityInterface.h b/Engine/Source/Runtime/TypedElementFramework/Public/Elements/Interfaces/TypedElementDataStorageCompatibilityInterface.h
--- a/Engine/Source/Runtime/TypedElementFramework/Public/Elements/Interfaces/TypedElementDataStorageCompatibilityInterface.h
+++ b/Engine/Source/Runtime/TypedElementFramework/Public/Elements/Interfaces/TypedElementDataStorageCompatibilityInterface.h
@@ -2,6 +2,7 @@
 
 #pragma once
 
+#include "Containers/Array.h"
 #include "Elements/Interfaces/TypedElementDataStorageInterface.h"
 #include "Features/IModularFeature.h"
 #include "Templates/Function.h"
@@ -14,188 +15,326 @@
 
 namespace UE::Editor::DataStorage
 {
-/**
- * Interface to provide compatibility with existing systems that don't directly
- * support the data storage.
- */
-class ICompatibilityProvider : public IModularFeature
-{
-public:
-	using ObjectRegistrationFilter = TFunction<bool(const ICompatibilityProvider&, const UObject*)>;
-	using ObjectToRowDealiaser = TFunction<RowHandle(const ICompatibilityProvider&, const UObject*)>;
-	
-	inline static const FName ObjectMappingDomain = "Object";
+	class FObjectCompatibilityImmediateScope;
 
 	/**
-	 * @section Type-agnostic functions
-	 * These allow compatibility with any type. These do eventually fall back to the explicit versions.
-	 * Any references given are non-owning so it's up to the caller to deregister the object after it's no longer
-	 * available.
+	 * Interface to provide compatibility with existing systems that don't directly
+	 * support the data storage.
 	 */
+	class ICompatibilityProvider : public IModularFeature
+	{
+		friend class FObjectCompatibilityImmediateScope;
+	public:
+		using ObjectRegistrationFilter = TFunction<bool(const ICompatibilityProvider&, const UObject*)>;
+		using ObjectToRowDealiaser = TFunction<RowHandle(const ICompatibilityProvider&, const UObject*)>;
+	
+		inline static const FName ObjectMappingDomain = "Object";
 
-	/** 
-	 * Adds a reference to an existing object to the data storage. The data storage does NOT take ownership of the object and
-	 * the caller is responsible for managing the life cycle of the object. The address is only used for associating the object
-	 * with a row and to setup the initial row data.
-	 */
-	template<typename ObjectType>
-	RowHandle AddCompatibleObject(ObjectType&& Object);
+		/**
+		 * @section Type-agnostic functions
+		 * These allow compatibility with any type. These do eventually fall back to the explicit versions.
+		 * Any references given are non-owning so it's up to the caller to deregister the object after it's no longer
+		 * available.
+		 */
+
+		/** 
+		 * Adds a reference to an existing object to the data storage. The data storage does NOT take ownership of the object and
+		 * the caller is responsible for managing the life cycle of the object. The address is only used for associating the object
+		 * with a row and to setup the initial row data.
+		 */
+		template<typename ObjectType>
+		RowHandle AddCompatibleObject(ObjectType&& Object);
 	
-	/** Removes a previously registered object from the data storage. */
-	template<typename ObjectType>
-	void RemoveCompatibleObject(ObjectType&& Object);
+		/** Removes a previously registered object from the data storage. */
+		template<typename ObjectType>
+		void RemoveCompatibleObject(ObjectType&& Object);
 
-	template<typename ObjectType>
-	RowHandle FindRowWithCompatibleObject(ObjectType&& Object) const;
+		template<typename ObjectType>
+		RowHandle FindRowWithCompatibleObject(ObjectType&& Object) const;
 
-	/**
-	 * @section Callback registration
-	 * Functions to register callbacks with the compatibility layer to help refine its operations.
-	 */
+		/**
+		 * @section Callback registration
+		 * Functions to register callbacks with the compatibility layer to help refine its operations.
+		 */
 	 
-	/**
-	 * Objects like actors are registered through the compatibility layer in bulk. This can lead to objects being added that cause
-	 * conflicts with other data in the data storage. This callback offers the opportunity to inspect the objects that are being
-	 * added and if they include an object that shouldn't be store it can filter them out.
-	 */
-	virtual void RegisterRegistrationFilter(ObjectRegistrationFilter Filter) = 0;
-	/**
-	 * Notifications and request can be made to the compatibility layer for objects that are stored but don't directly map to a row.
-	 * An example is a UObject represented by a column. If the UObject gets updated there's no direct mapping to the row the column is
-	 * stored in but the row still needs to be updated. For cases like this it's possible to store information to find the row that's
-	 * being aliased.
-	 */
-	virtual void RegisterDealiaserCallback(ObjectToRowDealiaser Dealiaser) = 0;
-	/**
-	 * Allows a specific type to be associated with a table. Whenever a compatible object is added, the type information of that object
-	 * will be used to find the closest match in the registered types and use the associated table. E.g. actors derive from uobjects so
-	 * if the type information of an actor is registered the actor table will be used instead of the uobject table.
-	 */
-	virtual void RegisterTypeTableAssociation(TWeakObjectPtr<UStruct> TypeInfo, TableHandle Table) = 0;
+		/**
+		 * Objects like actors are registered through the compatibility layer in bulk. This can lead to objects being added that cause
+		 * conflicts with other data in the data storage. This callback offers the opportunity to inspect the objects that are being
+		 * added and if they include an object that shouldn't be store it can filter them out.
+		 */
+		virtual void RegisterRegistrationFilter(ObjectRegistrationFilter Filter) = 0;
+		/**
+		 * Notifications and request can be made to the compatibility layer for objects that are stored but don't directly map to a row.
+		 * An example is a UObject represented by a column. If the UObject gets updated there's no direct mapping to the row the column is
+		 * stored in but the row still needs to be updated. For cases like this it's possible to store information to find the row that's
+		 * being aliased.
+		 */
+		virtual void RegisterDealiaserCallback(ObjectToRowDealiaser Dealiaser) = 0;
+		/**
+		 * Allows a specific type to be associated with a table. Whenever a compatible object is added, the type information of that object
+		 * will be used to find the closest match in the registered types and use the associated table. E.g. actors derive from uobjects so
+		 * if the type information of an actor is registered the actor table will be used instead of the uobject table.
+		 */
+		virtual void RegisterTypeTableAssociation(TWeakObjectPtr<UStruct> TypeInfo, TableHandle Table) = 0;
 
-	/**
-	 * @section Explicit functions
-	 * These are functions that work on specific types.
-	 */
+		/**
+		 * @section Explicit functions
+		 * These are functions that work on specific types.
+		 */
 
-	/** Adds a UObject to the data storage. */
-	virtual RowHandle AddCompatibleObjectExplicit(UObject* Object) = 0;
-	/** Adds an FStruct to the data storage. */
-	virtual RowHandle AddCompatibleObjectExplicit(void* Object, TWeakObjectPtr<const UScriptStruct> TypeInfo) = 0;
+		/** Adds a UObject to the data storage. */
+		virtual RowHandle AddCompatibleObjectExplicit(UObject* Object) = 0;
+		/** Adds an FStruct to the data storage. */
+		virtual RowHandle AddCompatibleObjectExplicit(void* Object, TWeakObjectPtr<const UScriptStruct> TypeInfo) = 0;
 	
-	/** Removes a UObject from the data storage. */
-	virtual void RemoveCompatibleObjectExplicit(UObject* Object) = 0;
-	/** Removes an FStruct from the data storage. */
-	virtual void RemoveCompatibleObjectExplicit(void* Object) = 0;
+		/** Removes a UObject from the data storage. */
+		virtual void RemoveCompatibleObjectExplicit(UObject* Object) = 0;
+		/** Removes an FStruct from the data storage. */
+		virtual void RemoveCompatibleObjectExplicit(void* Object) = 0;
+
+		/** Finds a previously stored UObject. If not found an invalid row handle will be returned. */
+		virtual RowHandle FindRowWithCompatibleObjectExplicit(const UObject* Object) const = 0;
+		/** Finds a previously stored FStruct. If not found an invalid row handle will be returned. */
+		virtual RowHandle FindRowWithCompatibleObjectExplicit(const void* Object) const = 0;
 
-	/** Finds a previously stored UObject. If not found an invalid row handle will be returned. */
-	virtual RowHandle FindRowWithCompatibleObjectExplicit(const UObject* Object) const = 0;
-	/** Finds a previously stored FStruct. If not found an invalid row handle will be returned. */
-	virtual RowHandle FindRowWithCompatibleObjectExplicit(const void* Object) const = 0;
+		/**
+		 * @section Miscellaneous functions
+		 */
+
+		/** Check if a custom extension is supported. This can be used to check for in-development features, custom extensions, etc. */
+		virtual bool SupportsExtension(FName Extension) const = 0;
+		/** Provides a list of all extensions that are enabled. */
+		virtual void ListExtensions(TFunctionRef<void(FName)> Callback) const = 0;
+
+	protected:
+		/**
+		 * @section Immediate mode scope functions
+		 */
+
+		struct UE_INTERNAL FObjectCompatibilityImmediateScopeWriteAccess
+		{
+			static void RegisterObject(FObjectCompatibilityImmediateScope& Scope, 
+				void* Object, const UScriptStruct* TypeInfo, RowHandle Row);
+			static void SetPrevious(FObjectCompatibilityImmediateScope& Scope, FObjectCompatibilityImmediateScope* InPrevious);
+			static FObjectCompatibilityImmediateScope* GetPrevious(const FObjectCompatibilityImmediateScope& Scope);
+
+			static void ListUObjects(const FObjectCompatibilityImmediateScope& Scope, 
+				TFunctionRef<void(RowHandle, UObject*, const UClass*)> Callback);
+			static void ListExternalObjects(const FObjectCompatibilityImmediateScope& Scope, 
+				TFunctionRef<void(RowHandle, void*, const UScriptStruct*)> Callback);
+		};
+		UE_INTERNAL
+		virtual void EnterImmediateMode(FObjectCompatibilityImmediateScope& Scope) = 0;
+		UE_INTERNAL
+		virtual void LeaveImmediateMode(FObjectCompatibilityImmediateScope& Scope) = 0;
+	};
 
 	/**
-	 * @section Miscellaneous functions
+	 * Scope object to tell TEDS Compatibility that any objects being registered on that thread with this scope should be
+	 * immediately registered instead of queued up once this goes out of scope.
+	 *
+	 * Scopes can be nested and are able to provide the row handles to the objects that have been newly or previously
+	 * registered. The scope is local to the calling thread so other threads will continue to queue registrations.
+	 * 
+	 * The immediate mode limits how much batch and distributed registration TEDS Compatibility can do. Avoid using
+	 * this scope when possible and prefer to limit the use to a small number of objects. This scope was introduced to
+	 * aid the transition to the multi-threaded environment TEDS provides and as such this scope may be removed in the future.
+	 * 
+	 * Warning: This scope is to aid the transition of existing single-threaded systems over to the editor data storage.
 	 */
+	class UE_INTERNAL FObjectCompatibilityImmediateScope
+	{
+		PRAGMA_DISABLE_INTERNAL_WARNINGS
+		friend struct ICompatibilityProvider::FObjectCompatibilityImmediateScopeWriteAccess;
+		PRAGMA_ENABLE_INTERNAL_WARNINGS
+	public:
+		explicit FObjectCompatibilityImmediateScope(ICompatibilityProvider& InCompatibility);
+		~FObjectCompatibilityImmediateScope();
 
-	/** Check if a custom extension is supported. This can be used to check for in-development features, custom extensions, etc. */
-	virtual bool SupportsExtension(FName Extension) const = 0;
-	/** Provides a list of all extensions that are enabled. */
-	virtual void ListExtensions(TFunctionRef<void(FName)> Callback) const = 0;
-};
-
-template<typename Type> Type* GetRawPointer(const TWeakObjectPtr<Type> Object)	{ return Object.Get(); }
-template<typename Type> Type* GetRawPointer(const TObjectPtr<Type> Object)		{ return Object.Get(); }
-template<typename Type> Type* GetRawPointer(const TStrongObjectPtr<Type> Object){ return Object.Get(); }
-template<typename Type> Type* GetRawPointer(const TObjectKey<Type> Object)		{ return Object.ResolveObjectPtr(); }
-template<typename Type> Type* GetRawPointer(const TUniquePtr<Type> Object)		{ return Object.Get(); }
-template<typename Type> Type* GetRawPointer(const TSharedPtr<Type> Object)		{ return Object.Get(); }
-#if UE_ENABLE_NOTNULL_WRAPPER
-template<typename Type> Type* GetRawPointer(TNotNull<Type*> Object)				{ return Object; }
-#endif
-template<typename Type> Type* GetRawPointer(Type* Object)						{ return Object; }
-template<typename Type> Type* GetRawPointer(Type& Object)						{ return &Object; }
-
-template<typename ObjectType>
-RowHandle ICompatibilityProvider::AddCompatibleObject(ObjectType&& Object)
-{
-	auto RawPointer = GetRawPointer(Forward<ObjectType>(Object));
-	using BaseType = std::remove_cv_t<std::remove_pointer_t<decltype(RawPointer)>>;
+		RowHandle FindObject(const void* Object) const;
+		int32 GetObjectCount() const;
+
+	private:
+		struct ObjectInfoType
+		{
+			void* Address;
+			const UScriptStruct* TypeInfo; // If this is null, Address points to a UObject.
+			RowHandle Row;
+		};
+		TArray<ObjectInfoType> Objects;
+		FObjectCompatibilityImmediateScope* Previous = nullptr;
+		ICompatibilityProvider& Compatibility;
+	};
+
+	// 
+	// Implementations
+	// 
+
+	PRAGMA_DISABLE_INTERNAL_WARNINGS // For FObjectCompatibilityImmediateScope and FObjectCompatibilityImmediateScopeWriteAccess
 
-	if constexpr (std::is_base_of_v<UObject, BaseType>)
+	inline FObjectCompatibilityImmediateScope::FObjectCompatibilityImmediateScope(ICompatibilityProvider& InCompatibility)
+		: Compatibility(InCompatibility)
 	{
-		return AddCompatibleObjectExplicit(RawPointer);
+		InCompatibility.EnterImmediateMode(*this);
 	}
-	else
+
+	inline FObjectCompatibilityImmediateScope::~FObjectCompatibilityImmediateScope()
 	{
-		return AddCompatibleObjectExplicit(RawPointer, BaseType::StaticStruct());
+		Compatibility.LeaveImmediateMode(*this);
 	}
-}
-
-template<typename ObjectType>
-void ICompatibilityProvider::RemoveCompatibleObject(ObjectType&& Object)
-{
-	RemoveCompatibleObjectExplicit(GetRawPointer(Forward<ObjectType>(Object)));
-}
 
-template<typename ObjectType>
-RowHandle ICompatibilityProvider::FindRowWithCompatibleObject(ObjectType&& Object) const
-{
-	return FindRowWithCompatibleObjectExplicit(GetRawPointer(Forward<ObjectType>(Object)));
-}
+	inline RowHandle FObjectCompatibilityImmediateScope::FindObject(const void* Object) const
+	{
+		const ObjectInfoType* Result = Objects.FindByPredicate([Object](const ObjectInfoType& Info) { return Info.Address == Object; });
+		return Result ? Result->Row : InvalidRowHandle;
+	}
+	inline int32 FObjectCompatibilityImmediateScope::GetObjectCount() const { return Objects.Num(); }
+	
+	inline void ICompatibilityProvider::FObjectCompatibilityImmediateScopeWriteAccess::RegisterObject(
+		FObjectCompatibilityImmediateScope& Scope, void* Object, const UScriptStruct* TypeInfo, RowHandle Row)
+	{
+		Scope.Objects.Add(
+			FObjectCompatibilityImmediateScope::ObjectInfoType
+			{
+				.Address = Object,
+				.TypeInfo = TypeInfo,
+				.Row = Row
+			});
+	}
+	
 
-template<typename Subsystem>
-struct TTypedElementSubsystemTraits final
-{
-	template<typename T, typename = void>
-	struct HasRequiresGameThreadVariable 
-	{ 
-		static constexpr bool bAvailable = false; 
-	};
-	template<typename T>
-	struct HasRequiresGameThreadVariable <T, decltype((void)T::bRequiresGameThread)>
-	{ 
-		static constexpr bool bAvailable = true; 
-	};
+	inline void ICompatibilityProvider::FObjectCompatibilityImmediateScopeWriteAccess::SetPrevious(
+		FObjectCompatibilityImmediateScope& Scope, FObjectCompatibilityImmediateScope* InPrevious)
+	{
+		checkf(Scope.Previous == nullptr, TEXT("A previous scope has already been set for FObjectCompatibilityImmediateScope."));
+		Scope.Previous = InPrevious;
+	}
 
-	template<typename T, typename = void>
-	struct HasIsHotReloadableVariable
+	inline FObjectCompatibilityImmediateScope* ICompatibilityProvider::FObjectCompatibilityImmediateScopeWriteAccess::GetPrevious(
+		const FObjectCompatibilityImmediateScope& Scope)
 	{ 
-		static constexpr bool bAvailable = false;
-	};
-	template<typename T>
-	struct HasIsHotReloadableVariable <T, decltype((void)T::bIsHotReloadable)>
-	{ 
-		static constexpr bool bAvailable = true;
-	};
+		return Scope.Previous;
+	}
 
-	static constexpr bool RequiresGameThread()
+	inline void ICompatibilityProvider::FObjectCompatibilityImmediateScopeWriteAccess::ListUObjects(
+		const FObjectCompatibilityImmediateScope& Scope, TFunctionRef<void(RowHandle, UObject*, const UClass*)> Callback)
 	{
-		if constexpr (HasRequiresGameThreadVariable<Subsystem>::bAvailable)
+		for (const FObjectCompatibilityImmediateScope::ObjectInfoType& Entry : Scope.Objects)
 		{
-			return Subsystem::bRequiresGameThread;
+			if (!Entry.TypeInfo)
+			{
+				Callback(Entry.Row, static_cast<UObject*>(Entry.Address), static_cast<UObject*>(Entry.Address)->GetClass());
+			}
 		}
-		else
+	}
+
+	inline void ICompatibilityProvider::FObjectCompatibilityImmediateScopeWriteAccess::ListExternalObjects(
+		const FObjectCompatibilityImmediateScope& Scope, TFunctionRef<void(RowHandle, void*, const UScriptStruct*)> Callback)
+	{
+		for (const FObjectCompatibilityImmediateScope::ObjectInfoType& Entry : Scope.Objects)
 		{
-			static_assert(HasRequiresGameThreadVariable<Subsystem>::bAvailable, "Subsystem provided to the Typed Elements did not "
-				"have a 'static constexpr bool bRequiresGameThread = true|false` declared or have a specialization for "
-				"TTypedElementSubsystemTraits.");
-			return true;
+			if (Entry.TypeInfo)
+			{ 
+				Callback(Entry.Row, Entry.Address, Entry.TypeInfo);
+			}
 		}
 	}
 
-	static constexpr bool IsHotReloadable()
+	PRAGMA_ENABLE_INTERNAL_WARNINGS // For FObjectCompatibilityImmediateScope and FObjectCompatibilityImmediateScopeWriteAccess
+	
+	template<typename Type> Type* GetRawPointer(const TWeakObjectPtr<Type> Object)	{ return Object.Get(); }
+	template<typename Type> Type* GetRawPointer(const TObjectPtr<Type> Object)		{ return Object.Get(); }
+	template<typename Type> Type* GetRawPointer(const TStrongObjectPtr<Type> Object){ return Object.Get(); }
+	template<typename Type> Type* GetRawPointer(const TObjectKey<Type> Object)		{ return Object.ResolveObjectPtr(); }
+	template<typename Type> Type* GetRawPointer(const TUniquePtr<Type> Object)		{ return Object.Get(); }
+	template<typename Type> Type* GetRawPointer(const TSharedPtr<Type> Object)		{ return Object.Get(); }
+	#if UE_ENABLE_NOTNULL_WRAPPER
+	template<typename Type> Type* GetRawPointer(TNotNull<Type*> Object)				{ return Object; }
+	#endif
+	template<typename Type> Type* GetRawPointer(Type* Object)						{ return Object; }
+	template<typename Type> Type* GetRawPointer(Type& Object)						{ return &Object; }
+
+	template<typename ObjectType>
+	RowHandle ICompatibilityProvider::AddCompatibleObject(ObjectType&& Object)
 	{
-		if constexpr (HasIsHotReloadableVariable<Subsystem>::bAvailable)
+		auto RawPointer = GetRawPointer(Forward<ObjectType>(Object));
+		using BaseType = std::remove_cv_t<std::remove_pointer_t<decltype(RawPointer)>>;
+
+		if constexpr (std::is_base_of_v<UObject, BaseType>)
 		{
-			return Subsystem::bIsHotReloadable;
+			return AddCompatibleObjectExplicit(RawPointer);
 		}
 		else
 		{
-			static_assert(HasIsHotReloadableVariable<Subsystem>::bAvailable, "Subsystem provided to the Typed Elements did not "
-				"have a 'static constexpr bool bIsHotReloadable = true|false` declared or have a specialization for "
-				"TTypedElementSubsystemTraits.");
-			return false;
+			return AddCompatibleObjectExplicit(RawPointer, BaseType::StaticStruct());
 		}
 	}
-};
+
+	template<typename ObjectType>
+	void ICompatibilityProvider::RemoveCompatibleObject(ObjectType&& Object)
+	{
+		RemoveCompatibleObjectExplicit(GetRawPointer(Forward<ObjectType>(Object)));
+	}
+
+	template<typename ObjectType>
+	RowHandle ICompatibilityProvider::FindRowWithCompatibleObject(ObjectType&& Object) const
+	{
+		return FindRowWithCompatibleObjectExplicit(GetRawPointer(Forward<ObjectType>(Object)));
+	}
+
+	template<typename Subsystem>
+	struct TTypedElementSubsystemTraits final
+	{
+		template<typename T, typename = void>
+		struct HasRequiresGameThreadVariable 
+		{ 
+			static constexpr bool bAvailable = false; 
+		};
+		template<typename T>
+		struct HasRequiresGameThreadVariable <T, decltype((void)T::bRequiresGameThread)>
+		{ 
+			static constexpr bool bAvailable = true; 
+		};
+
+		template<typename T, typename = void>
+		struct HasIsHotReloadableVariable
+		{ 
+			static constexpr bool bAvailable = false;
+		};
+		template<typename T>
+		struct HasIsHotReloadableVariable <T, decltype((void)T::bIsHotReloadable)>
+		{ 
+			static constexpr bool bAvailable = true;
+		};
+
+		static constexpr bool RequiresGameThread()
+		{
+			if constexpr (HasRequiresGameThreadVariable<Subsystem>::bAvailable)
+			{
+				return Subsystem::bRequiresGameThread;
+			}
+			else
+			{
+				static_assert(HasRequiresGameThreadVariable<Subsystem>::bAvailable, "Subsystem provided to the Typed Elements did not "
+					"have a 'static constexpr bool bRequiresGameThread = true|false` declared or have a specialization for "
+					"TTypedElementSubsystemTraits.");
+				return true;
+			}
+		}
+
+		static constexpr bool IsHotReloadable()
+		{
+			if constexpr (HasIsHotReloadableVariable<Subsystem>::bAvailable)
+			{
+				return Subsystem::bIsHotReloadable;
+			}
+			else
+			{
+				static_assert(HasIsHotReloadableVariable<Subsystem>::bAvailable, "Subsystem provided to the Typed Elements did not "
+					"have a 'static constexpr bool bIsHotReloadable = true|false` declared or have a specialization for "
+					"TTypedElementSubsystemTraits.");
+				return false;
+			}
+		}
+	};
 } // namespace UE::Editor::DataStorage
