[Fix] Incorrect PRAGMA_DISABLE_DEPRECATION_WARNINGS usage

#Synced-CL 50350999

diff --git a/Engine/Source/Runtime/TypedElementFramework/Public/Elements/Interfaces/Capabilities/TypedElementUiStyleOverrideCapability.h b/Engine/Source/Runtime/TypedElementFramework/Public/Elements/Interfaces/Capabilities/TypedElementUiStyleOverrideCapability.h
--- a/Engine/Source/Runtime/TypedElementFramework/Public/Elements/Interfaces/Capabilities/TypedElementUiStyleOverrideCapability.h
+++ b/Engine/Source/Runtime/TypedElementFramework/Public/Elements/Interfaces/Capabilities/TypedElementUiStyleOverrideCapability.h
@@ -7,15 +7,19 @@
 
 struct FSlateColor;
 
-PRAGMA_DISABLE_DEPRECATION_WARNINGS
 /**
  * Interface to provide access to widgets that support working with style overrides.
  */
 class UE_DEPRECATED(5.8, "Please use FTedsOutlinerColumn to find the foreground color using Helpers::GetTreeItemFromRowHandle and FSceneOutlinerCommonLabelData::GetForegroundColor")
-	ITypedElementUiStyleOverrideCapability : public ITypedElementUiCapability
+	ITypedElementUiStyleOverrideCapability
+PRAGMA_DISABLE_DEPRECATION_WARNINGS
+: public ITypedElementUiCapability
+PRAGMA_ENABLE_DEPRECATION_WARNINGS
 {
 public:
+	PRAGMA_DISABLE_DEPRECATION_WARNINGS
 	SLATE_METADATA_TYPE(ITypedElementUiStyleOverrideCapability, ITypedElementUiCapability)
+	PRAGMA_ENABLE_DEPRECATION_WARNINGS
 
 	~ITypedElementUiStyleOverrideCapability() override = default;
 
@@ -23,7 +27,11 @@ public:
 };
 
 template<typename WidgetType>
-class UE_DEPRECATED(5.8, "Please use FTedsOutlinerColumn to find the foreground color using Helpers::GetTreeItemFromRowHandle and FSceneOutlinerCommonLabelData::GetForegroundColor") TTypedElementUiStyleOverrideCapability : public ITypedElementUiStyleOverrideCapability
+class UE_DEPRECATED(5.8, "Please use FTedsOutlinerColumn to find the foreground color using Helpers::GetTreeItemFromRowHandle and FSceneOutlinerCommonLabelData::GetForegroundColor")
+	TTypedElementUiStyleOverrideCapability
+PRAGMA_DISABLE_DEPRECATION_WARNINGS
+: public ITypedElementUiStyleOverrideCapability
+PRAGMA_ENABLE_DEPRECATION_WARNINGS
 {
 public:
 	explicit TTypedElementUiStyleOverrideCapability(WidgetType& InWidget) : Widget(InWidget){}
@@ -36,4 +44,3 @@ public:
 private:
 	WidgetType& Widget;
 };
-PRAGMA_ENABLE_DEPRECATION_WARNINGS
\ No newline at end of file
diff --git a/Engine/Source/Runtime/TypedElementFramework/Public/Elements/Interfaces/Capabilities/TypedElementUiTextCapability.h b/Engine/Source/Runtime/TypedElementFramework/Public/Elements/Interfaces/Capabilities/TypedElementUiTextCapability.h
--- a/Engine/Source/Runtime/TypedElementFramework/Public/Elements/Interfaces/Capabilities/TypedElementUiTextCapability.h
+++ b/Engine/Source/Runtime/TypedElementFramework/Public/Elements/Interfaces/Capabilities/TypedElementUiTextCapability.h
@@ -6,14 +6,19 @@
 #include "Internationalization/Text.h"
 #include "Misc/Attribute.h"
 
-PRAGMA_DISABLE_DEPRECATION_WARNINGS
 /**
  * Interface to provide access to widgets that support working with text.
  */
-class UE_DEPRECATED(5.8, "Please use FTedsOutlinerColumn to get the Outlier ptr and use its GetFilterHighlightText() function") ITypedElementUiTextCapability : public ITypedElementUiCapability
+class UE_DEPRECATED(5.8, "Please use FTedsOutlinerColumn to get the Outlier ptr and use its GetFilterHighlightText() function")
+	ITypedElementUiTextCapability
+PRAGMA_DISABLE_DEPRECATION_WARNINGS
+: public ITypedElementUiCapability
+PRAGMA_ENABLE_DEPRECATION_WARNINGS
 {
 public:
+	PRAGMA_DISABLE_DEPRECATION_WARNINGS
 	SLATE_METADATA_TYPE(ITypedElementUiTextCapability, ITypedElementUiCapability)
+	PRAGMA_ENABLE_DEPRECATION_WARNINGS
 
 	~ITypedElementUiTextCapability() override = default;
 
@@ -23,7 +28,11 @@ public:
 };
 
 template<typename WidgetType>
-class UE_DEPRECATED(5.8, "Please use FTedsOutlinerColumn to get the Outlier ptr and use its GetFilterHighlightText() function") TTypedElementUiTextCapability : public ITypedElementUiTextCapability
+class UE_DEPRECATED(5.8, "Please use FTedsOutlinerColumn to get the Outlier ptr and use its GetFilterHighlightText() function")
+	TTypedElementUiTextCapability
+PRAGMA_DISABLE_DEPRECATION_WARNINGS
+: public ITypedElementUiTextCapability
+PRAGMA_ENABLE_DEPRECATION_WARNINGS
 {
 public:
 	explicit TTypedElementUiTextCapability(WidgetType& InWidget) : Widget(InWidget){}
@@ -41,4 +50,3 @@ public:
 private:
 	WidgetType& Widget;
 };
-PRAGMA_ENABLE_DEPRECATION_WARNINGS
\ No newline at end of file
diff --git a/Engine/Source/Runtime/TypedElementFramework/Public/Elements/Interfaces/Capabilities/TypedElementUiTooltipCapability.h b/Engine/Source/Runtime/TypedElementFramework/Public/Elements/Interfaces/Capabilities/TypedElementUiTooltipCapability.h
--- a/Engine/Source/Runtime/TypedElementFramework/Public/Elements/Interfaces/Capabilities/TypedElementUiTooltipCapability.h
+++ b/Engine/Source/Runtime/TypedElementFramework/Public/Elements/Interfaces/Capabilities/TypedElementUiTooltipCapability.h
@@ -10,11 +10,16 @@
 /**
  * Interface to provide access to tool tips on widgets.
  */
+class UE_DEPRECATED(5.8, "Please use FTypedElementLabelColumn and FAttributeBinder to bind the tooltip attribute to the widget instead")
+	ITypedElementUiTooltipCapability
 PRAGMA_DISABLE_DEPRECATION_WARNINGS
-class UE_DEPRECATED(5.8, "Please use FTypedElementLabelColumn and FAttributeBinder to bind the tooltip attribute to the widget instead") ITypedElementUiTooltipCapability : public ITypedElementUiCapability
+: public ITypedElementUiCapability
+PRAGMA_ENABLE_DEPRECATION_WARNINGS
 {
 public:
+	PRAGMA_DISABLE_DEPRECATION_WARNINGS
 	SLATE_METADATA_TYPE(ITypedElementUiTooltipCapability, ITypedElementUiCapability)
+	PRAGMA_ENABLE_DEPRECATION_WARNINGS
 
 	~ITypedElementUiTooltipCapability() override = default;
 
@@ -28,7 +33,11 @@ public:
 };
 
 template<typename WidgetType>
-class UE_DEPRECATED(5.8, "Please use FTypedElementLabelColumn and FAttributeBinder to bind the tooltip attribute to the widget instead") TTypedElementUiTooltipCapability : public ITypedElementUiTooltipCapability
+class UE_DEPRECATED(5.8, "Please use FTypedElementLabelColumn and FAttributeBinder to bind the tooltip attribute to the widget instead")
+	TTypedElementUiTooltipCapability
+PRAGMA_DISABLE_DEPRECATION_WARNINGS
+: public ITypedElementUiTooltipCapability
+PRAGMA_ENABLE_DEPRECATION_WARNINGS
 {
 public:
 	explicit TTypedElementUiTooltipCapability(WidgetType& InWidget) : Widget(InWidget){}
@@ -66,4 +75,3 @@ public:
 private:
 	WidgetType& Widget;
 };
-PRAGMA_ENABLE_DEPRECATION_WARNINGS
\ No newline at end of file
