[TEDS] Extending the mapping keys

The following additions to the mapping key were made:
- A new type "KeyIssue" was added to allow reporting errors when creating a key.
- There is now a templated function "IsType" to check if the key has the expected type.
- The templated function "GetAs" was added to retrieve a const reference to the

#Synced-CL 50897653

diff --git a/Engine/Source/Runtime/TypedElementFramework/Private/DataStorage/MapKey.cpp b/Engine/Source/Runtime/TypedElementFramework/Private/DataStorage/MapKey.cpp
--- a/Engine/Source/Runtime/TypedElementFramework/Private/DataStorage/MapKey.cpp
+++ b/Engine/Source/Runtime/TypedElementFramework/Private/DataStorage/MapKey.cpp
@@ -32,6 +32,14 @@ namespace UE::Editor::DataStorage
 			return CityHash64WithSeed(reinterpret_cast<const char*>(*SubPath), SubPath.Len() * sizeof(**SubPath), Hash);
 		}
 		uint64 operator()(const TUniquePtr<FSoftObjectPath>& Value) { return operator()(Value.Get()); }
+		uint64 operator()(const FKeyIssue& Value)
+		{
+			return operator()(Value.Description);
+		}
+		uint64 operator()(const FKeyIssueView& Value)
+		{
+			return operator()(Value.Description);
+		}
 	};
 
 	struct FKeyToKeyViewConverter
@@ -42,6 +50,11 @@ namespace UE::Editor::DataStorage
 			return FMapKeyView::KeyViewType(TInPlaceType<T>(), Value);
 		}
 
+		FMapKeyView::KeyViewType operator()(const FKeyIssue& Value)
+		{
+			return FMapKeyView::KeyViewType(TInPlaceType<FKeyIssueView>(), Value);
+		}
+
 		FMapKeyView::KeyViewType operator()(const FString& Value)
 		{
 			return FMapKeyView::KeyViewType(TInPlaceType<FStringView>(), Value);
@@ -71,6 +84,11 @@ namespace UE::Editor::DataStorage
 			return FMapKey(Value);
 		}
 
+		FMapKey operator()(const FKeyIssueView& Value)
+		{
+			return FMapKey(FKeyIssue(Value.Description));
+		}
+
 		FMapKey operator()(const FStringView& Value)
 		{
 			return FMapKey(FString(Value));
@@ -167,6 +185,18 @@ namespace UE::Editor::DataStorage
 			return Local ? *Local == Other : false;
 		}
 
+		bool operator()(const FKeyIssue& Other)
+		{
+			const FKeyIssueView* Local = Key.TryGet<FKeyIssueView>();
+			return Local ? *Local == Other : false;
+		}
+
+		bool operator()(const FKeyIssueView& Other)
+		{
+			const FKeyIssueView* Local = Key.TryGet<FKeyIssueView>();
+			return Local ? *Local == Other : false;
+		}
+
 		bool operator()(const FString& Other)
 		{
 			const FStringView* Local = Key.TryGet<FStringView>();
@@ -211,6 +241,13 @@ namespace UE::Editor::DataStorage
 	struct FToString
 	{
 		FString operator()(FEmptyVariantState) { return TEXT("Empty"); }
+		FString operator()(const FKeyIssue& Value) { return FString::Printf(TEXT("Issue: '%s'"), *Value.Description); }
+		FString operator()(const FKeyIssueView& Value)
+		{ 
+			return Value.Description.IsEmpty()
+				? FString(TEXT("Issue: <no description>"))
+				: FString::Printf(TEXT("Issue: '%.*s'"), Value.Description.Len(), Value.Description.GetData());
+		}
 		FString operator()(const void* Value) { return FString::Printf(TEXT("Pointer(0x%p)"), Value); }
 		FString operator()(const UObject* Value) { return FString::Printf(TEXT("UObject(0x%p)"), Value); }
 		FString operator()(int64 Value) { return FString::Printf(TEXT("%lld"), Value); }
@@ -225,6 +262,67 @@ namespace UE::Editor::DataStorage
 		FString operator()(const TUniquePtr<FSoftObjectPath>& Value) { return Value ? Value->ToString() : TEXT("Empty path"); }
 	};
 
+	//
+	// FKeyIssue
+	//
+
+	FKeyIssue::FKeyIssue(const FString::ElementType* Description)
+		: Description(Description)
+	{
+	}
+
+	FKeyIssue::FKeyIssue(FString Description)
+		: Description(MoveTemp(Description))
+	{
+	}
+	
+	FKeyIssue::FKeyIssue(FStringView Description)
+		: Description(Description)
+	{
+	}
+
+	FKeyIssue::FKeyIssue(const FKeyIssueView& Description)
+		: Description(Description.Description)
+	{
+	}
+
+	bool FKeyIssue::operator==(const FKeyIssueView& Rhs) const
+	{
+		return Description == Rhs.Description;
+	}
+
+	bool FKeyIssue::operator!=(const FKeyIssueView& Rhs) const
+	{
+		return Description != Rhs.Description;
+	}
+
+
+	//
+	// FKeyIssueView
+	//
+
+	FKeyIssueView::FKeyIssueView(FStringView Description)
+		: Description(Description)
+	{
+	}
+
+	FKeyIssueView::FKeyIssueView(const FKeyIssue& Description)
+		: Description(Description.Description)
+	{
+	}
+
+	bool FKeyIssueView::operator==(const FKeyIssue& Rhs) const
+	{
+		return Description == Rhs.Description;
+	}
+
+	bool FKeyIssueView::operator!=(const FKeyIssue& Rhs) const
+	{
+		return Description != Rhs.Description;
+	}
+
+
+
 	//
 	// FMapKey
 	//
@@ -247,6 +345,11 @@ namespace UE::Editor::DataStorage
 	{
 	}
 
+	FMapKey::FMapKey(FKeyIssue Key)
+		: Key(TInPlaceType<FKeyIssue>(), MoveTemp(Key))
+	{
+	}
+
 	FMapKey::FMapKey(const void* Key)
 		: Key(TInPlaceType<const void*>(), Key)
 	{
@@ -349,6 +452,16 @@ namespace UE::Editor::DataStorage
 	{
 	}
 
+	FMapKeyView::FMapKeyView(const FKeyIssue& Key)
+		: Key(TInPlaceType<FKeyIssueView>(), Key)
+	{
+	}
+
+	FMapKeyView::FMapKeyView(FKeyIssueView Key)
+		: Key(TInPlaceType<FKeyIssueView>(), Key)
+	{
+	}
+
 	FMapKeyView::FMapKeyView(const void* Key)
 		: Key(TInPlaceType<const void*>(), Key)
 	{
diff --git a/Engine/Source/Runtime/TypedElementFramework/Public/DataStorage/MapKey.h b/Engine/Source/Runtime/TypedElementFramework/Public/DataStorage/MapKey.h
--- a/Engine/Source/Runtime/TypedElementFramework/Public/DataStorage/MapKey.h
+++ b/Engine/Source/Runtime/TypedElementFramework/Public/DataStorage/MapKey.h
@@ -14,6 +14,40 @@ class UObject;
 namespace UE::Editor::DataStorage
 {
 	class FMapKeyView;
+	struct FKeyIssueView;
+
+	/**
+	 * If the key encountered a problem the type will be switched to FKeyIssue and the description contains an explanation of the problem.
+	 */
+	struct FKeyIssue
+	{
+		FKeyIssue() = default;
+		TYPEDELEMENTFRAMEWORK_API FKeyIssue(const FString::ElementType* Description);
+		TYPEDELEMENTFRAMEWORK_API FKeyIssue(FString Description);
+		TYPEDELEMENTFRAMEWORK_API FKeyIssue(FStringView Description);
+		TYPEDELEMENTFRAMEWORK_API FKeyIssue(const FKeyIssueView& Description);
+
+		bool operator==(const FKeyIssue& Rhs) const = default;
+		bool operator!=(const FKeyIssue& Rhs) const = default;
+		TYPEDELEMENTFRAMEWORK_API bool operator==(const FKeyIssueView& Rhs) const;
+		TYPEDELEMENTFRAMEWORK_API bool operator!=(const FKeyIssueView& Rhs) const;
+
+		FString Description;
+	};
+
+	struct FKeyIssueView
+	{
+		FKeyIssueView() = default;
+		TYPEDELEMENTFRAMEWORK_API FKeyIssueView(FStringView Description);
+		TYPEDELEMENTFRAMEWORK_API FKeyIssueView(const FKeyIssue& Description);
+
+		bool operator==(const FKeyIssueView& Rhs) const = default;
+		bool operator!=(const FKeyIssueView& Rhs) const = default;
+		TYPEDELEMENTFRAMEWORK_API bool operator==(const FKeyIssue& Rhs) const;
+		TYPEDELEMENTFRAMEWORK_API bool operator!=(const FKeyIssue& Rhs) const;
+
+		FStringView Description;
+	};
 
 	class FMapKey
 	{
@@ -25,10 +59,9 @@ namespace UE::Editor::DataStorage
 	public:
 		FMapKey() = default;
 		
-		TYPEDELEMENTFRAMEWORK_API FMapKey(FMapKey&& Rhs);
-		TYPEDELEMENTFRAMEWORK_API FMapKey& operator=(FMapKey&& Rhs);
-		
 		TYPEDELEMENTFRAMEWORK_API FMapKey(const FMapKey& Key);
+		TYPEDELEMENTFRAMEWORK_API FMapKey(FMapKey&& Rhs);
+		TYPEDELEMENTFRAMEWORK_API explicit FMapKey(FKeyIssue Key);
 		TYPEDELEMENTFRAMEWORK_API explicit FMapKey(const void* Key);
 		TYPEDELEMENTFRAMEWORK_API explicit FMapKey(const UObject* Key);
 		TYPEDELEMENTFRAMEWORK_API explicit FMapKey(int64 Key);
@@ -40,6 +73,7 @@ namespace UE::Editor::DataStorage
 		TYPEDELEMENTFRAMEWORK_API explicit FMapKey(FSoftObjectPath Key);
 
 		TYPEDELEMENTFRAMEWORK_API FMapKey& operator=(const FMapKey& Rhs);
+		TYPEDELEMENTFRAMEWORK_API FMapKey& operator=(FMapKey&& Rhs);
 
 		TYPEDELEMENTFRAMEWORK_API bool operator==(const FMapKey& Rhs) const;
 		TYPEDELEMENTFRAMEWORK_API bool operator==(const FMapKeyView& Rhs) const;
@@ -51,8 +85,13 @@ namespace UE::Editor::DataStorage
 		TYPEDELEMENTFRAMEWORK_API void Clear();
 		TYPEDELEMENTFRAMEWORK_API FString ToString() const;
 
+		template<typename T>
+		bool IsType() const { return Key.IsType<T>(); }
+		template<typename T>
+		const T& GetAs() const { return Key.Get<T>(); }
+
 	private:
-		using KeyType = TVariant<FEmptyVariantState, const void*, const UObject*, int64, uint64, float, double, FString, FName, TUniquePtr<FSoftObjectPath>>;
+		using KeyType = TVariant<FEmptyVariantState, FKeyIssue, const void*, const UObject*, int64, uint64, float, double, FString, FName, TUniquePtr<FSoftObjectPath>>;
 
 		KeyType Key;
 	};
@@ -65,6 +104,8 @@ namespace UE::Editor::DataStorage
 	public:
 		FMapKeyView() = default;
 		TYPEDELEMENTFRAMEWORK_API FMapKeyView(const FMapKey& InKey);
+		TYPEDELEMENTFRAMEWORK_API explicit FMapKeyView(const FKeyIssue& Key);
+		TYPEDELEMENTFRAMEWORK_API explicit FMapKeyView(FKeyIssueView Key);
 		TYPEDELEMENTFRAMEWORK_API explicit FMapKeyView(const void* Key);
 		TYPEDELEMENTFRAMEWORK_API explicit FMapKeyView(const UObject* Key);
 		TYPEDELEMENTFRAMEWORK_API explicit FMapKeyView(int64 Key);
@@ -88,8 +129,13 @@ namespace UE::Editor::DataStorage
 		TYPEDELEMENTFRAMEWORK_API uint64 CalculateHash() const;
 		TYPEDELEMENTFRAMEWORK_API FString ToString() const;
 
+		template<typename T>
+		bool IsType() const { return Key.IsType<T>(); }
+		template<typename T>
+		const T& GetAs() const { return Key.Get<T>(); }
+
 	private:
-		using KeyViewType = TVariant<FEmptyVariantState, const void*, const UObject*, int64, uint64, float, double, FStringView, const FName*, const FSoftObjectPath*>;
+		using KeyViewType = TVariant<FEmptyVariantState, FKeyIssueView, const void*, const UObject*, int64, uint64, float, double, FStringView, const FName*, const FSoftObjectPath*>;
 
 		KeyViewType Key;
 	};
