TEDS: Add ability to walk only children in hierarchy.

#Synced-CL 52134663

diff --git a/Engine/Plugins/Experimental/EditorDataStorage/Source/TedsCore/Private/Hierarchy/EditorDataStorageHierarchyImplementation.cpp b/Engine/Plugins/Experimental/EditorDataStorage/Source/TedsCore/Private/Hierarchy/EditorDataStorageHierarchyImplementation.cpp
--- a/Engine/Plugins/Experimental/EditorDataStorage/Source/TedsCore/Private/Hierarchy/EditorDataStorageHierarchyImplementation.cpp
+++ b/Engine/Plugins/Experimental/EditorDataStorage/Source/TedsCore/Private/Hierarchy/EditorDataStorageHierarchyImplementation.cpp
@@ -503,6 +503,28 @@ namespace UE::Editor::DataStorage
 		}
 	}
 
+	bool FTedsHierarchyAccessInterface::IterateChildren(const ICoreProvider& Context, RowHandle Row, 
+		TFunctionRef<bool(const ICoreProvider& Context, RowHandle Child)> OnVisitedFn) const
+	{
+		if (!Context.IsRowAvailable(Row))
+		{
+			return false;
+		}
+
+		const FEditorDataHierarchyData_Template* HierarchyData = static_cast<const FEditorDataHierarchyData_Template*>(Context.GetColumnData(Row, HierarchyDataColumnType));
+		if (HierarchyData)
+		{
+			for (RowHandle ChildRow : HierarchyData->Children)
+			{
+				if (!OnVisitedFn(Context, ChildRow))
+				{
+					return false;
+				}
+			}
+		}
+		return true;
+	}
+
 	void FTedsHierarchyAccessInterface::SetParentRow(ICoreProvider* CoreProvider, RowHandle Target, RowHandle Parent) const
 	{
 		CoreProvider->RemoveColumn(Target, UnresolvedParentColumnType);
diff --git a/Engine/Plugins/Experimental/EditorDataStorage/Source/TedsCore/Private/Hierarchy/EditorDataStorageHierarchyImplementation.h b/Engine/Plugins/Experimental/EditorDataStorage/Source/TedsCore/Private/Hierarchy/EditorDataStorageHierarchyImplementation.h
--- a/Engine/Plugins/Experimental/EditorDataStorage/Source/TedsCore/Private/Hierarchy/EditorDataStorageHierarchyImplementation.h
+++ b/Engine/Plugins/Experimental/EditorDataStorage/Source/TedsCore/Private/Hierarchy/EditorDataStorageHierarchyImplementation.h
@@ -91,6 +91,10 @@ namespace UE::Editor::DataStorage
 			RowHandle Row,
 			TFunction<void(const ICoreProvider& Context, RowHandle Owner, RowHandle Target)> OnVisitedFn,
 			ICoreProvider::ETraversalOrder TraversalOrder = ICoreProvider::ETraversalOrder::PreOrder) const;
+		bool IterateChildren(
+			const ICoreProvider& Context,
+			RowHandle Row,
+			TFunctionRef<bool(const ICoreProvider& Context, RowHandle Child)> OnVisitedFn) const;
 		
 		TFunction<RowHandle(const void*, const UScriptStruct*)> CreateParentExtractionFunction() const;
 		
diff --git a/Engine/Plugins/Experimental/EditorDataStorage/Source/TedsCore/Private/TypedElementDatabase.cpp b/Engine/Plugins/Experimental/EditorDataStorage/Source/TedsCore/Private/TypedElementDatabase.cpp
--- a/Engine/Plugins/Experimental/EditorDataStorage/Source/TedsCore/Private/TypedElementDatabase.cpp
+++ b/Engine/Plugins/Experimental/EditorDataStorage/Source/TedsCore/Private/TypedElementDatabase.cpp
@@ -1296,6 +1296,18 @@ void UEditorDataStorage::WalkDepthFirst(
 	}
 }
 
+bool UEditorDataStorage::IterateChildren(
+	UE::Editor::DataStorage::FHierarchyHandle InHierarchyHandle,
+	UE::Editor::DataStorage::RowHandle Row,
+	TFunctionRef<bool(const ICoreProvider& Context, UE::Editor::DataStorage::RowHandle Child)> VisitFn) const
+{
+	if (const UE::Editor::DataStorage::FTedsHierarchyAccessInterface* AccessInterface = Environment->GetHierarchyRegistrar().GetAccessInterface(InHierarchyHandle))
+	{
+		return AccessInterface->IterateChildren(*this, Row, MoveTemp(VisitFn));
+	}
+	return false;
+}
+
 void UEditorDataStorage::RegisterTickGroup(
 	FName GroupName, UE::Editor::DataStorage::EQueryTickPhase Phase, FName BeforeGroup, FName AfterGroup, UE::Editor::DataStorage::EExecutionMode ExecutionMode)
 {
diff --git a/Engine/Plugins/Experimental/EditorDataStorage/Source/TedsCore/Private/TypedElementDatabase.h b/Engine/Plugins/Experimental/EditorDataStorage/Source/TedsCore/Private/TypedElementDatabase.h
--- a/Engine/Plugins/Experimental/EditorDataStorage/Source/TedsCore/Private/TypedElementDatabase.h
+++ b/Engine/Plugins/Experimental/EditorDataStorage/Source/TedsCore/Private/TypedElementDatabase.h
@@ -173,6 +173,10 @@ public:
 		UE::Editor::DataStorage::RowHandle Row, 
 		TFunction<void(const ICoreProvider& Context, UE::Editor::DataStorage::RowHandle Owner, UE::Editor::DataStorage::RowHandle Target)> VisitFn,
 		ETraversalOrder TraversalOrder = ETraversalOrder::PreOrder) const override;
+	UE_API virtual bool IterateChildren(
+		UE::Editor::DataStorage::FHierarchyHandle,
+		UE::Editor::DataStorage::RowHandle Row,
+		TFunctionRef<bool(const ICoreProvider& Context, UE::Editor::DataStorage::RowHandle Child)> VisitFn) const override;
 
 	UE_API void RegisterTickGroup(FName GroupName, UE::Editor::DataStorage::EQueryTickPhase Phase, FName BeforeGroup, FName AfterGroup, UE::Editor::DataStorage::EExecutionMode ExecutionMode);
 	UE_API void UnregisterTickGroup(FName GroupName, UE::Editor::DataStorage::EQueryTickPhase Phase);
diff --git a/Engine/Source/Runtime/TypedElementFramework/Public/Elements/Interfaces/TypedElementDataStorageInterface.h b/Engine/Source/Runtime/TypedElementFramework/Public/Elements/Interfaces/TypedElementDataStorageInterface.h
--- a/Engine/Source/Runtime/TypedElementFramework/Public/Elements/Interfaces/TypedElementDataStorageInterface.h
+++ b/Engine/Source/Runtime/TypedElementFramework/Public/Elements/Interfaces/TypedElementDataStorageInterface.h
@@ -454,6 +454,26 @@ public:
 		FHierarchyIterationCallback VisitFn,
 		ETraversalOrder TraversalOrder = ETraversalOrder::PreOrder) const = 0;
 	
+	/**
+	 * Iterates the immediate children of the given row without descending further
+	 *  down the hierarchy.
+	 * 
+	 * @param InHierarchyHandle Hierarchy to use
+	 * @param Row Parent row (not included in iteration)
+	 * @param VisitFn Iteration function that is expected to return true if iteration should continue
+	 *  or false if it should be stopped early.
+	 * @return true if iteration completed (including if there were no children) without any VisitFn call returning
+	 *  false. Returns false if VisitFn stopped iteration early or there was an error.
+	 */
+	UE_EXPERIMENTAL(5.8, TEXT("This function signature might change"))
+	virtual bool IterateChildren(
+		FHierarchyHandle InHierarchyHandle,
+		RowHandle Row,
+		TFunctionRef<bool(const ICoreProvider& Context, RowHandle Child)> VisitFn) const
+	{
+		return false;
+	}
+
 	/**
 	 * Outputs the registered query callbacks to the given output device for debugging purposes.
 	 */
