[TEDS] Fix for ICommonQueryContext::PushCommand not compiling for empty structs.

Currently, this fails compilation and this change would allow this to compile:
```
struct FSomeCommand
{
void operator()() {}
};
Context.PushCommand(FSomeCommand{}); // Error C2819 : type 'FSomeCommand' does not have an overloaded member 'operator ->'
```

#Synced-CL 52095318

diff --git a/Engine/Source/Runtime/TypedElementFramework/Public/Elements/Interfaces/TypedElementQueryStorageInterfaces.h b/Engine/Source/Runtime/TypedElementFramework/Public/Elements/Interfaces/TypedElementQueryStorageInterfaces.h
--- a/Engine/Source/Runtime/TypedElementFramework/Public/Elements/Interfaces/TypedElementQueryStorageInterfaces.h
+++ b/Engine/Source/Runtime/TypedElementFramework/Public/Elements/Interfaces/TypedElementQueryStorageInterfaces.h
@@ -652,7 +652,7 @@ namespace UE::Editor::DataStorage
 				[](void*)
 				{
 					T Instance;
-					Instance->operator()();
+					Instance.operator()();
 				}, nullptr);
 		}
 		else
