[TEDS] Updates to the capabilities mocks.

Mocks are automatically generated for query capabilities, but they were having problems with overloaded functions. To address this problem and to reduce the complexity of setting up a function, function are now set with a `Assign_*`, where * is the function name. Function overloading now works based on the callback function that's provided.

#Synced-CL 52147401

diff --git a/Engine/Source/Runtime/TypedElementFramework/Tests/Elements/Framework/TypedElementQueryContextMock.h b/Engine/Source/Runtime/TypedElementFramework/Tests/Elements/Framework/TypedElementQueryContextMock.h
--- a/Engine/Source/Runtime/TypedElementFramework/Tests/Elements/Framework/TypedElementQueryContextMock.h
+++ b/Engine/Source/Runtime/TypedElementFramework/Tests/Elements/Framework/TypedElementQueryContextMock.h
@@ -10,27 +10,118 @@
 
 namespace UE::Editor::DataStorage::Queries
 {
-	struct FMockChainTerminator {};
+	//
+	// Composite the environment. This will hold all the function implementations and the calls to set them.
+	//
 
-	struct FQueryContextMock final : public
+	struct FMockContextEnvironment
+	{
+#define Concat2(a, b) a##b
+#define Concat(a, b) Concat2(a, b)
 
-#define Function0(Capability, Return, Function)
-#define Function1(Capability, Return, Function, Arg1)
-#define Function2(Capability, Return, Function, Arg1, Arg2)
-#define Function3(Capability, Return, Function, Arg1, Arg2, Arg3)
-#define Function4(Capability, Return, Function, Arg1, Arg2, Arg3, Arg4)
+#define ArgTypeName(Type, Name) Type Name
+#define ArgName(Type, Name) Name
+#define ForwardArg(Type, Name) Forward< Type >( Name )
+#define MockFuncVar Concat(MockFunction_, __LINE__)
 
-#define ConstFunction0(Capability, Return, Function)
-#define ConstFunction1(Capability, Return, Function, Arg1)
-#define ConstFunction2(Capability, Return, Function, Arg1, Arg2)
-#define ConstFunction3(Capability, Return, Function, Arg1, Arg2, Arg3)
-#define ConstFunction4(Capability, Return, Function, Arg1, Arg2, Arg3, Arg4)
+#define CapabilityStart(Capability, Flags)\
+		friend struct F##Capability##Mock;
+
+#define Function0(Capability, Return, Function) \
+		private: \
+			TFunction<Return ()> MockFuncVar; \
+		public: \
+			void Assign_##Function(TFunction<Return ()> Callback) \
+			{ \
+				MockFuncVar = MoveTemp(Callback); \
+			}
+
+#define Function1(Capability, Return, Function, Arg1) \
+		private: \
+			TFunction<Return (ArgTypeName Arg1 )> MockFuncVar; \
+		public: \
+			void Assign_##Function(TFunction<Return (ArgTypeName Arg1 )> Callback) \
+			{ \
+				MockFuncVar = MoveTemp(Callback); \
+			}
+
+#define Function2(Capability, Return, Function, Arg1, Arg2) \
+		private: \
+			TFunction<Return (ArgTypeName Arg1, ArgTypeName Arg2 )> MockFuncVar; \
+		public: \
+			void Assign_##Function(TFunction<Return (ArgTypeName Arg1, ArgTypeName Arg2 )> Callback) \
+			{ \
+				MockFuncVar = MoveTemp(Callback); \
+			}
+
+#define Function3(Capability, Return, Function, Arg1, Arg2, Arg3) \
+		private: \
+			TFunction<Return (ArgTypeName Arg1, ArgTypeName Arg2, ArgTypeName Arg3 )> MockFuncVar; \
+		public: \
+			void Assign_##Function(TFunction<Return (ArgTypeName Arg1, ArgTypeName Arg2, ArgTypeName Arg3 )> Callback) \
+			{ \
+				MockFuncVar = MoveTemp(Callback); \
+			}
+
+#define Function4(Capability, Return, Function, Arg1, Arg2, Arg3, Arg4) \
+		private: \
+			TFunction<Return (ArgTypeName Arg1, ArgTypeName Arg2, ArgTypeName Arg3, ArgTypeName Arg4 )> MockFuncVar; \
+		public: \
+			void Assign_##Function(TFunction<Return (ArgTypeName Arg1, ArgTypeName Arg2, ArgTypeName Arg3, ArgTypeName Arg4 )> Callback) \
+			{ \
+				MockFuncVar = MoveTemp(Callback); \
+			}
+
+#define ConstFunction0(Capability, Return, Function) \
+		private: \
+			TFunction<Return ()> MockFuncVar; \
+		public: \
+			void Assign_##Function##_Const(TFunction<Return ()> Callback) \
+			{ \
+				MockFuncVar = MoveTemp(Callback); \
+			}
+
+#define ConstFunction1(Capability, Return, Function, Arg1) \
+		private: \
+			TFunction<Return (ArgTypeName Arg1 )> MockFuncVar; \
+		public: \
+			void Assign_##Function##_Const(TFunction<Return (ArgTypeName Arg1 )> Callback) \
+			{ \
+				MockFuncVar = MoveTemp(Callback); \
+			}
+
+#define ConstFunction2(Capability, Return, Function, Arg1, Arg2) \
+		private: \
+			TFunction<Return (ArgTypeName Arg1, ArgTypeName Arg2 )> MockFuncVar; \
+		public: \
+			void Assign_##Function##_Const(TFunction<Return (ArgTypeName Arg1, ArgTypeName Arg2 )> Callback) \
+			{ \
+				MockFuncVar = MoveTemp(Callback); \
+			}
+
+#define ConstFunction3(Capability, Return, Function, Arg1, Arg2, Arg3) \
+		private: \
+			TFunction<Return (ArgTypeName Arg1, ArgTypeName Arg2, ArgTypeName Arg3 )> MockFuncVar; \
+		public: \
+			void Assign_##Function##_Const(TFunction<Return (ArgTypeName Arg1, ArgTypeName Arg2, ArgTypeName Arg3 )> Callback) \
+			{ \
+				MockFuncVar = MoveTemp(Callback); \
+			}
+
+#define ConstFunction4(Capability, Return, Function, Arg1, Arg2, Arg3, Arg4) \
+		private: \
+			TFunction<Return (ArgTypeName Arg1, ArgTypeName Arg2, ArgTypeName Arg3, ArgTypeName Arg4 )> MockFuncVar; \
+		public: \
+			void Assign_##Function##_Const(TFunction<Return (ArgTypeName Arg1, ArgTypeName Arg2, ArgTypeName Arg3, ArgTypeName Arg4 )> Callback) \
+			{ \
+				MockFuncVar = MoveTemp(Callback); \
+			}
 
-#define CapabilityStart(Capability, Flags) Capability, 
 #define CapabilityEnd(Capability)
 
 #include "Elements/Framework/TypedElementQueryCapabilities.inl"
 
+#undef CapabilityStart
 #undef Function0
 #undef Function1
 #undef Function2
@@ -41,85 +132,188 @@ namespace UE::Editor::DataStorage::Queries
 #undef ConstFunction2
 #undef ConstFunction3
 #undef ConstFunction4
-#undef CapabilityStart
 #undef CapabilityEnd
+#undef Concat2
+#undef Concat
+#undef ArgTypeName
+#undef ArgName
+#undef ForwardArg
+#undef MockFunction
+#undef MockFuncVar
+	};
 
-		FMockChainTerminator
-	{
+	//
+	// Create a mock implementation for each capability. These will calls the functions users can provide in the environment.
+	//
 
-#define WithWrappers 1
+#define Concat2(a, b) a##b
+#define Concat(a, b) Concat2(a, b)
 
 #define ArgTypeName(Type, Name) Type Name
 #define ArgName(Type, Name) Name
-#define MockFunction(Capability, Function) Capability##_##Function##_Mock
+#define ForwardArg(Type, Name) Forward< Type >( Name )
+#define MockFuncVar Concat(Environment.MockFunction_, __LINE__)
+
+#define CapabilityStart(Capability, Flags) \
+		struct F##Capability##Mock : IContextCapability \
+		{ \
+			inline static const FName Name = #Capability;
+
+#define Function0(Capability, Return, Function) \
+			static Return Function (FMockContextEnvironment& Environment) \
+			{ \
+				checkf(MockFuncVar, TEXT( "Context function " #Function " in capability " #Capability " was unexpectedly called." )); \
+				return MockFuncVar ? MockFuncVar () : UE::MakeDefault< Return >(); \
+			}
+		
+#define Function1(Capability, Return, Function, Arg1) \
+			static Return Function (FMockContextEnvironment& Environment, ArgTypeName Arg1) \
+			{ \
+				checkf(MockFuncVar, TEXT( "Context function " #Function " in capability " #Capability " was unexpectedly called." )); \
+				return MockFuncVar ? MockFuncVar (ForwardArg Arg1) : UE::MakeDefault< Return >(); \
+			}
+
+#define Function2(Capability, Return, Function, Arg1, Arg2) \
+			static Return Function (FMockContextEnvironment& Environment, ArgTypeName Arg1, ArgTypeName Arg2) \
+			{ \
+				checkf(MockFuncVar, TEXT( "Context function " #Function " in capability " #Capability " was unexpectedly called." )); \
+				return MockFuncVar ? MockFuncVar(ForwardArg Arg1, ForwardArg Arg2) : UE::MakeDefault< Return >(); \
+			}
+
+#define Function3(Capability, Return, Function, Arg1, Arg2, Arg3) \
+			static Return Function (FMockContextEnvironment& Environment, ArgTypeName Arg1, ArgTypeName Arg2, ArgTypeName Arg3) \
+			{ \
+				checkf(MockFuncVar, TEXT( "Context function " #Function " in capability " #Capability " was unexpectedly called." )); \
+				return MockFuncVar ? MockFuncVar(ForwardArg Arg1, ForwardArg Arg2, ForwardArg Arg3) : UE::MakeDefault< Return >(); \
+			}
+
+#define Function4(Capability, Return, Function, Arg1, Arg2, Arg3, Arg4) \
+			static Return Function (FMockContextEnvironment& Environment, ArgTypeName Arg1, ArgTypeName Arg2, ArgTypeName Arg3, ArgTypeName Arg4) \
+			{ \
+				checkf(MockFuncVar, TEXT( "Context function " #Function " in capability " #Capability " was unexpectedly called." )); \
+				return MockFuncVar ? MockFuncVar(ForwardArg Arg1, ForwardArg Arg2, ForwardArg Arg3, ForwardArg Arg4) : UE::MakeDefault< Return >(); \
+			}
+
+#define ConstFunction0(Capability, Return, Function) \
+			static Return Function (const FMockContextEnvironment& Environment) \
+			{ \
+				checkf(MockFuncVar, TEXT( "Context function " #Function " in capability " #Capability " was unexpectedly called." )); \
+				return MockFuncVar ? MockFuncVar () : UE::MakeDefault< Return >(); \
+			}
+
+#define ConstFunction1(Capability, Return, Function, Arg1) \
+			static Return Function (const FMockContextEnvironment& Environment, ArgTypeName Arg1) \
+			{ \
+				checkf(MockFuncVar, TEXT( "Context function " #Function " in capability " #Capability " was unexpectedly called." )); \
+				return MockFuncVar ? MockFuncVar (ForwardArg Arg1) : UE::MakeDefault< Return >(); \
+			}
+
+#define ConstFunction2(Capability, Return, Function, Arg1, Arg2) \
+			static Return Function (const FMockContextEnvironment& Environment, ArgTypeName Arg1, ArgTypeName Arg2) \
+			{ \
+				checkf(MockFuncVar, TEXT( "Context function " #Function " in capability " #Capability " was unexpectedly called." )); \
+				return MockFuncVar ? MockFuncVar(ForwardArg Arg1, ForwardArg Arg2) : UE::MakeDefault< Return >(); \
+			}
+
+#define ConstFunction3(Capability, Return, Function, Arg1, Arg2, Arg3) \
+			static Return Function (const FMockContextEnvironment& Environment, ArgTypeName Arg1, ArgTypeName Arg2, ArgTypeName Arg3) \
+			{ \
+				checkf(MockFuncVar, TEXT( "Context function " #Function " in capability " #Capability " was unexpectedly called." )); \
+				return MockFuncVar ? MockFuncVar(ForwardArg Arg1, ForwardArg Arg2, ForwardArg Arg3) : UE::MakeDefault< Return >(); \
+			}
+
+#define ConstFunction4(Capability, Return, Function, Arg1, Arg2, Arg3, Arg4) \
+			static Return Function (const FMockContextEnvironment& Environment, ArgTypeName Arg1, ArgTypeName Arg2, ArgTypeName Arg3, ArgTypeName Arg4) \
+			{ \
+				checkf(MockFuncVar, TEXT( "Context function " #Function " in capability " #Capability " was unexpectedly called." )); \
+				return MockFuncVar ? MockFuncVar(ForwardArg Arg1, ForwardArg Arg2, ForwardArg Arg3, ForwardArg Arg4) : UE::MakeDefault< Return >(); \
+			}
+
+#define CapabilityEnd(Capability) \
+		};
+
+#include "Elements/Framework/TypedElementQueryCapabilities.inl"
+
+#undef CapabilityStart
+#undef Function0
+#undef Function1
+#undef Function2
+#undef Function3
+#undef Function4
+#undef ConstFunction0
+#undef ConstFunction1
+#undef ConstFunction2
+#undef ConstFunction3
+#undef ConstFunction4
+#undef CapabilityEnd
+#undef Concat2
+#undef Concat
+#undef ArgTypeName
+#undef ArgName
+#undef ForwardArg
+#undef MockFunction
+#undef MockFuncVar
+
+	//
+	// Temporary forwarder that bridges FMockContextEnvironment with the capability interfaces.
+	//
+
+#define ArgTypeName(Type, Name) Type Name
+#define ArgName(Type, Name) Name
+#define CapabilityStart(Capability, Flags)
+#define CapabilityEnd(Capability)
 
 #define Function0(Capability, Return, Function) \
-		TFunction<Return ()> MockFunction(Capability, Function); \
 		virtual Return Function() override \
-		{ \
-			return MockFunction(Capability, Function)(); \
-		}
+		{ return F##Capability##Mock::Function(*this); }
+
 #define Function1(Capability, Return, Function, Arg1) \
-		TFunction<Return (ArgTypeName Arg1 )> MockFunction(Capability, Function); \
 		virtual Return Function(ArgTypeName Arg1) override \
-		{ \
-			return MockFunction(Capability, Function)(ArgName Arg1); \
-		}
+		{ return F##Capability##Mock::Function(*this, ArgName Arg1); }
+
 #define Function2(Capability, Return, Function, Arg1, Arg2) \
-		TFunction<Return (ArgTypeName Arg1, ArgTypeName Arg2)> MockFunction(Capability, Function); \
 		virtual Return Function(ArgTypeName Arg1, ArgTypeName Arg2) override \
-		{ \
-			return MockFunction(Capability, Function)(ArgName Arg1, ArgName Arg2); \
-		}
+		{ return F##Capability##Mock::Function(*this, ArgName Arg1, ArgName Arg2); }
+
 #define Function3(Capability, Return, Function, Arg1, Arg2, Arg3) \
-		TFunction<Return (ArgTypeName Arg1, ArgTypeName Arg2, ArgTypeName Arg3)> MockFunction(Capability, Function); \
 		virtual Return Function(ArgTypeName Arg1, ArgTypeName Arg2, ArgTypeName Arg3) override \
-		{ \
-			return MockFunction(Capability, Function)(ArgName Arg1, ArgName Arg2, ArgName Arg3); \
-		}
+		{ return F##Capability##Mock::Function(*this, ArgName Arg1, ArgName Arg2, ArgName Arg3); }
+
 #define Function4(Capability, Return, Function, Arg1, Arg2, Arg3, Arg4) \
-		TFunction<Return (ArgTypeName Arg1, ArgTypeName Arg2, ArgTypeName Arg3, ArgTypeName Arg4)> MockFunction(Capability, Function); \
 		virtual Return Function(ArgTypeName Arg1, ArgTypeName Arg2, ArgTypeName Arg3, ArgTypeName Arg4) override \
-		{ \
-			return MockFunction(Capability, Function)(ArgName Arg1, ArgName Arg2, ArgName Arg3, ArgName Arg4); \
-		}
+		{ return F##Capability##Mock::Function(*this, ArgName Arg1, ArgName Arg2, ArgName Arg3, ArgName Arg4); }
 
 #define ConstFunction0(Capability, Return, Function) \
-		TFunction<Return ()> MockFunction(Capability, Function); \
 		virtual Return Function() const override \
-		{ \
-			return MockFunction(Capability, Function)(); \
-		}
+		{ return F##Capability##Mock::Function(*this); }
+
 #define ConstFunction1(Capability, Return, Function, Arg1) \
-		TFunction<Return (ArgTypeName Arg1)> MockFunction(Capability, Function); \
 		virtual Return Function(ArgTypeName Arg1) const override \
-		{ \
-			return MockFunction(Capability, Function)(ArgName Arg1); \
-		}
+		{ return F##Capability##Mock::Function(*this, ArgName Arg1); }
+
 #define ConstFunction2(Capability, Return, Function, Arg1, Arg2) \
-		TFunction<Return (ArgTypeName Arg1, ArgTypeName Arg2)> MockFunction(Capability, Function); \
 		virtual Return Function(ArgTypeName Arg1, ArgTypeName Arg2) const override \
-		{ \
-			return MockFunction(Capability, Function)(ArgName Arg1, ArgName Arg2); \
-		}
+		{ return F##Capability##Mock::Function(*this, ArgName Arg1, ArgName Arg2); }
+
 #define ConstFunction3(Capability, Return, Function, Arg1, Arg2, Arg3) \
-		TFunction<Return (ArgTypeName Arg1, ArgTypeName Arg2, ArgTypeName Arg3)> MockFunction(Capability, Function); \
 		virtual Return Function(ArgTypeName Arg1, ArgTypeName Arg2, ArgTypeName Arg3) const override \
-		{ \
-			return MockFunction(Capability, Function)(ArgName Arg1, ArgName Arg2, ArgName Arg3); \
-		}
+		{ return F##Capability##Mock::Function(*this, ArgName Arg1, ArgName Arg2, ArgName Arg3); }
+
 #define ConstFunction4(Capability, Return, Function, Arg1, Arg2, Arg3, Arg4) \
-		TFunction<Return (ArgTypeName Arg1, ArgTypeName Arg2, ArgTypeName Arg3, ArgTypeName Arg4)> MockFunction(Capability, Function); \
 		virtual Return Function(ArgTypeName Arg1, ArgTypeName Arg2, ArgTypeName Arg3, ArgTypeName Arg4) const override \
-		{ \
-			return MockFunction(Capability, Function)(ArgName Arg1, ArgName Arg2, ArgName Arg3, ArgName Arg4); \
-		}
-
-#define CapabilityStart(Capability, Flags) 
-#define CapabilityEnd(Capability)
+		{ return F##Capability##Mock::Function(*this, ArgName Arg1, ArgName Arg2, ArgName Arg3, ArgName Arg4); }
 
+	struct FMockContextForwarder
+		: FMockContextEnvironment
+		, ISingleRowInfo<IContextCapability>
+		, IRowBatchInfo<IContextCapability>
+		, IDynamicColumnInfo<IContextCapability>
+	{
 #include "Elements/Framework/TypedElementQueryCapabilities.inl"
+	};
 
+#undef CapabilityStart
+#undef CapabilityEnd
 #undef Function0
 #undef Function1
 #undef Function2
@@ -130,36 +324,33 @@ namespace UE::Editor::DataStorage::Queries
 #undef ConstFunction2
 #undef ConstFunction3
 #undef ConstFunction4
-#undef WithWrappers
+#undef ArgTypeName
+#undef ArgName
 
-#define MockFunctionImplementation(Capability, Return, Function) \
-		checkf(false, TEXT( "Context function " #Function " in capability " #Capability " was unexpectedly called." )); \
-		return UE::MakeDefault< Return >();
-#define Function0(Capability, Return, Function) \
-		MockFunction(Capability, Function) = []() { MockFunctionImplementation(Capability, Return, Function) };
-#define Function1(Capability, Return, Function, Arg1) \
-		MockFunction(Capability, Function) = [](ArgTypeName Arg1) { MockFunctionImplementation(Capability, Return, Function) };
-#define Function2(Capability, Return, Function, Arg1, Arg2) \
-		MockFunction(Capability, Function) = [](ArgTypeName Arg1, ArgTypeName Arg2) { MockFunctionImplementation(Capability, Return, Function) };
-#define Function3(Capability, Return, Function, Arg1, Arg2, Arg3) \
-		MockFunction(Capability, Function) = [](ArgTypeName Arg1, ArgTypeName Arg2, ArgTypeName Arg3) { MockFunctionImplementation(Capability, Return, Function) };
-#define Function4(Capability, Return, Function, Arg1, Arg2, Arg3, Arg4) \
-		MockFunction(Capability, Function) = [](ArgTypeName Arg1, ArgTypeName Arg2, ArgTypeName Arg3, ArgTypeName Arg4) { MockFunctionImplementation(Capability, Return, Function) };
+	//
+	// Combine all mock implementations together into a final implementation.
+	//
+
+using QueryContextMock = TQueryContextImpl<FMockContextForwarder
 
-#define ConstFunction0(Capability, Return, Function) Function0(Capability, Return, Function)
-#define ConstFunction1(Capability, Return, Function, Arg1) Function1(Capability, Return, Function, Arg1)
-#define ConstFunction2(Capability, Return, Function, Arg1, Arg2) Function2(Capability, Return, Function, Arg1, Arg2)
-#define ConstFunction3(Capability, Return, Function, Arg1, Arg2, Arg3) Function3(Capability, Return, Function, Arg1, Arg2, Arg3)
-#define ConstFunction4(Capability, Return, Function, Arg1, Arg2, Arg3, Arg4) Function4(Capability, Return, Function, Arg1, Arg2, Arg3, Arg4)
+#define Function0(Capability, Return, Function)
+#define Function1(Capability, Return, Function, Arg1)
+#define Function2(Capability, Return, Function, Arg1, Arg2)
+#define Function3(Capability, Return, Function, Arg1, Arg2, Arg3)
+#define Function4(Capability, Return, Function, Arg1, Arg2, Arg3, Arg4)
+
+#define ConstFunction0(Capability, Return, Function)
+#define ConstFunction1(Capability, Return, Function, Arg1)
+#define ConstFunction2(Capability, Return, Function, Arg1, Arg2)
+#define ConstFunction3(Capability, Return, Function, Arg1, Arg2, Arg3)
+#define ConstFunction4(Capability, Return, Function, Arg1, Arg2, Arg3, Arg4)
+
+#define CapabilityStart(Capability, Flags) , Capability
+#define CapabilityEnd(Capability)
 
-		FQueryContextMock()
-		{
 #include "Elements/Framework/TypedElementQueryCapabilities.inl"
-		}
+	>;
 
-#undef ArgTypeName
-#undef ArgName
-#undef MockFunctionName
 #undef Function0
 #undef Function1
 #undef Function2
@@ -172,8 +363,5 @@ namespace UE::Editor::DataStorage::Queries
 #undef ConstFunction4
 #undef CapabilityStart
 #undef CapabilityEnd
-#undef WithWrappers
-	};
 
-	using QueryContextMock = TQueryContextImpl<FQueryContextMock>;
 } // namespace UE::Editor::DataStorage::Queries
diff --git a/Engine/Source/Runtime/TypedElementFramework/Tests/Elements/Framework/TypedElementQueryFunctionsTests.cpp b/Engine/Source/Runtime/TypedElementFramework/Tests/Elements/Framework/TypedElementQueryFunctionsTests.cpp
--- a/Engine/Source/Runtime/TypedElementFramework/Tests/Elements/Framework/TypedElementQueryFunctionsTests.cpp
+++ b/Engine/Source/Runtime/TypedElementFramework/Tests/Elements/Framework/TypedElementQueryFunctionsTests.cpp
@@ -6,6 +6,7 @@
 #include "Elements/Framework/TypedElementQueryContext.h"
 #include "Elements/Framework/TypedElementQueryContextMock.h"
 #include "Elements/Framework/TypedElementQueryFunctions.h"
+#include "Elements/Framework/TypedElementRowHandleArrayView.h"
 #include "Elements/Framework/TypedElementTestColumns.h"
 #include "Tests/TestHarnessAdapter.h"
 
@@ -47,41 +48,56 @@ namespace UE::Editor::DataStorage::Queries::Tests
 
 		virtual bool NextBatch() override
 		{
-			CHECK_MESSAGE(TEXT("Row count was not set in test."), RowCount > 0);
 			return false;
 		}
 
 		virtual bool NextRow() override
 		{
-			CHECK_MESSAGE(TEXT("Row count was not set in test."), RowCount > 0);
-			CurrentRow++;
-			return CurrentRow < RowCount;
+			++CurrentRow;
+			return CurrentRow < static_cast<RowHandle>(FakeRows.Num());
 		}
 
 		virtual void GetConstColumns(TArrayView<const void*> ColumnsData, TConstArrayView<const UScriptStruct*> ColumnTypes) override
 		{
-			CHECK_MESSAGE(TEXT("Row count was not set in test."), RowCount > 0);
+			CHECK_MESSAGE(TEXT("Row count was not set in test."), !FakeRows.IsEmpty());
 			for (int32 Index = 0, Num = ColumnTypes.Num(); Index < Num; ++Index)
 			{
-				const void** Data = ConstColumns.Find(ColumnTypes[Index]);
-				ColumnsData[Index] = Data ? (static_cast<const char*>(*Data) + ColumnTypes[Index]->GetStructureSize() * CurrentRow) : nullptr;
+				const void* const* Data = ConstColumns.Find(ColumnTypes[Index]);
+				if (Data == nullptr)
+				{
+					Data = MutableColumns.Find(ColumnTypes[Index]);
+				}
+				CHECK_MESSAGE(TEXT("Unable to retrieve const column data."), Data != nullptr && *Data != nullptr);
+				ColumnsData[Index] = *Data;
 			}
 		}
 
 		virtual void GetMutableColumns(TArrayView<void*> ColumnsData, TConstArrayView<const UScriptStruct*> ColumnTypes) override
 		{
-			CHECK_MESSAGE(TEXT("Row count was not set in test."), RowCount > 0);
+			CHECK_MESSAGE(TEXT("Row count was not set in test."), !FakeRows.IsEmpty());
 			for (int32 Index = 0, Num = ColumnTypes.Num(); Index < Num; ++Index)
 			{
 				void** Data = MutableColumns.Find(ColumnTypes[Index]);
-				ColumnsData[Index] = Data ? (static_cast<char*>(*Data) + ColumnTypes[Index]->GetStructureSize() * CurrentRow) : nullptr;
+				CHECK_MESSAGE(TEXT("Unable to retrieve mutable column data."), Data != nullptr && *Data != nullptr);
+				ColumnsData[Index] = *Data;
 			}
 		}
 
+		void SetCurrentRow(RowHandle Row)
+		{
+			CurrentRow = Row;
+		}
+
 		void SetRowCount(int32 InRowCount)
 		{
 			CHECK_MESSAGE(TEXT("Row count was not set in test."), InRowCount > 0);
-			RowCount = InRowCount;
+			FakeRows.Reset();
+			FakeRows.Reserve(InRowCount);
+			for (RowHandle Row = 0; Row < InRowCount; ++Row)
+			{
+				FakeRows.Add(Row);
+			}
+			CurrentRow = FakeRows[0];
 		}
 
 		template<TColumnType... ColumnType>
@@ -98,8 +114,8 @@ namespace UE::Editor::DataStorage::Queries::Tests
 
 		TMap<const UScriptStruct*, const void*> ConstColumns;
 		TMap<const UScriptStruct*, void*> MutableColumns;
-		int32 CurrentRow = 0;
-		int32 RowCount = 0;
+		TArray<RowHandle> FakeRows;
+		RowHandle CurrentRow = InvalidRowHandle;
 	};
 
 	TEST_CASE_NAMED(TQueryFunction_Tests, "Editor::DataStorage::Queries::Query Function (TQueryFunction)", "[ApplicationContextMask][EngineFilter]")
@@ -219,15 +235,6 @@ namespace UE::Editor::DataStorage::Queries::Tests
 			CHECK_MESSAGE(TEXT("Function with context incorrectly found to not be matching."), bResult);
 		}
 
-		SECTION("Check context compatibility with restricted context")
-		{
-			TQueryFunction<void> Function = BuildQueryFunction<void>([](TQueryContext<SingleRowInfo> Context) {});
-
-			TQueryContextImpl<FQueryContextMock, SingleRowInfo> Context;
-			bool bResult = Context.CheckCompatiblity(Function);
-			CHECK_MESSAGE(TEXT("Function with context incorrectly found to not be matching."), bResult);
-		}
-
 		SECTION("Call without arguments")
 		{
 			TQueryFunction<bool> Function = BuildQueryFunction<bool>(
@@ -239,7 +246,7 @@ namespace UE::Editor::DataStorage::Queries::Tests
 			QueryContextMock Context;
 			FTestQueryFunctionResponse Response;
 			Response.SetRowCount(1);
-			
+
 			FBoolOrResultAccumulator Result;
 			Function.Call<EFunctionCallConfig::VerifyColumns>(Result, Context, Response);
 			CHECK_MESSAGE(TEXT("Query function wasn't called."), Result.Value);
@@ -255,6 +262,7 @@ namespace UE::Editor::DataStorage::Queries::Tests
 
 			QueryContextMock Context;
 			FTestColumnString Column;
+			Column.TestString = TEXT("Not set");
 			
 			FTestQueryFunctionResponse Response;
 			Response.SetRowCount(1);
@@ -299,11 +307,12 @@ namespace UE::Editor::DataStorage::Queries::Tests
 
 			static constexpr uint32 RowCount = 2;
 			TStaticArray<RowHandle, 2> Rows = { 0, 1 };
-			TConstArrayView<RowHandle> RowsView(Rows.GetData(), Rows.Num());
+			FRowHandleArrayView RowsView(Rows.GetData(), Rows.Num(), 
+				FRowHandleArrayView::EFlags::IsSorted | FRowHandleArrayView::EFlags::IsUnique);
 
 			QueryContextMock Context;
-			Context.GetContextImplementation().RowBatchInfo_GetBatchRowCount_Mock = [](){ return RowCount; };
-			Context.GetContextImplementation().RowBatchInfo_GetBatchRowHandles_Mock = [RowsView](){ return RowsView; };
+			Context.GetContextImplementation().Assign_GetBatchRowCount_Const([]() { return RowCount; });
+			Context.GetContextImplementation().Assign_GetBatchRowHandles_Const([RowsView](){ return RowsView; });
 			FTestColumnString ColumnsA[RowCount];
 			FTestColumnInt ColumnsB[] =
 			{
@@ -331,22 +340,22 @@ namespace UE::Editor::DataStorage::Queries::Tests
 				});
 
 			QueryContextMock Context;
-			
+
 			bool bGetCurrentRowCalled = false;
-			Context.GetContextImplementation().SingleRowInfo_GetCurrentRow_Mock =
+			Context.GetContextImplementation().Assign_GetCurrentRow_Const(
 				[&bGetCurrentRowCalled]()
 				{
 					bGetCurrentRowCalled = true;
 					return 0;
-				};
-			
+				});
+
 			bool bCurrentRowHasColumnsCalled = false;
-			Context.GetContextImplementation().SingleRowInfo_CurrentRowHasColumns_Mock =
+			Context.GetContextImplementation().Assign_CurrentRowHasColumns_Const(
 				[&bCurrentRowHasColumnsCalled](TConstArrayView<const UScriptStruct*>)
 				{
 					bCurrentRowHasColumnsCalled = true;
 					return false;
-				};
+				});
 
 			FTestQueryFunctionResponse Response;
 			Response.SetRowCount(1);
@@ -356,39 +365,35 @@ namespace UE::Editor::DataStorage::Queries::Tests
 			CHECK_MESSAGE(TEXT("CurrentRowHasColumns not called on context."), bCurrentRowHasColumnsCalled);
 		}
 
-		SECTION("Call and use restricted context")
+		SECTION("Call and check provided row handles")
 		{
+			constexpr int32 CheckRange = 5;
+
+			TArray<RowHandle> RowCheck;
+			RowCheck.Reserve(CheckRange);
 			TQueryFunction<void> Function = BuildQueryFunction<void>(
-				[](TQueryContext<SingleRowInfo> Context)
+				[&RowCheck](TQueryContext<SingleRowInfo> Context)
 				{
-					Context.GetCurrentRow();
-					Context.CurrentRowHasColumns<FTestColumnString>();
+					RowCheck.Add(Context.GetCurrentRow());
 				});
 
-			TQueryContextImpl<FQueryContextMock, SingleRowInfo> Context;
-
-			bool bGetCurrentRowCalled = false;
-			Context.GetContextImplementation().SingleRowInfo_GetCurrentRow_Mock =
-				[&bGetCurrentRowCalled]()
-				{
-					bGetCurrentRowCalled = true;
-					return 0;
-				};
+			QueryContextMock Context;
+			FTestQueryFunctionResponse Response;
 
-			bool bCurrentRowHasColumnsCalled = false;
-			Context.GetContextImplementation().SingleRowInfo_CurrentRowHasColumns_Mock =
-				[&bCurrentRowHasColumnsCalled](TConstArrayView<const UScriptStruct*>)
+			Context.GetContextImplementation().Assign_GetCurrentRow_Const(
+				[&Response]()
 				{
-					bCurrentRowHasColumnsCalled = true;
-					return false;
-				};
+					return Response.CurrentRow;
+				});
 
-			FTestQueryFunctionResponse Response;
-			Response.SetRowCount(1);
+			Response.SetRowCount(CheckRange);
 			Function.Call(Context, Response);
 
-			CHECK_MESSAGE(TEXT("GetCurrentRow not called on context."), bGetCurrentRowCalled);
-			CHECK_MESSAGE(TEXT("CurrentRowHasColumns not called on context."), bCurrentRowHasColumnsCalled);
+			CHECK_MESSAGE(TEXT("Not all Rows were added."), RowCheck.Num() == CheckRange);
+			for (int32 FakeRow = 0; FakeRow < CheckRange; ++FakeRow)
+			{
+				CHECK_MESSAGE(TEXT("One of the rows wasn't called or called with the wrong row handle."), RowCheck[FakeRow] == FakeRow);
+			}
 		}
 
 		SECTION("Call with result from callback")
@@ -411,12 +416,13 @@ namespace UE::Editor::DataStorage::Queries::Tests
 		SECTION("Call with result through TResult")
 		{
 			TQueryFunction<bool> Function = BuildQueryFunction<bool>(
-				[](TResult<bool>& Result)
+				[](TQueryContext<SingleRowInfo> Context, TResult<bool>& Result)
 				{ 
 					Result.Add(true);
 				});
 
 			QueryContextMock Context;
+			Context.GetContextImplementation().Assign_GetCurrentRow_Const([]() { return RowHandle(1); });
 			FTestColumnString Column;
 
 			FBoolOrResultAccumulator Result;
@@ -442,11 +448,12 @@ namespace UE::Editor::DataStorage::Queries::Tests
 
 			static constexpr int32 RowCount = 4;
 			TStaticArray<RowHandle, 4> Rows = { 0, 1, 2, 3 };
-			TConstArrayView<RowHandle> RowsView(Rows.GetData(), Rows.Num());
+			FRowHandleArrayView RowsView(Rows.GetData(), Rows.Num(),
+				FRowHandleArrayView::EFlags::IsSorted | FRowHandleArrayView::EFlags::IsUnique);
 
 			QueryContextMock Context;
-			Context.GetContextImplementation().RowBatchInfo_GetBatchRowCount_Mock = []() { return RowCount; };
-			Context.GetContextImplementation().RowBatchInfo_GetBatchRowHandles_Mock = [RowsView](){ return RowsView; };
+			Context.GetContextImplementation().Assign_GetBatchRowCount_Const([]() { return RowCount; });
+			Context.GetContextImplementation().Assign_GetBatchRowHandles_Const([RowsView]() { return RowsView; });
 
 			FTestColumnInt Columns[] =
 			{
