[TEDS] static analysis fix

#rnx
#Synced-CL 52147401

diff --git a/Engine/Source/Runtime/TypedElementFramework/Tests/Elements/Framework/TypedElementQueryFunctionsTests.cpp b/Engine/Source/Runtime/TypedElementFramework/Tests/Elements/Framework/TypedElementQueryFunctionsTests.cpp
--- a/Engine/Source/Runtime/TypedElementFramework/Tests/Elements/Framework/TypedElementQueryFunctionsTests.cpp
+++ b/Engine/Source/Runtime/TypedElementFramework/Tests/Elements/Framework/TypedElementQueryFunctionsTests.cpp
@@ -68,7 +68,14 @@ namespace UE::Editor::DataStorage::Queries::Tests
 					Data = MutableColumns.Find(ColumnTypes[Index]);
 				}
 				CHECK_MESSAGE(TEXT("Unable to retrieve const column data."), Data != nullptr && *Data != nullptr);
-				ColumnsData[Index] = *Data;
+				if (Data != nullptr && *Data != nullptr)
+				{
+					ColumnsData[Index] = *Data;
+				}
+				else
+				{
+					ColumnsData[Index] = nullptr;
+				}
 			}
 		}
 
@@ -79,7 +86,14 @@ namespace UE::Editor::DataStorage::Queries::Tests
 			{
 				void** Data = MutableColumns.Find(ColumnTypes[Index]);
 				CHECK_MESSAGE(TEXT("Unable to retrieve mutable column data."), Data != nullptr && *Data != nullptr);
-				ColumnsData[Index] = *Data;
+				if (Data != nullptr && *Data != nullptr)
+				{
+					ColumnsData[Index] = *Data;
+				}
+				else
+				{
+					ColumnsData[Index] = nullptr;
+				}
 			}
 		}
 
