[TEDS] Editor Scope Tree system + editor integration + TEDS debugger tab

Adds the Editor Scope Tree system for propagating editing scope
through the editor widget hierarchy via TEDS scope rows and TLS.

Scope Core (ICoreProvider + FScopeSystem):
- EditorDataScope.h: TLS-based current scope, FPushScopeGuard RAII guard
- EditorDataScopeTypes.h: FScopeDataVersion for change detection
- EditorDataScopeColumns.h: FEditingVerseScope (editing scope verse path),
FScopeRowSourceInfo (creation callstack capture for debugging)
- FScopeSystem: Manager class in FEnvironment, implements hierarchy-walking
Get/Set/Delete for scope data, version tracking
- ICoreProvider: Virtual methods for scope operations (AddScopeRow,
GetScopeData, SetScopeData, DeleteScopeData, etc.)

Editor Integration:
- LevelEditor: Creates root + level editor scope rows, pushes onto TLS
during Tick and widget construction, sets FEditingVerseScope from world
- SubobjectEditor: Creates child scope rows parented via GetCurrentScope()

TEDS Debugger:
- SScopeHierarchyTab: New tab showing the scope row hierarchy tree
with creation callstack, visible columns, and data values

#Synced-CL 52185938

diff --git a/Engine/Plugins/Experimental/EditorDataStorage/Source/TedsCore/Private/Scope/TypedElementScopeSystem.cpp b/Engine/Plugins/Experimental/EditorDataStorage/Source/TedsCore/Private/Scope/TypedElementScopeSystem.cpp
new file mode 100644
--- /dev/null
+++ b/Engine/Plugins/Experimental/EditorDataStorage/Source/TedsCore/Private/Scope/TypedElementScopeSystem.cpp
@@ -0,0 +1,230 @@
+// Copyright Epic Games, Inc. All Rights Reserved.
+
+#include "Scope/TypedElementScopeSystem.h"
+
+#include "DataStorage/Scope/EditorDataScope.h"
+#include "DataStorage/Scope/EditorDataScopeColumns.h"
+#include "DataStorage/Scope/EditorDataScopeVersion.h"
+#include "Elements/Columns/TypedElementLabelColumns.h"
+#include "Elements/Interfaces/TypedElementDataStorageInterface.h"
+#include "HAL/PlatformStackWalk.h"
+
+namespace UE::Editor::DataStorage
+{
+
+static const FName GScopeHierarchyName = TEXT("EditorScope");
+static const FName GScopeTableName = TEXT("EditorScopeTable");
+
+FScopeSystem::FScopeSystem(ICoreProvider& InDataStorage)
+	: DataStorage(InDataStorage)
+{
+}
+
+TableHandle FScopeSystem::GetOrCreateScopeTable()
+{
+	TableHandle Table = DataStorage.FindTable(GScopeTableName);
+	if (Table == InvalidTableHandle)
+	{
+		Table = DataStorage.RegisterTable<FDataStorageScopeDataTag, FScopeRowSourceInfo, FScopeDataVersionColumn, FTypedElementLabelColumn>(
+			GScopeTableName);
+	}
+	return Table;
+}
+
+// ============================================================================
+// Scope Row Lifecycle & Hierarchy
+// ============================================================================
+
+RowHandle FScopeSystem::AddScopeRow()
+{
+	return AddScopeRow(FStringView());
+}
+
+RowHandle FScopeSystem::AddScopeRow(FStringView Label)
+{
+	RowHandle Row = DataStorage.AddRow(GetOrCreateScopeTable());
+
+	// Capture callstack at creation site for debugger visibility
+	FScopeRowSourceInfo SourceInfo;
+	SourceInfo.CallstackDepth = FPlatformStackWalk::CaptureStackBackTrace(
+		SourceInfo.Callstack, FScopeRowSourceInfo::MaxCallstackDepth);
+	DataStorage.SetScopeData<FScopeRowSourceInfo>(Row, MoveTemp(SourceInfo));
+
+	const FString LabelString = Label.IsEmpty() ? LexToString(Row) : FString(Label);
+	SetScopeData(Row, FTypedElementLabelColumn::StaticStruct(),
+		[&LabelString](void* ColumnData, const UScriptStruct&)
+		{
+			new(ColumnData) FTypedElementLabelColumn{ .Label = LabelString };
+		},
+		[](const UScriptStruct&, void* Destination, void* Source)
+		{
+			*static_cast<FTypedElementLabelColumn*>(Destination) = MoveTemp(*static_cast<FTypedElementLabelColumn*>(Source));
+		});
+
+	return Row;
+}
+
+RowHandle FScopeSystem::AddScopeRow(FStringView Label, RowHandle Parent)
+{
+	RowHandle Row = AddScopeRow(Label);
+	SetParentScope(Row, Parent);
+	return Row;
+}
+
+void FScopeSystem::RemoveScopeRow(RowHandle Row)
+{
+	DataStorage.RemoveRow(Row);
+}
+
+void FScopeSystem::SetParentScope(RowHandle Child, RowHandle Parent)
+{
+	DataStorage.SetParentRow(GetScopeHierarchy(), Child, Parent);
+}
+
+RowHandle FScopeSystem::GetParentScope(RowHandle Row)
+{
+	return DataStorage.GetParentRow(GetScopeHierarchy(), Row);
+}
+
+// ============================================================================
+// Scope Data Lookup (hierarchy-walking)
+// ============================================================================
+
+const void* FScopeSystem::GetScopeDataRaw(RowHandle Row, const UScriptStruct* ColumnType)
+{
+	RowHandle Current = Row;
+	while (DataStorage.IsRowAvailable(Current))
+	{
+		const void* Data = DataStorage.GetColumnData(Current, ColumnType);
+		if (Data != nullptr)
+		{
+			return Data;
+		}
+		Current = GetParentScope(Current);
+	}
+
+	return nullptr;
+}
+
+// ============================================================================
+// Visible Scope Columns (hierarchy-walking)
+// ============================================================================
+
+TArray<const UScriptStruct*> FScopeSystem::GetAllVisibleScopeColumns(RowHandle Row)
+{
+	TArray<const UScriptStruct*> Result;
+	TSet<const UScriptStruct*> SeenTypes;
+
+	RowHandle Current = Row;
+	while (DataStorage.IsRowAvailable(Current))
+	{
+		DataStorage.ListColumns(Current, [&](const UScriptStruct& ColumnType)
+		{
+			if (!SeenTypes.Contains(&ColumnType))
+			{
+				SeenTypes.Add(&ColumnType);
+				Result.Add(&ColumnType);
+			}
+		});
+		Current = GetParentScope(Current);
+	}
+	return Result;
+}
+
+// ============================================================================
+// Versioned Scope Data Operations (single-row)
+// ============================================================================
+
+Scope::FScopeDataVersion FScopeSystem::SetScopeData(
+	RowHandle Row, const UScriptStruct* ColumnType,
+	ColumnCreationCallbackRef Initializer, ColumnCopyOrMoveCallback Relocator)
+{
+	DataStorage.AddColumnData(Row, ColumnType, Initializer, Relocator);
+
+	FScopeDataVersionColumn* VersionCol = DataStorage.GetColumn<FScopeDataVersionColumn>(Row);
+	if (!VersionCol)
+	{
+		return Scope::FScopeDataVersion();
+	}
+
+	return Scope::FScopeDataVersion::Make(Row, ++VersionCol->Version);
+}
+
+bool FScopeSystem::RemoveScopeData(RowHandle Row, const UScriptStruct* ColumnType)
+{
+	if (!DataStorage.HasColumns(Row, TConstArrayView<const UScriptStruct*>(&ColumnType, 1)))
+	{
+		return false;
+	}
+
+	DataStorage.RemoveColumn(Row, ColumnType);
+
+	if (FScopeDataVersionColumn* VersionCol = DataStorage.GetColumn<FScopeDataVersionColumn>(Row))
+	{
+		++VersionCol->Version;
+	}
+
+	return true;
+}
+
+// ============================================================================
+// Versioned Hierarchy Walk
+// ============================================================================
+
+Scope::FScopeDataVersion FScopeSystem::GetScopeDataVersion(
+	RowHandle Row, const UScriptStruct* ColumnType)
+{
+	RowHandle Current = Row;
+	while (DataStorage.IsRowAvailable(Current))
+	{
+		if (DataStorage.GetColumnData(Current, ColumnType) != nullptr)
+		{
+			if (const FScopeDataVersionColumn* VersionCol =
+					DataStorage.GetColumn<FScopeDataVersionColumn>(Current))
+			{
+				return Scope::FScopeDataVersion::Make(Current, VersionCol->Version);
+			}
+		}
+		Current = GetParentScope(Current);
+	}
+
+	return Scope::FScopeDataVersion();
+}
+
+// ============================================================================
+// Scope Hierarchy
+// ============================================================================
+
+FHierarchyHandle FScopeSystem::GetScopeHierarchy()
+{
+	FHierarchyHandle Handle = DataStorage.FindHierarchyByName(GScopeHierarchyName);
+	if (!DataStorage.IsValidHierarchyHandle(Handle))
+	{
+		FHierarchyRegistrationParams Params;
+		Params.Name = GScopeHierarchyName;
+		Params.bEnableParentChangedColumn = true;
+		Handle = DataStorage.RegisterHierarchy(Params);
+	}
+	return Handle;
+}
+
+// ============================================================================
+// Root Scope
+// ============================================================================
+
+RowHandle FScopeSystem::GetRootScope() const
+{
+	return RootScopeRow;
+}
+
+void FScopeSystem::InitRootScope()
+{
+	if (RootScopeRow != InvalidRowHandle)
+	{
+		return;
+	}
+	RootScopeRow = AddScopeRow(TEXT("Root"));
+	Scope::SetCurrentScope(RootScopeRow);
+}
+
+} // namespace UE::Editor::DataStorage
diff --git a/Engine/Plugins/Experimental/EditorDataStorage/Source/TedsCore/Private/Scope/TypedElementScopeSystem.h b/Engine/Plugins/Experimental/EditorDataStorage/Source/TedsCore/Private/Scope/TypedElementScopeSystem.h
new file mode 100644
--- /dev/null
+++ b/Engine/Plugins/Experimental/EditorDataStorage/Source/TedsCore/Private/Scope/TypedElementScopeSystem.h
@@ -0,0 +1,44 @@
+// Copyright Epic Games, Inc. All Rights Reserved.
+
+#pragma once
+
+#include "DataStorage/CommonTypes.h"
+#include "DataStorage/Scope/EditorDataScopeTypes.h"
+#include "DataStorage/Handles.h"
+
+class UScriptStruct;
+
+namespace UE::Editor::DataStorage
+{
+
+class ICoreProvider;
+
+class FScopeSystem
+{
+public:
+	explicit FScopeSystem(ICoreProvider& InDataStorage);
+
+	RowHandle AddScopeRow();
+	RowHandle AddScopeRow(FStringView Label);
+	RowHandle AddScopeRow(FStringView Label, RowHandle Parent);
+	void RemoveScopeRow(RowHandle Row);
+	void SetParentScope(RowHandle Child, RowHandle Parent);
+	RowHandle GetParentScope(RowHandle Row);
+	const void* GetScopeDataRaw(RowHandle Row, const UScriptStruct* ColumnType);
+	Scope::FScopeDataVersion GetScopeDataVersion(RowHandle Row, const UScriptStruct* ColumnType);
+	Scope::FScopeDataVersion SetScopeData(RowHandle Row, const UScriptStruct* ColumnType,
+		ColumnCreationCallbackRef Initializer, ColumnCopyOrMoveCallback Relocator);
+	bool RemoveScopeData(RowHandle Row, const UScriptStruct* ColumnType);
+	TArray<const UScriptStruct*> GetAllVisibleScopeColumns(RowHandle Row);
+	FHierarchyHandle GetScopeHierarchy();
+	RowHandle GetRootScope() const;
+	void InitRootScope();
+
+private:
+	TableHandle GetOrCreateScopeTable();
+
+	ICoreProvider& DataStorage;
+	RowHandle RootScopeRow = InvalidRowHandle;
+};
+
+} // namespace UE::Editor::DataStorage
diff --git a/Engine/Plugins/Experimental/EditorDataStorage/Source/TedsCore/Private/TypedElementDatabase.h b/Engine/Plugins/Experimental/EditorDataStorage/Source/TedsCore/Private/TypedElementDatabase.h
--- a/Engine/Plugins/Experimental/EditorDataStorage/Source/TedsCore/Private/TypedElementDatabase.h
+++ b/Engine/Plugins/Experimental/EditorDataStorage/Source/TedsCore/Private/TypedElementDatabase.h
@@ -178,6 +178,25 @@ public:
 		UE::Editor::DataStorage::RowHandle Row,
 		TFunctionRef<bool(const ICoreProvider& Context, UE::Editor::DataStorage::RowHandle Child)> VisitFn) const override;
 
+	// Scope
+	UE_API virtual UE::Editor::DataStorage::RowHandle AddScopeRow() override;
+	UE_API virtual UE::Editor::DataStorage::RowHandle AddScopeRow(FStringView Label) override;
+	UE_API virtual UE::Editor::DataStorage::RowHandle AddScopeRow(FStringView Label, UE::Editor::DataStorage::RowHandle Parent) override;
+	UE_API virtual void RemoveScopeRow(UE::Editor::DataStorage::RowHandle Row) override;
+	UE_API virtual void SetParentScope(UE::Editor::DataStorage::RowHandle Child, UE::Editor::DataStorage::RowHandle Parent) override;
+	UE_API virtual UE::Editor::DataStorage::RowHandle GetParentScope(UE::Editor::DataStorage::RowHandle Row) override;
+	UE_API virtual const void* GetScopeDataRaw(UE::Editor::DataStorage::RowHandle Row, const UScriptStruct* ColumnType) override;
+	UE_API virtual UE::Editor::DataStorage::Scope::FScopeDataVersion GetScopeDataVersion(
+		UE::Editor::DataStorage::RowHandle Row, const UScriptStruct* ColumnType) override;
+	UE_API virtual UE::Editor::DataStorage::Scope::FScopeDataVersion SetScopeData(
+		UE::Editor::DataStorage::RowHandle Row, const UScriptStruct* ColumnType,
+		UE::Editor::DataStorage::ColumnCreationCallbackRef Initializer, UE::Editor::DataStorage::ColumnCopyOrMoveCallback Relocator) override;
+	UE_API virtual bool RemoveScopeData(UE::Editor::DataStorage::RowHandle Row, const UScriptStruct* ColumnType) override;
+	UE_API virtual TArray<const UScriptStruct*> GetAllVisibleScopeColumns(
+		UE::Editor::DataStorage::RowHandle Row) override;
+	UE_API virtual UE::Editor::DataStorage::RowHandle GetRootScope() const override;
+	UE_API virtual UE::Editor::DataStorage::FHierarchyHandle GetScopeHierarchy() override;
+
 	UE_API void RegisterTickGroup(FName GroupName, UE::Editor::DataStorage::EQueryTickPhase Phase, FName BeforeGroup, FName AfterGroup, UE::Editor::DataStorage::EExecutionMode ExecutionMode);
 	UE_API void UnregisterTickGroup(FName GroupName, UE::Editor::DataStorage::EQueryTickPhase Phase);
 
diff --git a/Engine/Plugins/Experimental/EditorDataStorage/Source/TedsCore/Private/TypedElementDatabaseEnvironment.cpp b/Engine/Plugins/Experimental/EditorDataStorage/Source/TedsCore/Private/TypedElementDatabaseEnvironment.cpp
--- a/Engine/Plugins/Experimental/EditorDataStorage/Source/TedsCore/Private/TypedElementDatabaseEnvironment.cpp
+++ b/Engine/Plugins/Experimental/EditorDataStorage/Source/TedsCore/Private/TypedElementDatabaseEnvironment.cpp
@@ -14,6 +14,7 @@ namespace UE::Editor::DataStorage
 		, MappingTable(InDataStorage)
 		, Queries(DynamicColumnGenerator)
 		, MementoSystem(InDataStorage)
+		, ScopeSystem(InDataStorage)
 		, ValueTagManager(InMassEntityManager, DynamicColumnGenerator)
 		, ActiveCommandQueue(CommandQueues)
 		, MassEntityManager(InMassEntityManager)
@@ -78,6 +79,16 @@ namespace UE::Editor::DataStorage
 		return MementoSystem;
 	}
 
+	FScopeSystem& FEnvironment::GetScopeSystem()
+	{
+		return ScopeSystem;
+	}
+
+	const FScopeSystem& FEnvironment::GetScopeSystem() const
+	{
+		return ScopeSystem;
+	}
+
 	FMassEntityManager& FEnvironment::GetMassEntityManager()
 	{
 		return MassEntityManager;
diff --git a/Engine/Plugins/Experimental/EditorDataStorage/Source/TedsCore/Private/TypedElementDatabaseEnvironment.h b/Engine/Plugins/Experimental/EditorDataStorage/Source/TedsCore/Private/TypedElementDatabaseEnvironment.h
--- a/Engine/Plugins/Experimental/EditorDataStorage/Source/TedsCore/Private/TypedElementDatabaseEnvironment.h
+++ b/Engine/Plugins/Experimental/EditorDataStorage/Source/TedsCore/Private/TypedElementDatabaseEnvironment.h
@@ -13,6 +13,7 @@
 #include "Hierarchy/EditorDataStorageHierarchyImplementation.h"
 #include "Memento/TypedElementMementoSystem.h"
 #include "Queries/TypedElementExtendedQueryStore.h"
+#include "Scope/TypedElementScopeSystem.h"
 
 class UEditorDataStorage;
 struct FMassCommandBuffer;
@@ -42,6 +43,9 @@ namespace UE::Editor::DataStorage
 		FMementoSystem& GetMementoSystem();
 		const FMementoSystem& GetMementoSystem() const;
 
+		FScopeSystem& GetScopeSystem();
+		const FScopeSystem& GetScopeSystem() const;
+
 		FMassEntityManager& GetMassEntityManager();
 		const FMassEntityManager& GetMassEntityManager() const;
 
@@ -95,6 +99,7 @@ namespace UE::Editor::DataStorage
 		FDynamicColumnGenerator DynamicColumnGenerator;
 		FExtendedQueryStore Queries;
 		FMementoSystem MementoSystem;
+		FScopeSystem ScopeSystem;
 		FValueTagManager ValueTagManager;
 		FTedsHierarchyRegistrar HierarchyInterface;
 
diff --git a/Engine/Plugins/Experimental/EditorDataStorage/Source/TedsCore/Private/TypedElementsDataStorage.cpp b/Engine/Plugins/Experimental/EditorDataStorage/Source/TedsCore/Private/TypedElementsDataStorage.cpp
--- a/Engine/Plugins/Experimental/EditorDataStorage/Source/TedsCore/Private/TypedElementsDataStorage.cpp
+++ b/Engine/Plugins/Experimental/EditorDataStorage/Source/TedsCore/Private/TypedElementsDataStorage.cpp
@@ -4,6 +4,7 @@
 
 #include "DataStorage/Features.h"
 #include "DataStorage/Debug/Log.h"
+#include "TypedElementDatabaseEnvironment.h"
 #include "EditorDataStorageSettings.h"
 #include "Elements/Framework/TypedElementRegistry.h"
 #include "Elements/Interfaces/TypedElementDataStorageFactory.h"
@@ -166,6 +167,9 @@ void FEditorDataStorageModule::StartupModule()
 					Factory->PostRegister(*DataStorage);
 				}
 
+				// Initialize the root scope row for the Editor Scope Tree system.
+				DataStorage->GetEnvironment().GetScopeSystem().InitRootScope();
+
 				UE_LOG(LogEditorDataStorage, Log, TEXT("Initialized"));
 
 				bInitialized = true;
@@ -220,5 +224,80 @@ FString FEditorDataStorageModule::GetReferencerName() const
 }
 
 #undef LOCTEXT_NAMESPACE
-	
+
+// ============================================================================
+// ICoreProvider Scope forwarders
+// ============================================================================
+
+using namespace UE::Editor::DataStorage;
+using namespace UE::Editor::DataStorage::Scope;
+
+RowHandle UEditorDataStorage::AddScopeRow()
+{
+	return GetEnvironment().GetScopeSystem().AddScopeRow();
+}
+
+RowHandle UEditorDataStorage::AddScopeRow(FStringView Label)
+{
+	return GetEnvironment().GetScopeSystem().AddScopeRow(Label);
+}
+
+RowHandle UEditorDataStorage::AddScopeRow(FStringView Label, RowHandle Parent)
+{
+	return GetEnvironment().GetScopeSystem().AddScopeRow(Label, Parent);
+}
+
+void UEditorDataStorage::RemoveScopeRow(RowHandle Row)
+{
+	GetEnvironment().GetScopeSystem().RemoveScopeRow(Row);
+}
+
+void UEditorDataStorage::SetParentScope(RowHandle Child, RowHandle Parent)
+{
+	GetEnvironment().GetScopeSystem().SetParentScope(Child, Parent);
+}
+
+RowHandle UEditorDataStorage::GetParentScope(RowHandle Row)
+{
+	return GetEnvironment().GetScopeSystem().GetParentScope(Row);
+}
+
+const void* UEditorDataStorage::GetScopeDataRaw(RowHandle Row, const UScriptStruct* ColumnType)
+{
+	return GetEnvironment().GetScopeSystem().GetScopeDataRaw(Row, ColumnType);
+}
+
+TArray<const UScriptStruct*> UEditorDataStorage::GetAllVisibleScopeColumns(RowHandle Row)
+{
+	return GetEnvironment().GetScopeSystem().GetAllVisibleScopeColumns(Row);
+}
+
+FScopeDataVersion UEditorDataStorage::SetScopeData(
+	RowHandle Row, const UScriptStruct* ColumnType,
+	ColumnCreationCallbackRef Initializer, ColumnCopyOrMoveCallback Relocator)
+{
+	return GetEnvironment().GetScopeSystem().SetScopeData(Row, ColumnType, Initializer, Relocator);
+}
+
+bool UEditorDataStorage::RemoveScopeData(RowHandle Row, const UScriptStruct* ColumnType)
+{
+	return GetEnvironment().GetScopeSystem().RemoveScopeData(Row, ColumnType);
+}
+
+FScopeDataVersion UEditorDataStorage::GetScopeDataVersion(
+	RowHandle Row, const UScriptStruct* ColumnType)
+{
+	return GetEnvironment().GetScopeSystem().GetScopeDataVersion(Row, ColumnType);
+}
+
+FHierarchyHandle UEditorDataStorage::GetScopeHierarchy()
+{
+	return GetEnvironment().GetScopeSystem().GetScopeHierarchy();
+}
+
+RowHandle UEditorDataStorage::GetRootScope() const
+{
+	return GetEnvironment().GetScopeSystem().GetRootScope();
+}
+
 IMPLEMENT_MODULE(FEditorDataStorageModule, TedsCore)
\ No newline at end of file
diff --git a/Engine/Plugins/Experimental/EditorDataStorageFeatures/Source/TedsDebugger/Private/STedsDebugger.cpp b/Engine/Plugins/Experimental/EditorDataStorageFeatures/Source/TedsDebugger/Private/STedsDebugger.cpp
--- a/Engine/Plugins/Experimental/EditorDataStorageFeatures/Source/TedsDebugger/Private/STedsDebugger.cpp
+++ b/Engine/Plugins/Experimental/EditorDataStorageFeatures/Source/TedsDebugger/Private/STedsDebugger.cpp
@@ -14,6 +14,7 @@
 #include "QueryEditor/TedsQueryEditor.h"
 #include "TedsOutlinerModule.h"
 #include "TedsOutlinerMode.h"
+#include "Widgets/SScopeHierarchyTab.h"
 #include "Widgets/STedsDiscoveryDebugger.h"
 
 #define LOCTEXT_NAMESPACE "STedsDebugger"
@@ -24,6 +25,7 @@ namespace UE::Editor::DataStorage::Debug
 	{
 		const FName ToolbarTabName = TEXT("TEDS Debugger Toolbar");
 		const FName DiscoveryTabName = TEXT("Discover");
+		const FName ContextHierarchyTabName = TEXT("Scope Hierarchy");
 	}
 
 enum class EQueryEditorTabs
@@ -54,7 +56,7 @@ void STedsDebugger::Construct(const FArguments& InArgs, const TSharedRef<SDockTa
 	RegisterTabSpawners();
 
 	// Setup the default layout
-	TSharedRef<FTabManager::FLayout> Layout = FTabManager::NewLayout("TedsDebuggerLayout_v1.1")
+	TSharedRef<FTabManager::FLayout> Layout = FTabManager::NewLayout("TedsDebuggerLayout_v1.2")
 	->AddArea
 	(
 		FTabManager::NewPrimaryArea()
@@ -77,6 +79,7 @@ void STedsDebugger::Construct(const FArguments& InArgs, const TSharedRef<SDockTa
 					->AddTab(GetQueryTabName(EQueryEditorTabs::QueryEditorThree), ETabState::ClosedTab)
 					->AddTab(GetQueryTabName(EQueryEditorTabs::QueryEditorFour), ETabState::ClosedTab)
 					->AddTab(Private::DiscoveryTabName, ETabState::OpenedTab)
+					->AddTab(Private::ContextHierarchyTabName, ETabState::OpenedTab)
 					->SetHideTabWell(false)
 					->SetForegroundTab(GetQueryTabName(EQueryEditorTabs::QueryEditorOne))
 			)
@@ -194,6 +197,16 @@ TSharedRef<SDockTab> STedsDebugger::SpawnDiscoverTab(const FSpawnTabArgs& Args)
 	return DockTab;
 }
 
+TSharedRef<SDockTab> STedsDebugger::SpawnContextHierarchyTab(const FSpawnTabArgs& Args)
+{
+	TSharedRef<SDockTab> DockTab = SNew(SDockTab)
+		.TabRole(ETabRole::PanelTab)
+		.Label(LOCTEXT("TedsDebugger_ContextHierarchyDisplayName", "Scope Hierarchy"));
+
+	DockTab->SetContent(SNew(SScopeHierarchyTab));
+	return DockTab;
+}
+
 void STedsDebugger::RegisterTabSpawners()
 {
 	const TSharedRef<FWorkspaceItem> AppMenuGroup =
@@ -246,6 +259,14 @@ void STedsDebugger::RegisterTabSpawners()
 		.SetDisplayName(LOCTEXT("TedsDebugger_DiscoverDisplayName", "Discover"))
 		.SetTooltipText(LOCTEXT("TedsDebugger_DiscoverToolTip", "Opens the Discover tab. Contains special searching functionality to explore TEDS data without a query."))
 		.SetIcon(FSlateIcon(FAppStyle::GetAppStyleSetName(), "Icons.Search"));
+
+	TabManager->RegisterTabSpawner(
+		Private::ContextHierarchyTabName,
+		FOnSpawnTab::CreateRaw(this, &STedsDebugger::SpawnContextHierarchyTab))
+		.SetGroup(AppMenuGroup)
+		.SetDisplayName(LOCTEXT("TedsDebugger_ContextHierarchyDisplayName", "Scope Hierarchy"))
+		.SetTooltipText(LOCTEXT("TedsDebugger_ContextHierarchyToolTip", "Shows the ECP scope row hierarchy and lets you inspect creation callstacks and visible columns."))
+		.SetIcon(FSlateIcon(FAppStyle::GetAppStyleSetName(), "Icons.Hierarchy"));
 }
 
 } // namespace UE::Editor::DataStorage::Debug
diff --git a/Engine/Plugins/Experimental/EditorDataStorageFeatures/Source/TedsDebugger/Private/STedsDebugger.h b/Engine/Plugins/Experimental/EditorDataStorageFeatures/Source/TedsDebugger/Private/STedsDebugger.h
--- a/Engine/Plugins/Experimental/EditorDataStorageFeatures/Source/TedsDebugger/Private/STedsDebugger.h
+++ b/Engine/Plugins/Experimental/EditorDataStorageFeatures/Source/TedsDebugger/Private/STedsDebugger.h
@@ -41,6 +41,7 @@ namespace UE::Editor::DataStorage
 			TSharedRef<SDockTab> SpawnToolbar(const FSpawnTabArgs& Args);
 			TSharedRef<SDockTab> SpawnQueryEditorTab(const FSpawnTabArgs& Args, const EQueryEditorTabs QueryTabIdentifier);
 			TSharedRef<SDockTab> SpawnDiscoverTab(const FSpawnTabArgs& Args);
+			TSharedRef<SDockTab> SpawnContextHierarchyTab(const FSpawnTabArgs& Args);
 			void FillWindowMenu( FMenuBuilder& MenuBuilder);
 
 			void RegisterTabSpawners();
diff --git a/Engine/Plugins/Experimental/EditorDataStorageFeatures/Source/TedsDebugger/Private/Widgets/SScopeHierarchyTab.cpp b/Engine/Plugins/Experimental/EditorDataStorageFeatures/Source/TedsDebugger/Private/Widgets/SScopeHierarchyTab.cpp
new file mode 100644
--- /dev/null
+++ b/Engine/Plugins/Experimental/EditorDataStorageFeatures/Source/TedsDebugger/Private/Widgets/SScopeHierarchyTab.cpp
@@ -0,0 +1,368 @@
+// Copyright Epic Games, Inc. All Rights Reserved.
+
+#include "SScopeHierarchyTab.h"
+
+#include "ScopeRowLabelWidget.h"
+#include "DataStorage/Scope/EditorDataScope.h"
+#include "DataStorage/Scope/EditorDataScopeColumns.h"
+#include "Elements/Framework/TypedElementQueryBuilder.h"
+#include "Elements/Columns/TypedElementMiscColumns.h"
+#include "Elements/Columns/TypedElementTypeInfoColumns.h"
+#include "Elements/Interfaces/TypedElementDataStorageInterface.h"
+#include "Elements/Interfaces/TypedElementDataStorageUiInterface.h"
+#include "HAL/PlatformApplicationMisc.h"
+#include "HAL/PlatformStackWalk.h"
+#include "TedsTableViewerColumn.h"
+#include "Widgets/Images/SImage.h"
+#include "Widgets/Input/SButton.h"
+#include "Widgets/Input/SMultiLineEditableTextBox.h"
+#include "Widgets/Layout/SBorder.h"
+#include "Widgets/Layout/SScrollBox.h"
+#include "Widgets/Layout/SSplitter.h"
+#include "Widgets/SBoxPanel.h"
+#include "Widgets/SNullWidget.h"
+#include "Widgets/STedsHierarchyViewer.h"
+#include "Widgets/Text/STextBlock.h"
+#include "Widgets/Views/SHeaderRow.h"
+
+#define LOCTEXT_NAMESPACE "SScopeHierarchyTab"
+
+namespace
+{
+	const FName EditorWidgetTableName(TEXT("Editor_WidgetTable"));
+}
+
+namespace UE::Editor::DataStorage::Debug
+{
+
+void SScopeHierarchyTab::Construct(const FArguments& InArgs)
+{
+	Storage = GetMutableDataStorageFeature<ICoreProvider>(StorageFeatureName);
+	checkf(Storage, TEXT("Cannot create the Scope Hierarchy Tab without TEDS Core"));
+	DataStorageUi = GetMutableDataStorageFeature<IUiProvider>(UiFeatureName);
+
+	ChildSlot
+	[
+		SNew(SBorder)
+		.BorderImage(FAppStyle::GetBrush("Brushes.Background"))
+		[
+			SNew(SSplitter)
+			+SSplitter::Slot()
+			.Value(0.5f)
+			[
+				CreateHierarchyViewer()
+			]
+			+SSplitter::Slot()
+			.Value(0.5f)
+			[
+				CreateDetailsPanel()
+			]
+		]
+	];
+}
+
+TSharedRef<SWidget> SScopeHierarchyTab::CreateHierarchyViewer()
+{
+	// Build a query matching all scope rows (by tag) and selecting FScopeRowSourceInfo
+	FQueryDescription QueryDesc = Queries::Select()
+		.ReadOnly<FScopeRowSourceInfo>()
+		.Where()
+		.All<FDataStorageScopeDataTag>()
+		.Compile();
+
+	ContextQueryNode = MakeShared<QueryStack::FQueryNode>(*Storage, MoveTemp(QueryDesc));
+	ScopeRowsNode = MakeShared<QueryStack::FRowQueryResultsNode>(
+		*Storage, ContextQueryNode, QueryStack::FRowQueryResultsNode::ESyncActions::RefreshOnUpdate);
+	ScopeRowsNode->Refresh();
+
+	// Create hierarchy data wrapping the editor scope hierarchy
+	FHierarchyHandle HierarchyHandle = Storage->GetScopeHierarchy();
+	if (Storage->IsValidHierarchyHandle(HierarchyHandle))
+	{
+		HierarchyData = MakeShared<FHierarchyViewerData>(*Storage, HierarchyHandle);
+	}
+
+	HierarchyViewer = SNew(SHierarchyViewer, HierarchyData)
+		.TableViewerIdentifier(TEXT("ScopeHierarchyViewer"))
+		.AllNodeProvider(ScopeRowsNode)
+		.Columns({})
+		.ExpandNewRows(true)
+		.EmptyRowsMessage(LOCTEXT("EmptyRowsMessage", "No scope rows found. Is TEDS initialized?"))
+		.OnSelectionChanged(
+			SHierarchyViewer::FOnSelectionChanged::CreateSP(this, &SScopeHierarchyTab::OnHierarchySelectionChanged));
+
+	// Add a label display column — shows FTypedElementLabelColumn text when available,
+	// falls back to the stringified RowHandle otherwise.
+	{
+		IUiProvider::FWidgetConstructorPtr LabelConstructor = MakeShared<FScopeRowLabelWidgetConstructor>();
+		TSharedRef<FTedsTableViewerColumn> LabelColumn =
+			MakeShared<FTedsTableViewerColumn>(TEXT("Scope Row"), LabelConstructor);
+		HierarchyViewer->AddCustomRowWidget(LabelColumn);
+	}
+
+	return HierarchyViewer.ToSharedRef();
+}
+
+TSharedRef<SWidget> SScopeHierarchyTab::CreateDetailsPanel()
+{
+	return SNew(SVerticalBox)
+		// Callstack header + copy button
+		+SVerticalBox::Slot()
+		.AutoHeight()
+		.Padding(4.f)
+		[
+			SNew(SHorizontalBox)
+			+SHorizontalBox::Slot()
+			.FillWidth(1.f)
+			.VAlign(VAlign_Center)
+			[
+				SNew(STextBlock)
+				.Text(LOCTEXT("CallstackHeader", "Creation Callstack"))
+				.Font(FAppStyle::GetFontStyle("BoldFont"))
+			]
+			+SHorizontalBox::Slot()
+			.AutoWidth()
+			[
+				SNew(SButton)
+				.ButtonStyle(FAppStyle::Get(), "SimpleButton")
+				.ToolTipText(LOCTEXT("CopyCallstackTooltip", "Copy callstack to clipboard"))
+				.OnClicked(this, &SScopeHierarchyTab::CopyCallstackToClipboard)
+				[
+					SNew(SImage)
+					.Image(FAppStyle::GetBrush("Icons.Clipboard"))
+					.ColorAndOpacity(FSlateColor::UseForeground())
+				]
+			]
+		]
+		// Callstack text (read-only, selectable, no wrap, horizontal scroll)
+		+SVerticalBox::Slot()
+		.FillHeight(0.5f)
+		.Padding(4.f)
+		[
+			SAssignNew(CallstackTextBox, SMultiLineEditableTextBox)
+			.IsReadOnly(true)
+			.Text(LOCTEXT("NoSelection", "Select a scope row to view details."))
+			.AutoWrapText(false)
+			.Font(FCoreStyle::GetDefaultFontStyle("Mono", 9))
+			.AlwaysShowScrollbars(true)
+		]
+		// Visible columns header
+		+SVerticalBox::Slot()
+		.AutoHeight()
+		.Padding(4.f)
+		[
+			SNew(STextBlock)
+			.Text(LOCTEXT("VisibleColumnsHeader", "Visible Columns"))
+			.Font(FAppStyle::GetFontStyle("BoldFont"))
+		]
+		// Visible columns list
+		+SVerticalBox::Slot()
+		.FillHeight(0.5f)
+		.Padding(4.f)
+		[
+			SAssignNew(ColumnListView, SListView<FScopeColumnItemPtr>)
+			.ListItemsSource(&ColumnItems)
+			.OnGenerateRow(this, &SScopeHierarchyTab::MakeColumnListRow)
+			.HeaderRow(
+				SNew(SHeaderRow)
+				+SHeaderRow::Column(TEXT("Origin"))
+					.DefaultLabel(FText())
+					.FixedWidth(24.f)
+				+SHeaderRow::Column(TEXT("ColumnType"))
+					.DefaultLabel(LOCTEXT("ColumnTypeHeader", "Column"))
+					.FillWidth(0.3f)
+				+SHeaderRow::Column(TEXT("Value"))
+					.DefaultLabel(LOCTEXT("ValueHeader", "Value"))
+					.FillWidth(0.7f)
+			)
+		];
+}
+
+void SScopeHierarchyTab::OnHierarchySelectionChanged(RowHandle Row, ESelectInfo::Type SelectInfo)
+{
+	if (Storage->IsRowAvailable(Row))
+	{
+		RebuildDetailsForRow(Row);
+	}
+	else
+	{
+		CachedCallstackString.Empty();
+		CallstackTextBox->SetText(LOCTEXT("NoSelection", "Select a scope row to view details."));
+		ColumnItems.Empty();
+		ColumnListView->RequestListRefresh();
+	}
+}
+
+void SScopeHierarchyTab::RebuildDetailsForRow(RowHandle Row)
+{
+	// Resolve callstack (skip first frame — it's always the CaptureStackBackTrace call)
+	CachedCallstackString.Empty();
+	if (const FScopeRowSourceInfo* SourceInfo =
+			Storage->GetScopeData<FScopeRowSourceInfo>(Row))
+	{
+		ANSICHAR LineBuffer[1024];
+		for (int32 i = 1; i < SourceInfo->CallstackDepth; ++i)
+		{
+			LineBuffer[0] = 0;
+			FPlatformStackWalk::ProgramCounterToHumanReadableString(
+				i, SourceInfo->Callstack[i], LineBuffer, UE_ARRAY_COUNT(LineBuffer));
+			if (CachedCallstackString.Len() > 0)
+			{
+				CachedCallstackString += TEXT("\n");
+			}
+			CachedCallstackString += ANSI_TO_TCHAR(LineBuffer);
+		}
+		CallstackTextBox->SetText(FText::FromString(CachedCallstackString));
+	}
+	else
+	{
+		CallstackTextBox->SetText(LOCTEXT("NoCallstack", "No callstack info available for this row."));
+	}
+
+	// Rebuild visible columns
+	ColumnItems.Empty();
+	TArray<const UScriptStruct*> VisibleColumns = Storage->GetAllVisibleScopeColumns(Row);
+	for (const UScriptStruct* ColumnType : VisibleColumns)
+	{
+		auto Item = MakeShared<FScopeColumnItem>();
+		Item->TypeName = ColumnType ? ColumnType->GetName() : TEXT("(null)");
+		Item->OwnerRow = Row;
+		Item->bInherited = !Storage->HasColumns(Row, TConstArrayView<const UScriptStruct*>(&ColumnType, 1));
+		Item->ColumnType = ColumnType;
+
+		if (DataStorageUi && ColumnType)
+		{
+			DataStorageUi->CreateWidgetConstructors(
+				DataStorageUi->FindPurpose(DataStorageUi->GetDefaultWidgetPurposeID()),
+				FMetaDataView(),
+				[&Item](IUiProvider::FWidgetConstructorPtr Constructor, TConstArrayView<TWeakObjectPtr<const UScriptStruct>>)
+				{
+					Item->WidgetConstructor = MoveTemp(Constructor);
+					return false;
+				});
+		}
+
+		ColumnItems.Add(Item);
+	}
+	ColumnListView->RequestListRefresh();
+}
+
+FReply SScopeHierarchyTab::CopyCallstackToClipboard()
+{
+	if (!CachedCallstackString.IsEmpty())
+	{
+		FPlatformApplicationMisc::ClipboardCopy(*CachedCallstackString);
+	}
+	return FReply::Handled();
+}
+
+TSharedRef<ITableRow> SScopeHierarchyTab::MakeColumnListRow(
+	FScopeColumnItemPtr Item, const TSharedRef<STableViewBase>& OwnerTable)
+{
+	class SColumnRow : public SMultiColumnTableRow<FScopeColumnItemPtr>
+	{
+	public:
+		SLATE_BEGIN_ARGS(SColumnRow) {}
+			SLATE_ARGUMENT(FScopeColumnItemPtr, Item)
+			SLATE_ARGUMENT(ICoreProvider*, DataStorage)
+			SLATE_ARGUMENT(IUiProvider*, DataStorageUi)
+		SLATE_END_ARGS()
+
+		void Construct(const FArguments& InArgs, const TSharedRef<STableViewBase>& InOwnerTable)
+		{
+			RowItem = InArgs._Item;
+			DataStorage = InArgs._DataStorage;
+			DataStorageUi = InArgs._DataStorageUi;
+			SMultiColumnTableRow::Construct(FSuperRowType::FArguments(), InOwnerTable);
+		}
+
+		TSharedRef<SWidget> GenerateWidgetForColumn(const FName& ColumnName) override
+		{
+			if (ColumnName == TEXT("Origin"))
+			{
+				const FName BrushName = RowItem->bInherited
+					? TEXT("Icons.ArrowDown")
+					: TEXT("Icons.FilledCircle");
+
+				const FText Tooltip = RowItem->bInherited
+					? FText::Format(LOCTEXT("InheritedTooltip", "Inherited from row {0}"),
+						FText::FromString(LexToString(RowItem->OwnerRow)))
+					: LOCTEXT("OwnTooltip", "Owned by this row");
+
+				return SNew(SBox)
+					.HAlign(HAlign_Center)
+					.VAlign(VAlign_Center)
+					.WidthOverride(16.f)
+					.HeightOverride(16.f)
+					[
+						SNew(SImage)
+						.Image(FAppStyle::GetBrush(BrushName))
+						.ColorAndOpacity(RowItem->bInherited
+							? FSlateColor(FLinearColor(0.5f, 0.5f, 0.8f))
+							: FSlateColor(FLinearColor(0.2f, 0.8f, 0.2f)))
+						.ToolTipText(Tooltip)
+					];
+			}
+			if (ColumnName == TEXT("ColumnType"))
+			{
+				return SNew(STextBlock).Text(FText::FromString(RowItem->TypeName));
+			}
+			if (ColumnName == TEXT("Value"))
+			{
+				return SNew(STextBlock)
+					.Text_Lambda([WeakThis = TWeakPtr<SColumnRow>(StaticCastSharedRef<SColumnRow>(this->AsShared()))]()
+						{
+							if (TSharedPtr<SColumnRow> ColumnRowWidget = WeakThis.Pin())
+							{
+								if (!ColumnRowWidget->RowItem.IsValid())
+								{
+									return LOCTEXT("ColumnNotFoundTextEarly", "Column not found on row");
+								}
+								ICoreProvider* DataStorage = ColumnRowWidget->DataStorage;
+								RowHandle ScopeRow = ColumnRowWidget->RowItem->OwnerRow;
+								if (!DataStorage || !DataStorage->IsRowAvailable(ScopeRow))
+								{
+									return LOCTEXT("ColumnNotFoundTextEarly", "Column not found on row");
+								}
+
+								TStrongObjectPtr<const UScriptStruct> TypeInfo = ColumnRowWidget->RowItem->ColumnType.Pin();
+								if (!TypeInfo)
+								{
+									return LOCTEXT("ColumnNotFoundTextEarly", "Column not found on row");
+								}
+								
+								// Tags have no data — show blank
+								if (TypeInfo->IsChildOf(FEditorDataStorageTag::StaticStruct()))
+								{
+									return FText::GetEmpty();
+								}
+
+								const void* Data = DataStorage->GetScopeDataRaw(ScopeRow, TypeInfo.Get());
+								if (!Data)
+								{
+									return LOCTEXT("ColumnNotFoundTextEarly", "Column not found on row");
+								}
+
+								FString Label;
+								TypeInfo->ExportText(Label, Data, Data, nullptr, PPF_None, nullptr);
+								return FText::FromString(Label);
+								
+							}
+							return LOCTEXT("ColumnNotFoundTextEarly", "Column not found on row");
+						});
+			}
+			return SNullWidget::NullWidget;
+		}
+
+	private:
+		FScopeColumnItemPtr RowItem;
+		ICoreProvider* DataStorage = nullptr;
+		IUiProvider* DataStorageUi = nullptr;
+	};
+
+	return SNew(SColumnRow, OwnerTable).Item(Item).DataStorage(Storage).DataStorageUi(DataStorageUi);
+}
+
+} // namespace UE::Editor::DataStorage::Debug
+
+#undef LOCTEXT_NAMESPACE
diff --git a/Engine/Plugins/Experimental/EditorDataStorageFeatures/Source/TedsDebugger/Private/Widgets/SScopeHierarchyTab.h b/Engine/Plugins/Experimental/EditorDataStorageFeatures/Source/TedsDebugger/Private/Widgets/SScopeHierarchyTab.h
new file mode 100644
--- /dev/null
+++ b/Engine/Plugins/Experimental/EditorDataStorageFeatures/Source/TedsDebugger/Private/Widgets/SScopeHierarchyTab.h
@@ -0,0 +1,73 @@
+// Copyright Epic Games, Inc. All Rights Reserved.
+#pragma once
+
+#include "DataStorage/Scope/EditorDataScope.h"
+#include "DataStorage/Features.h"
+#include "Elements/Interfaces/TypedElementDataStorageUiInterface.h"
+#include "HierarchyViewerIntefaces.h"
+#include "TedsQueryNode.h"
+#include "TedsRowQueryResultsNode.h"
+#include "Widgets/SCompoundWidget.h"
+#include "Widgets/Views/SListView.h"
+
+class SMultiLineEditableTextBox;
+class STextBlock;
+
+namespace UE::Editor::DataStorage
+{
+	class SHierarchyViewer;
+}
+
+namespace UE::Editor::DataStorage::Debug
+{
+	/** Item for the visible-columns list view on the details side. */
+	struct FScopeColumnItem
+	{
+		FString TypeName;
+		RowHandle OwnerRow = InvalidRowHandle;
+		bool bInherited = false;
+		TWeakObjectPtr<const UScriptStruct> ColumnType;
+		IUiProvider::FWidgetConstructorPtr WidgetConstructor;
+		RowHandle WidgetRow = InvalidRowHandle;
+	};
+	using FScopeColumnItemPtr = TSharedPtr<FScopeColumnItem>;
+
+	/**
+	 * Debugger tab showing the Editor Scope Tree row hierarchy.
+	 * Left: SHierarchyViewer tree of all scope rows.
+	 * Right: On selection, resolved creation callstack + visible columns list.
+	 */
+	class SScopeHierarchyTab final : public SCompoundWidget
+	{
+	public:
+		SLATE_BEGIN_ARGS(SScopeHierarchyTab) {}
+		SLATE_END_ARGS()
+
+		void Construct(const FArguments& InArgs);
+
+	private:
+		TSharedRef<SWidget> CreateHierarchyViewer();
+		TSharedRef<SWidget> CreateDetailsPanel();
+		void OnHierarchySelectionChanged(RowHandle Row, ESelectInfo::Type SelectInfo);
+		void RebuildDetailsForRow(RowHandle Row);
+		FReply CopyCallstackToClipboard();
+
+		TSharedRef<ITableRow> MakeColumnListRow(
+			FScopeColumnItemPtr Item, const TSharedRef<STableViewBase>& OwnerTable);
+
+		ICoreProvider* Storage = nullptr;
+		IUiProvider* DataStorageUi = nullptr;
+
+		// Hierarchy viewer
+		TSharedPtr<SHierarchyViewer> HierarchyViewer;
+		TSharedPtr<FHierarchyViewerData> HierarchyData;
+		TSharedPtr<QueryStack::FQueryNode> ContextQueryNode;
+		TSharedPtr<QueryStack::FRowQueryResultsNode> ScopeRowsNode;
+
+		// Details panel
+		TSharedPtr<SMultiLineEditableTextBox> CallstackTextBox;
+		FString CachedCallstackString;
+		TSharedPtr<SListView<FScopeColumnItemPtr>> ColumnListView;
+		TArray<FScopeColumnItemPtr> ColumnItems;
+	};
+} // namespace UE::Editor::DataStorage::Debug
diff --git a/Engine/Plugins/Experimental/EditorDataStorageFeatures/Source/TedsDebugger/Private/Widgets/ScopeRowLabelWidget.cpp b/Engine/Plugins/Experimental/EditorDataStorageFeatures/Source/TedsDebugger/Private/Widgets/ScopeRowLabelWidget.cpp
new file mode 100644
--- /dev/null
+++ b/Engine/Plugins/Experimental/EditorDataStorageFeatures/Source/TedsDebugger/Private/Widgets/ScopeRowLabelWidget.cpp
@@ -0,0 +1,43 @@
+// Copyright Epic Games, Inc. All Rights Reserved.
+
+#include "ScopeRowLabelWidget.h"
+
+#include "Elements/Columns/TypedElementLabelColumns.h"
+#include "Elements/Interfaces/TypedElementDataStorageInterface.h"
+#include "Widgets/Text/STextBlock.h"
+
+FScopeRowLabelWidgetConstructor::FScopeRowLabelWidgetConstructor()
+	: FSimpleWidgetConstructor(StaticStruct())
+{
+}
+
+TSharedPtr<SWidget> FScopeRowLabelWidgetConstructor::CreateWidget(
+	UE::Editor::DataStorage::ICoreProvider* DataStorage,
+	UE::Editor::DataStorage::IUiProvider* DataStorageUi,
+	UE::Editor::DataStorage::RowHandle TargetRow,
+	UE::Editor::DataStorage::RowHandle WidgetRow,
+	const UE::Editor::DataStorage::FMetaDataView& Arguments)
+{
+	using namespace UE::Editor::DataStorage;
+
+	const FString RowHandleString = LexToString(TargetRow);
+
+	return SNew(STextBlock)
+		.Text_Lambda([DataStorage, TargetRow, RowHandleString]()
+		{
+			if (DataStorage && DataStorage->IsRowAvailable(TargetRow))
+			{
+				// NOTE: We _don't_ use the scope APIs here because we are displaying the data about the individual scope row itself
+				// We don't want to show the label for a row above if this one doesn't have one or if it is empty.
+				if (const FTypedElementLabelColumn* LabelColumn = DataStorage->GetColumn<FTypedElementLabelColumn>(TargetRow))
+				{
+					if (!LabelColumn->Label.IsEmpty())
+					{
+						return FText::FromString(LabelColumn->Label);
+					}
+				}
+			}
+			return FText::FromString(RowHandleString);
+		})
+		.ToolTipText(FText::FromString(RowHandleString));
+}
diff --git a/Engine/Plugins/Experimental/EditorDataStorageFeatures/Source/TedsDebugger/Private/Widgets/ScopeRowLabelWidget.h b/Engine/Plugins/Experimental/EditorDataStorageFeatures/Source/TedsDebugger/Private/Widgets/ScopeRowLabelWidget.h
new file mode 100644
--- /dev/null
+++ b/Engine/Plugins/Experimental/EditorDataStorageFeatures/Source/TedsDebugger/Private/Widgets/ScopeRowLabelWidget.h
@@ -0,0 +1,29 @@
+// Copyright Epic Games, Inc. All Rights Reserved.
+
+#pragma once
+
+#include "Elements/Interfaces/TypedElementDataStorageUiInterface.h"
+
+#include "ScopeRowLabelWidget.generated.h"
+
+/**
+ * Widget constructor for scope hierarchy rows that shows the FTypedElementLabelColumn
+ * text when available, falling back to the stringified RowHandle. When a label is shown,
+ * the tooltip displays the RowHandle for debugging.
+ */
+USTRUCT()
+struct FScopeRowLabelWidgetConstructor : public FSimpleWidgetConstructor
+{
+	GENERATED_BODY()
+
+public:
+	FScopeRowLabelWidgetConstructor();
+	~FScopeRowLabelWidgetConstructor() override = default;
+
+	TSharedPtr<SWidget> CreateWidget(
+		UE::Editor::DataStorage::ICoreProvider* DataStorage,
+		UE::Editor::DataStorage::IUiProvider* DataStorageUi,
+		UE::Editor::DataStorage::RowHandle TargetRow,
+		UE::Editor::DataStorage::RowHandle WidgetRow,
+		const UE::Editor::DataStorage::FMetaDataView& Arguments) override;
+};
diff --git a/Engine/Plugins/Experimental/EditorDataStorageFeatures/Source/TedsDebugger/TedsDebugger.Build.cs b/Engine/Plugins/Experimental/EditorDataStorageFeatures/Source/TedsDebugger/TedsDebugger.Build.cs
--- a/Engine/Plugins/Experimental/EditorDataStorageFeatures/Source/TedsDebugger/TedsDebugger.Build.cs
+++ b/Engine/Plugins/Experimental/EditorDataStorageFeatures/Source/TedsDebugger/TedsDebugger.Build.cs
@@ -35,9 +35,9 @@ public class TedsDebugger : ModuleRules
 		            "InputCore",
 		            "EditorWidgets",
 		            "ToolWidgets",
-		            "EditorWidgets",
 		            "TedsTableViewer",
-		            "TedsQueryStack"
+		            "TedsQueryStack",
+		            "ApplicationCore"
             	});
 			
 			DynamicallyLoadedModuleNames.AddRange(new string[] {});
diff --git a/Engine/Source/Editor/LevelEditor/Private/SLevelEditor.cpp b/Engine/Source/Editor/LevelEditor/Private/SLevelEditor.cpp
--- a/Engine/Source/Editor/LevelEditor/Private/SLevelEditor.cpp
+++ b/Engine/Source/Editor/LevelEditor/Private/SLevelEditor.cpp
@@ -86,11 +86,16 @@
 #include "LevelEditorOutlinerSettings.h"
 #include "IWorldHierarchy.h"
 #include "LevelViewportActions.h"
+#include "Containers/VersePath.h"
 #include "Tools/EdModeInteractiveToolsContext.h"
 #include "Framework/Commands/GenericCommands.h"
 #include "ViewportInteractions/LevelViewportClickSelection.h"
 #include "ViewportInteractions/ViewportInteractionsBehaviorSource.h"
 #include "ViewportInteractions/ViewportInteractionsUtilities.h"
+#include "DataStorage/Scope/EditorDataScope.h"
+#include "DataStorage/Scope/EditorDataScopeColumns.h"
+#include "DataStorage/Features.h"
+#include "Elements/Columns/TypedElementLabelColumns.h"
 
 #define LOCTEXT_NAMESPACE "SLevelEditor"
 
@@ -295,9 +300,19 @@ void SLevelEditor::Construct( const SLevelEditor::FArguments& InArgs)
 			}
 			World->ShaderPlatformChanged(NewShaderPlatform, true);
 		});
-
+	// Create the level editor scope row, parented to the current context (root)
+	{
+		using namespace UE::Editor::DataStorage;
+		using namespace UE::Editor::DataStorage::Scope;
+		if (ICoreProvider* Storage = GetMutableDataStorageFeature<ICoreProvider>(StorageFeatureName))
+		{
+			LevelEditorScopeRow = Storage->AddScopeRow(TEXT("Level Editor"));
+			Storage->SetParentScope(LevelEditorScopeRow, GetCurrentScope());
+		}
+	}
 	FEditorDelegates::MapChange.AddRaw(this, &SLevelEditor::HandleEditorMapChange);
 	FEditorDelegates::OnAssetsDeleted.AddRaw(this, &SLevelEditor::HandleAssetsDeleted);
+
 	HandleEditorMapChange(MapChangeEventFlags::NewMap);
 
 	// Register the display names of the outliners
@@ -309,6 +324,9 @@ void SLevelEditor::Construct( const SLevelEditor::FArguments& InArgs)
 
 void SLevelEditor::Initialize( const TSharedRef<SDockTab>& OwnerTab, const TSharedRef<SWindow>& OwnerWindow )
 {
+	// Push the level editor context so child widgets constructed during Initialize see it
+	UE::Editor::DataStorage::Scope::FPushScopeGuard PushScopeGuard(LevelEditorScopeRow);
+
 	SelectedElements = NewObject<UTypedElementSelectionSet>(GetTransientPackage(), NAME_None, RF_Transactional);
 	SelectedElements->AddToRoot();
 
@@ -417,6 +435,16 @@ void SLevelEditor::ConstructTitleBarMessages()
 
 SLevelEditor::~SLevelEditor()
 {
+	// Clean up the scope row
+	{
+		using namespace UE::Editor::DataStorage;
+		if (ICoreProvider* Storage = GetMutableDataStorageFeature<ICoreProvider>(StorageFeatureName))
+		{
+			Storage->RemoveScopeRow(LevelEditorScopeRow);
+		}
+		LevelEditorScopeRow = InvalidRowHandle;
+	}
+
 	ActorDetailsSCSEditorUICustomization = nullptr;
 
 	// We're going away now, so make sure all toolkits that are hosted within this level editor are shut down
@@ -541,6 +569,9 @@ bool SLevelEditor::HasActivePlayInEditorViewport() const
 
 void SLevelEditor::Tick(const FGeometry& AllottedGeometry, const double InCurrentTime, const float InDeltaTime)
 {
+	// Push the level editor context so child widget ticks see it via TLS
+	UE::Editor::DataStorage::Scope::FPushScopeGuard PushScopeGuard(LevelEditorScopeRow);
+
 	// Update the ActiveViewport
 	TSharedPtr<SLevelViewport> LastActiveViewport = CachedActiveViewport.Pin();
 	if (!LastActiveViewport || &LastActiveViewport->GetLevelViewportClient() != GCurrentLevelEditingViewportClient)
@@ -1019,6 +1050,8 @@ FText SLevelEditor::GetSceneOutlinerLabel(FName SceneOutlinerTabIdentifier)
 /** Method to call when a tab needs to be spawned by the FLayoutService */
 TSharedRef<SDockTab> SLevelEditor::SpawnLevelEditorTab( const FSpawnTabArgs& Args, FName TabIdentifier, FString InitializationPayload )
 {
+	UE::Editor::DataStorage::Scope::FPushScopeGuard PushScopeGuard(LevelEditorScopeRow);
+	
 	if( TabIdentifier == LevelEditorTabIds::LevelEditorViewport)
 	{
 		return this->BuildViewportTab( NSLOCTEXT("LevelViewportTypes", "LevelEditorViewport", "Viewport 1"), TEXT("Viewport 1"), InitializationPayload );
@@ -2111,6 +2144,25 @@ void SLevelEditor::HandleEditorMapChange( uint32 MapChangeFlags )
 		WorldSettingsView->SetObject(GetWorld()->GetWorldSettings(), true);
 	}
 
+	// Update the editing verse scope on the scope row when a new map is loaded
+	if (MapChangeFlags == MapChangeEventFlags::NewMap)
+	{
+		using namespace UE::Editor::DataStorage;
+		if (ICoreProvider* Storage = GetMutableDataStorageFeature<ICoreProvider>(StorageFeatureName))
+		{
+			if (Storage->IsRowAvailable(LevelEditorScopeRow))
+			{
+				FEditingVerseScope Scope;
+				if (UWorld* EditorWorld = GetWorld())
+				{
+					UE::Core::FVersePath VersePath = EditorWorld->GetVersePath();
+					Scope.EditingScopeVersePath = VersePath.ToString();
+				}
+				Storage->SetScopeData<FEditingVerseScope>(LevelEditorScopeRow, MoveTemp(Scope));
+			}
+		}
+	}
+
 	// We want to recreate all the Scene Outliners on loading a new map
 	if (MapChangeFlags == MapChangeEventFlags::NewMap)
 	{
diff --git a/Engine/Source/Editor/LevelEditor/Private/SLevelEditor.h b/Engine/Source/Editor/LevelEditor/Private/SLevelEditor.h
--- a/Engine/Source/Editor/LevelEditor/Private/SLevelEditor.h
+++ b/Engine/Source/Editor/LevelEditor/Private/SLevelEditor.h
@@ -17,6 +17,7 @@
 #include "LevelViewportTabContent.h"
 #include "SLevelEditorToolBox.h"
 #include "ActorDetailsSCSEditorUICustomization.h"
+#include "DataStorage/Handles.h"
 
 class IAssetEditorInstance;
 class IDetailsView;
@@ -373,6 +374,9 @@ private:
 	/** A delegate which is called any time the LevelEditor's active viewport changes. */
 	FOnActiveViewportChanged OnActiveViewportChangedDelegate;
 
+	/** Context row for the Editor Scope Tree system */
+	UE::Editor::DataStorage::RowHandle LevelEditorScopeRow = UE::Editor::DataStorage::InvalidRowHandle;
+
 	/** If this flag is raised we will force refresh on next selection update. */
 	std::atomic<bool> bNeedsRefresh;
 
diff --git a/Engine/Source/Editor/SubobjectEditor/Private/SSubobjectBlueprintEditor.cpp b/Engine/Source/Editor/SubobjectEditor/Private/SSubobjectBlueprintEditor.cpp
--- a/Engine/Source/Editor/SubobjectEditor/Private/SSubobjectBlueprintEditor.cpp
+++ b/Engine/Source/Editor/SubobjectEditor/Private/SSubobjectBlueprintEditor.cpp
@@ -51,7 +51,9 @@ void SSubobjectBlueprintEditor::Construct(const FArguments& InArgs)
 	HideComponentClassCombo = InArgs._HideComponentClassCombo;
 	bAllowTreeUpdates = true;
 	UICustomization = InArgs._SCSEditorUICustomization;
-	
+
+	InitScopeRow(TEXT("Blueprint Subobject Editor"));
+
 	CreateCommandList();
 
 	// Build the tree widget
diff --git a/Engine/Source/Editor/SubobjectEditor/Private/SSubobjectEditor.cpp b/Engine/Source/Editor/SubobjectEditor/Private/SSubobjectEditor.cpp
--- a/Engine/Source/Editor/SubobjectEditor/Private/SSubobjectEditor.cpp
+++ b/Engine/Source/Editor/SubobjectEditor/Private/SSubobjectEditor.cpp
@@ -2,7 +2,10 @@
 
 #include "SSubobjectEditor.h"
 
+#include "DataStorage/Scope/EditorDataScope.h"
+#include "DataStorage/Features.h"
 #include "Editor.h"
+#include "Elements/Interfaces/TypedElementDataStorageInterface.h"
 #include "ClassIconFinder.h"
 #include "ScopedTransaction.h"
 #include "ComponentAssetBroker.h"		// FComponentAssetBrokerage
@@ -2129,6 +2132,30 @@ FMenuBuilder SSubobjectEditor::CreateMenuBuilder()
 	return EditBlueprintMenuBuilder;
 }
 
+void SSubobjectEditor::InitScopeRow(FStringView Label)
+{
+	using namespace UE::Editor::DataStorage;
+	using namespace UE::Editor::DataStorage::Scope;
+	if (ICoreProvider* Storage = GetMutableDataStorageFeature<ICoreProvider>(StorageFeatureName))
+	{
+		ScopeRow = Storage->AddScopeRow(Label);
+		Storage->SetParentScope(ScopeRow, GetCurrentScope());
+	}
+}
+
+SSubobjectEditor::~SSubobjectEditor()
+{
+	using namespace UE::Editor::DataStorage;
+	if (ICoreProvider* Storage = GetMutableDataStorageFeature<ICoreProvider>(StorageFeatureName))
+	{
+		if (Storage->IsRowAvailable(ScopeRow))
+		{
+			Storage->RemoveScopeRow(ScopeRow);
+		}
+	}
+	ScopeRow = InvalidRowHandle;
+}
+
 void SSubobjectEditor::ConstructTreeWidget()
 {
 	TSharedPtr<SHeaderRow> HeaderRow = SNew(SHeaderRow)
@@ -2370,6 +2397,10 @@ void SSubobjectEditor::UpdateSelectionFromNodes(const TArray<FSubobjectEditorTre
 {
 	bUpdatingSelection = true;
 
+	// Push our scope row so that downstream code (e.g. property visibility checks
+	// in the details panel) sees the correct editing scope via GetCurrentScope()
+	UE::Editor::DataStorage::Scope::FPushScopeGuard PushScopeGuard(ScopeRow);
+
 	// Notify that the selection has updated
 	OnSelectionUpdated.ExecuteIfBound(SelectedNodes);
 
diff --git a/Engine/Source/Editor/SubobjectEditor/Private/SSubobjectInstanceEditor.cpp b/Engine/Source/Editor/SubobjectEditor/Private/SSubobjectInstanceEditor.cpp
--- a/Engine/Source/Editor/SubobjectEditor/Private/SSubobjectInstanceEditor.cpp
+++ b/Engine/Source/Editor/SubobjectEditor/Private/SSubobjectInstanceEditor.cpp
@@ -44,7 +44,9 @@ void SSubobjectInstanceEditor::Construct(const FArguments& InArgs)
 	AllowEditing = InArgs._AllowEditing;
 	bAllowTreeUpdates = true;
 	UICustomization = InArgs._SCSEditorUICustomization;
-	
+
+	InitScopeRow(TEXT("Instance Subobject Editor"));
+
 	CreateCommandList();
 
 	// Build the tree widget
diff --git a/Engine/Source/Editor/SubobjectEditor/Public/SSubobjectEditor.h b/Engine/Source/Editor/SubobjectEditor/Public/SSubobjectEditor.h
--- a/Engine/Source/Editor/SubobjectEditor/Public/SSubobjectEditor.h
+++ b/Engine/Source/Editor/SubobjectEditor/Public/SSubobjectEditor.h
@@ -16,6 +16,7 @@
 #include "BlueprintEditor.h"		// FComponentEventConstructionData
 
 #include "ScopedTransaction.h"
+#include "DataStorage/Handles.h"
 
 #define UE_API SUBOBJECTEDITOR_API
 
@@ -407,17 +408,25 @@ public:
 	
 protected:
 	
-	// Do not allow public construction of this widget! 
+	// Do not allow public construction of this widget!
 	SSubobjectEditor() = default;
 
-	virtual ~SSubobjectEditor() = default;
-	
+	UE_API virtual ~SSubobjectEditor();
+
 	/** Delegate to invoke on selection update. */
 	FOnSelectionUpdated OnSelectionUpdated;
 
 	/** Delegate to invoke when an item in the tree is double clicked. */
 	FOnItemDoubleClicked OnItemDoubleClicked;
 
+	/** Create and parent the Editor Scope Tree row.  Call once from derived Construct(). */
+	UE_API void InitScopeRow(FStringView Label);
+
+	/** Context row owned by this subobject editor for Editor Scope Tree.
+	 *  Created by InitScopeRow(), parented to the current context at that time.
+	 *  Pushed onto TLS in event handlers so property visibility checks see the correct editing scope. */
+	UE::Editor::DataStorage::RowHandle ScopeRow = UE::Editor::DataStorage::InvalidRowHandle;
+
 	/** Attribute that provides access to the Object context for which we are viewing/editing. */
 	TAttribute<UObject*> ObjectContext;
 
diff --git a/Engine/Source/Runtime/TypedElementFramework/Private/DataStorage/Scope/EditorDataScope.cpp b/Engine/Source/Runtime/TypedElementFramework/Private/DataStorage/Scope/EditorDataScope.cpp
new file mode 100644
--- /dev/null
+++ b/Engine/Source/Runtime/TypedElementFramework/Private/DataStorage/Scope/EditorDataScope.cpp
@@ -0,0 +1,26 @@
+// Copyright Epic Games, Inc. All Rights Reserved.
+
+#include "DataStorage/Scope/EditorDataScope.h"
+
+namespace UE::Editor::DataStorage::Scope
+{
+
+// ============================================================================
+// Thread-Local Current Context
+// ============================================================================
+
+static thread_local RowHandle GCurrentScope = InvalidRowHandle;
+
+RowHandle GetCurrentScope()
+{
+	return GCurrentScope;
+}
+
+RowHandle SetCurrentScope(RowHandle NewContext)
+{
+	RowHandle Previous = GCurrentScope;
+	GCurrentScope = NewContext;
+	return Previous;
+}
+
+} // namespace UE::Editor::DataStorage::Scope
diff --git a/Engine/Source/Runtime/TypedElementFramework/Private/DataStorage/Scope/EditorDataScopeTypes.cpp b/Engine/Source/Runtime/TypedElementFramework/Private/DataStorage/Scope/EditorDataScopeTypes.cpp
new file mode 100644
--- /dev/null
+++ b/Engine/Source/Runtime/TypedElementFramework/Private/DataStorage/Scope/EditorDataScopeTypes.cpp
@@ -0,0 +1,33 @@
+// Copyright Epic Games, Inc. All Rights Reserved.
+
+#include "DataStorage/Scope/EditorDataScopeTypes.h"
+
+namespace UE::Editor::DataStorage::Scope
+{
+
+FScopeDataVersion::FScopeDataVersion() = default;
+
+bool FScopeDataVersion::operator==(const FScopeDataVersion& Other) const
+{
+	return SourceRow == Other.SourceRow && Version == Other.Version;
+}
+
+bool FScopeDataVersion::operator!=(const FScopeDataVersion& Other) const
+{
+	return !(*this == Other);
+}
+
+bool FScopeDataVersion::IsValid() const
+{
+	return SourceRow != InvalidRowHandle;
+}
+
+FScopeDataVersion FScopeDataVersion::Make(RowHandle InSourceRow, uint64 InVersion)
+{
+	FScopeDataVersion Result;
+	Result.SourceRow = InSourceRow;
+	Result.Version = InVersion;
+	return Result;
+}
+
+} // namespace UE::Editor::DataStorage::Scope
diff --git a/Engine/Source/Runtime/TypedElementFramework/Public/DataStorage/Scope/EditorDataScope.h b/Engine/Source/Runtime/TypedElementFramework/Public/DataStorage/Scope/EditorDataScope.h
new file mode 100644
--- /dev/null
+++ b/Engine/Source/Runtime/TypedElementFramework/Public/DataStorage/Scope/EditorDataScope.h
@@ -0,0 +1,127 @@
+// Copyright Epic Games, Inc. All Rights Reserved.
+
+#pragma once
+
+#include <type_traits>
+#include "DataStorage/Scope/EditorDataScopeTypes.h"
+#include "DataStorage/Handles.h"
+#include "Elements/Interfaces/TypedElementDataStorageInterface.h"
+
+/**
+ * Editor Scope Tree
+ * ====================================
+ *
+ * Provides a hierarchy-based context propagation mechanism built on TEDS.
+ * Editor components publish contextual data (selection, active object, etc.) at their level in a
+ * named hierarchy. Downstream consumers read contextual data by walking up the hierarchy from their
+ * own scope row until a matching column is found.
+ *
+ * Key concepts:
+ *   - Scope Row: A row representing an editor component's scope, created via ICoreProvider::AddScopeRow
+ *   - Thread-Local Current Context: An ambient scope row set per-thread, used as the default for lookups
+ *   - Hierarchy Walk: GetScopeData walks parent links to find the nearest ancestor with data
+ *   - Versioned: SetScopeData/RemoveScopeData bump a per-row version counter for change detection
+ *
+ * === Example: Panel Publishing Context ===
+ *   using namespace UE::Editor::DataStorage;
+ *   using namespace UE::Editor::DataStorage::Scope;
+ *   ICoreProvider& Storage = *GetMutableDataStorageFeature<ICoreProvider>(StorageFeatureName);
+ *   RowHandle PanelCtx = Storage.AddScopeRow();
+ *   Storage.SetParentScope(PanelCtx, ParentCtx);
+ *   Storage.SetScopeData<FSelectionColumn>(PanelCtx, FSelectionColumn{.Object = MyObj});
+ *
+ * === Example: Downstream Consumer ===
+ *   FPushScopeGuard Guard(PanelCtx);
+ *   // The widget's Construct would have CurrentContext
+ *   // be PanelCtx without knowing its parent.
+ *   SNew(<SomeSubWidget>);
+ *   // Works for non slate too.  ContextData is slate agnostic.
+ *
+ * === Example: Change Detection ===
+ *   FScopeDataVersion Cached;
+ *   void Tick() {
+ *       FScopeDataVersion V = Storage.GetScopeDataVersion<FSelectionColumn>(MyCtx);
+ *       if (V != Cached) { Cached = V; RefreshUI(); }
+ *   }
+ *   // No SlateAttribute bindings yet, but above _could_ be done within an attribute to check if
+ *   // dependent data changed and invalidate widget if so
+ */
+
+namespace UE::Editor::DataStorage::Scope
+{
+
+// ============================================================================
+// Thread-Local Current Context
+// ============================================================================
+
+/** Returns the thread-local current scope row handle. InvalidRowHandle if none set. */
+TYPEDELEMENTFRAMEWORK_API RowHandle GetCurrentScope();
+
+/** Sets the thread-local current scope row handle. Returns the previous value.
+ *  Prefer FPushScopeGuard for RAII usage. */
+TYPEDELEMENTFRAMEWORK_API RowHandle SetCurrentScope(RowHandle NewContext);
+
+/**
+ * RAII guard that pushes a new scope row as the thread-local current context.
+ * On destruction, restores the previous context. Supports nesting.
+ */
+class FPushScopeGuard
+{
+public:
+	explicit FPushScopeGuard(RowHandle NewContext)
+		: PreviousScope(SetCurrentScope(NewContext))
+	{
+	}
+
+	~FPushScopeGuard()
+	{
+		SetCurrentScope(PreviousScope);
+	}
+
+	// Non-copyable, non-movable
+	FPushScopeGuard(const FPushScopeGuard&) = delete;
+	FPushScopeGuard& operator=(const FPushScopeGuard&) = delete;
+	FPushScopeGuard(FPushScopeGuard&&) = delete;
+	FPushScopeGuard& operator=(FPushScopeGuard&&) = delete;
+
+private:
+	RowHandle PreviousScope;
+};
+
+// ============================================================================
+// TLS convenience overloads
+// ============================================================================
+// These free functions wrap ICoreProvider scope methods using the thread-local
+// current scope (GetCurrentScope()) as the row argument. They exist so callers
+// that operate in the ambient scope don't need to pass GetCurrentScope()
+// explicitly at every call site.
+
+/** Typed hierarchy-walking lookup using the thread-local current scope. */
+template<TColumnType T>
+const T* GetScopeData(ICoreProvider& Storage)
+{
+	return Storage.GetScopeData<T>(GetCurrentScope());
+}
+
+/** Hierarchy-walking version query using thread-local current scope. */
+template<TColumnType T>
+FScopeDataVersion GetScopeDataVersion(ICoreProvider& Storage)
+{
+	return Storage.GetScopeDataVersion<T>(GetCurrentScope());
+}
+
+/** Set scope data using thread-local current scope. */
+template<TColumnType T>
+FScopeDataVersion SetScopeData(ICoreProvider& Storage, T&& Data)
+{
+	return Storage.SetScopeData<T>(GetCurrentScope(), Forward<T>(Data));
+}
+
+/** Remove scope data using thread-local current scope. */
+template<TColumnType T>
+bool RemoveScopeData(ICoreProvider& Storage)
+{
+	return Storage.RemoveScopeData<T>(GetCurrentScope());
+}
+
+} // namespace UE::Editor::DataStorage::Scope
diff --git a/Engine/Source/Runtime/TypedElementFramework/Public/DataStorage/Scope/EditorDataScopeColumns.h b/Engine/Source/Runtime/TypedElementFramework/Public/DataStorage/Scope/EditorDataScopeColumns.h
new file mode 100644
--- /dev/null
+++ b/Engine/Source/Runtime/TypedElementFramework/Public/DataStorage/Scope/EditorDataScopeColumns.h
@@ -0,0 +1,53 @@
+// Copyright Epic Games, Inc. All Rights Reserved.
+
+#pragma once
+#include "DataStorage/CommonTypes.h"
+
+#include "EditorDataScopeColumns.generated.h"
+
+/**
+ * Tag applied to all scope rows created via AddScopeRow.
+ * Used to identify Editor Scope Tree rows in the TEDS debugger.
+ */
+USTRUCT(meta = (DisplayName = "Data Storage Scope Data"))
+struct FDataStorageScopeDataTag : public FEditorDataStorageTag
+{
+	GENERATED_BODY()
+};
+
+/**
+ * The VersePath of the asset/scope currently being edited by a tool.
+ * Published on scope rows by editors (Level Editor, Prefab Editor, etc.)
+ * so that downstream consumers can determine access-specifier visibility
+ * via FVersePathRegistry::CanAccess(ContextPath, TargetPath).
+ */
+USTRUCT(meta = (DisplayName = "Editing Verse Scope"))
+struct FEditingVerseScope : public FEditorDataStorageColumn
+{
+	GENERATED_BODY()
+
+	UPROPERTY()
+	FString EditingScopeVersePath;
+};
+
+/**
+ * Captures the callstack at the point where a scope row was created via AddScopeRow.
+ * Used by the Scope Hierarchy debugger tab to show where each scope row originated.
+ */
+USTRUCT(meta = (DisplayName = "Scope Row Source Info"))
+struct FScopeRowSourceInfo : public FEditorDataStorageColumn
+{
+	GENERATED_BODY()
+
+	static constexpr int32 MaxCallstackDepth = 64;
+
+	uint64 Callstack[MaxCallstackDepth] = {};
+	int32 CallstackDepth = 0;
+};
+
+namespace UE::Editor::DataStorage
+{
+	using FScopeDataTag = FDataStorageScopeDataTag;
+	using FEditingScopeColumn = FEditingVerseScope;
+	using FScopeSourceInfoColumn = FScopeRowSourceInfo;
+}
diff --git a/Engine/Source/Runtime/TypedElementFramework/Public/DataStorage/Scope/EditorDataScopeTypes.h b/Engine/Source/Runtime/TypedElementFramework/Public/DataStorage/Scope/EditorDataScopeTypes.h
new file mode 100644
--- /dev/null
+++ b/Engine/Source/Runtime/TypedElementFramework/Public/DataStorage/Scope/EditorDataScopeTypes.h
@@ -0,0 +1,37 @@
+// Copyright Epic Games, Inc. All Rights Reserved.
+
+#pragma once
+
+#include "DataStorage/Handles.h"
+
+class UScriptStruct;
+
+namespace UE::Editor::DataStorage::Scope
+{
+
+/**
+ * Represents a versioned snapshot of scope data for a specific column type on a row.
+ * Used for change detection: compare two FScopeDataVersion values to detect if scope
+ * data has changed (including removal).
+ *
+ * The version is opaque -- callers should only use operator==, operator!=, and IsValid().
+ */
+struct FScopeDataVersion
+{
+	TYPEDELEMENTFRAMEWORK_API FScopeDataVersion();
+
+	TYPEDELEMENTFRAMEWORK_API bool operator==(const FScopeDataVersion& Other) const;
+	TYPEDELEMENTFRAMEWORK_API bool operator!=(const FScopeDataVersion& Other) const;
+
+	/** Returns true if this version refers to a valid row (i.e., was obtained from a real operation). */
+	TYPEDELEMENTFRAMEWORK_API bool IsValid() const;
+
+	/** Factory for internal use by ICoreProvider implementations. */
+	TYPEDELEMENTFRAMEWORK_API static FScopeDataVersion Make(RowHandle InSourceRow, uint64 InVersion);
+
+private:
+	RowHandle SourceRow = InvalidRowHandle;
+	uint64 Version = 0;
+};
+
+} // namespace UE::Editor::DataStorage::Scope
diff --git a/Engine/Source/Runtime/TypedElementFramework/Public/DataStorage/Scope/EditorDataScopeVersion.h b/Engine/Source/Runtime/TypedElementFramework/Public/DataStorage/Scope/EditorDataScopeVersion.h
new file mode 100644
--- /dev/null
+++ b/Engine/Source/Runtime/TypedElementFramework/Public/DataStorage/Scope/EditorDataScopeVersion.h
@@ -0,0 +1,20 @@
+// Copyright Epic Games, Inc. All Rights Reserved.
+
+#pragma once
+
+#include "DataStorage/CommonTypes.h"
+
+#include "EditorDataScopeVersion.generated.h"
+
+/**
+ * Monotonic version counter for a scope row.
+ * Bumped on every SetScopeData or RemoveScopeData call targeting this row.
+ * Part of the scope table schema — every scope row has one automatically.
+ */
+USTRUCT(meta = (DisplayName = "ContextDataVersion"))
+struct FScopeDataVersionColumn : public FEditorDataStorageColumn
+{
+	GENERATED_BODY()
+
+	uint64 Version = 0;
+};
diff --git a/Engine/Source/Runtime/TypedElementFramework/Public/Elements/Interfaces/TypedElementDataStorageInterface.h b/Engine/Source/Runtime/TypedElementFramework/Public/Elements/Interfaces/TypedElementDataStorageInterface.h
--- a/Engine/Source/Runtime/TypedElementFramework/Public/Elements/Interfaces/TypedElementDataStorageInterface.h
+++ b/Engine/Source/Runtime/TypedElementFramework/Public/Elements/Interfaces/TypedElementDataStorageInterface.h
@@ -6,6 +6,7 @@
 #include "Containers/Array.h"
 #include "Containers/ContainersFwd.h"
 #include "DataStorage/CommonTypes.h"
+#include "DataStorage/Scope/EditorDataScopeTypes.h"
 #include "DataStorage/Handles.h"
 #include "DataStorage/MapKey.h"
 #include "DataStorage/Queries/Description.h"
@@ -465,7 +466,6 @@ public:
 	 * @return true if iteration completed (including if there were no children) without any VisitFn call returning
 	 *  false. Returns false if VisitFn stopped iteration early or there was an error.
 	 */
-	UE_EXPERIMENTAL(5.8, TEXT("This function signature might change"))
 	virtual bool IterateChildren(
 		FHierarchyHandle InHierarchyHandle,
 		RowHandle Row,
@@ -630,6 +630,74 @@ public:
 	virtual void ListExtensions(TFunctionRef<void(FName)> Callback) const = 0;
 
 
+	/**
+	 * @section Scope
+	 *
+	 * @description
+	 * Context rows form a hierarchy used to propagate contextual data to the tree of components that make up the editor.
+	 * Downstream consumers read scope data by walking up the hierarchy until a matching column is found.
+	 * It is expected that users interface with these rows via the functions provided in this section.
+	 */
+
+	/** Creates a new scope row with a default label. */
+	virtual RowHandle AddScopeRow() = 0;
+	/** Creates a new scope row with the given label. */
+	virtual RowHandle AddScopeRow(FStringView Label) = 0;
+	/** Creates a new scope row with the given label, parented to Parent. */
+	virtual RowHandle AddScopeRow(FStringView Label, RowHandle Parent) = 0;
+	/** Removes a scope row. */
+	virtual void RemoveScopeRow(RowHandle Row) = 0;
+	/** Sets the parent of Child to Parent in the scope hierarchy. */
+	virtual void SetParentScope(RowHandle Child, RowHandle Parent) = 0;
+	/** Returns the parent scope row for the given row in the scope hierarchy. */
+	virtual RowHandle GetParentScope(RowHandle Row) = 0;
+	/**
+	 * Walks up the scope hierarchy from Row, returning the first non-null column data
+	 * pointer for ColumnType. Returns nullptr if not found.
+	 */
+	virtual const void* GetScopeDataRaw(RowHandle Row, const UScriptStruct* ColumnType) = 0;
+	/**
+	 * Walks up the scope hierarchy from Row, returning the version of the first row that
+	 * holds ColumnType data. The version is a per-row counter incremented by any write
+	 * to any column on that row — not a per-column version. Returns an invalid version
+	 * if no row in the hierarchy holds ColumnType.
+	 */
+	virtual Scope::FScopeDataVersion GetScopeDataVersion(RowHandle Row, const UScriptStruct* ColumnType) = 0;
+	/**
+	 * Sets (creates if absent) scope data on a specific row.
+	 * The row's version counter is bumped (shared across all columns on this row).
+	 * Returns the new version.
+	 */
+	virtual Scope::FScopeDataVersion SetScopeData(RowHandle Row, const UScriptStruct* ColumnType,
+		ColumnCreationCallbackRef Initializer, ColumnCopyOrMoveCallback Relocator) = 0;
+	/**
+	 * Removes scope data from a specific row. The row's version counter is bumped.
+	 * Returns true if the column was present and removed.
+	 */
+	virtual bool RemoveScopeData(RowHandle Row, const UScriptStruct* ColumnType) = 0;
+	/**
+	 * Returns all columns visible from a scope row by walking up the hierarchy.
+	 * Closer rows shadow parents.
+	 */
+	virtual TArray<const UScriptStruct*> GetAllVisibleScopeColumns(RowHandle Row) = 0;
+	/** Returns the root scope row. InvalidRowHandle if not yet initialized. */
+	virtual RowHandle GetRootScope() const = 0;
+	/** Returns the hierarchy handle for the editor scope hierarchy. */
+	virtual FHierarchyHandle GetScopeHierarchy() = 0;
+
+	/** Typed hierarchy-walking lookup from a specific row. */
+	template<TColumnType T>
+	const T* GetScopeData(RowHandle Row);
+	/** Typed hierarchy-walking version query from a specific row. */
+	template<TColumnType T>
+	Scope::FScopeDataVersion GetScopeDataVersion(RowHandle Row);
+	/** Template convenience: set scope data on a specific row. */
+	template<TColumnType T>
+	Scope::FScopeDataVersion SetScopeData(RowHandle Row, T&& Data);
+	/** Template convenience: remove scope data from a specific row. */
+	template<TColumnType T>
+	bool RemoveScopeData(RowHandle Row);
+
 	/**
 	 * @section Deprecated
 	 */
@@ -998,4 +1066,52 @@ FQueryResult ICoreProvider::RunQuery(QueryHandle Query, ERunQueryFlags Flags, Qu
 	return RunQuery(Query, Flags, Queries::BuildQueryFunction(Result, Forward<Function>(Callback)));
 }
 
+// Scope template implementations
+
+template<TColumnType T>
+const T* ICoreProvider::GetScopeData(RowHandle Row)
+{
+	return static_cast<const T*>(GetScopeDataRaw(Row, T::StaticStruct()));
+}
+
+template<TColumnType T>
+Scope::FScopeDataVersion ICoreProvider::GetScopeDataVersion(RowHandle Row)
+{
+	return GetScopeDataVersion(Row, T::StaticStruct());
+}
+
+template<TColumnType T>
+Scope::FScopeDataVersion ICoreProvider::SetScopeData(RowHandle Row, T&& Data)
+{
+	return SetScopeData(Row, T::StaticStruct(),
+		[&Data](void* ColumnData, const UScriptStruct&)
+		{
+			if constexpr (std::is_move_constructible_v<T>)
+			{
+				new(ColumnData) T(MoveTemp(Data));
+			}
+			else
+			{
+				new(ColumnData) T(Data);
+			}
+		},
+		[](const UScriptStruct&, void* Destination, void* Source)
+		{
+			if constexpr (std::is_move_assignable_v<T>)
+			{
+				*static_cast<T*>(Destination) = MoveTemp(*static_cast<T*>(Source));
+			}
+			else
+			{
+				*static_cast<T*>(Destination) = *static_cast<T*>(Source);
+			}
+		});
+}
+
+template<TColumnType T>
+bool ICoreProvider::RemoveScopeData(RowHandle Row)
+{
+	return RemoveScopeData(Row, T::StaticStruct());
+}
+
 } // namespace UE::Editor::DataStorage
diff --git a/Engine/Source/Runtime/TypedElementFramework/Tests/Elements/Framework/EditorDataScopeTests.cpp b/Engine/Source/Runtime/TypedElementFramework/Tests/Elements/Framework/EditorDataScopeTests.cpp
new file mode 100644
--- /dev/null
+++ b/Engine/Source/Runtime/TypedElementFramework/Tests/Elements/Framework/EditorDataScopeTests.cpp
@@ -0,0 +1,818 @@
+// Copyright Epic Games, Inc. All Rights Reserved.
+
+#include "DataStorage/Scope/EditorDataScope.h"
+#include "DataStorage/Scope/EditorDataScopeColumns.h"
+#include "Elements/Framework/TypedElementTestColumns.h"
+#include "Elements/Interfaces/TypedElementDataStorageInterface.h"
+
+#if WITH_TESTS
+
+#include "DataStorage/Features.h"
+#include "Misc/AutomationTest.h"
+
+namespace UE::Editor::DataStorage::Scope::Tests
+{
+
+BEGIN_DEFINE_SPEC(FEditorDataScopeTestFixture,
+	"Editor.DataStorage.Scope",
+	EAutomationTestFlags::EditorContext | EAutomationTestFlags::EngineFilter)
+
+	ICoreProvider* Storage = nullptr;
+	TArray<RowHandle> CreatedRows;
+
+	RowHandle CreateRow()
+	{
+		RowHandle Row = Storage->AddScopeRow();
+		CreatedRows.Add(Row);
+		return Row;
+	}
+
+	void SetupParentChild(RowHandle Parent, RowHandle Child)
+	{
+		Storage->SetParentScope(Child, Parent);
+	}
+
+END_DEFINE_SPEC(FEditorDataScopeTestFixture)
+
+PRAGMA_DISABLE_EXPERIMENTAL_WARNINGS
+void FEditorDataScopeTestFixture::Define()
+{
+	BeforeEach([this]()
+	{
+		Storage = GetMutableDataStorageFeature<ICoreProvider>(StorageFeatureName);
+		TestTrue(TEXT("TEDS ICoreProvider must be available"), Storage != nullptr);
+		if (!Storage) return;
+
+		// Reset TLS context
+		SetCurrentScope(InvalidRowHandle);
+	});
+
+	AfterEach([this]()
+	{
+		if (Storage)
+		{
+			for (RowHandle Row : CreatedRows)
+			{
+				Storage->RemoveScopeRow(Row);
+			}
+		}
+		CreatedRows.Empty();
+		SetCurrentScope(InvalidRowHandle);
+	});
+
+	// ========================================================================
+	// TLS Tests
+	// ========================================================================
+
+	Describe("TLS", [this]
+	{
+		It(TEXT("DefaultIsInvalid"), EAsyncExecution::TaskGraphMainTick, [this]()
+		{
+			TestEqual(TEXT("Default current context should be InvalidRowHandle"),
+				GetCurrentScope(), InvalidRowHandle);
+		});
+
+		It(TEXT("SetAndGet"), EAsyncExecution::TaskGraphMainTick, [this]()
+		{
+			RowHandle Row = CreateRow();
+			SetCurrentScope(Row);
+			TestEqual(TEXT("GetCurrentScope should return the row we set"),
+				GetCurrentScope(), Row);
+			SetCurrentScope(InvalidRowHandle);
+		});
+
+		It(TEXT("ScopedPushRestore"), EAsyncExecution::TaskGraphMainTick, [this]()
+		{
+			RowHandle RowA = CreateRow();
+			RowHandle RowB = CreateRow();
+
+			SetCurrentScope(RowA);
+			{
+				FPushScopeGuard Guard(RowB);
+				TestEqual(TEXT("Inside scope: context should be RowB"),
+					GetCurrentScope(), RowB);
+			}
+			TestEqual(TEXT("After scope: context should be restored to RowA"),
+				GetCurrentScope(), RowA);
+			SetCurrentScope(InvalidRowHandle);
+		});
+
+		It(TEXT("NestedScoped"), EAsyncExecution::TaskGraphMainTick, [this]()
+		{
+			RowHandle RowA = CreateRow();
+			RowHandle RowB = CreateRow();
+			RowHandle RowC = CreateRow();
+
+			SetCurrentScope(RowA);
+			{
+				FPushScopeGuard Guard1(RowB);
+				TestEqual(TEXT("Level 1: context should be RowB"), GetCurrentScope(), RowB);
+				{
+					FPushScopeGuard Guard2(RowC);
+					TestEqual(TEXT("Level 2: context should be RowC"), GetCurrentScope(), RowC);
+				}
+				TestEqual(TEXT("Back to level 1: context should be RowB"), GetCurrentScope(), RowB);
+			}
+			TestEqual(TEXT("Back to original: context should be RowA"), GetCurrentScope(), RowA);
+			SetCurrentScope(InvalidRowHandle);
+		});
+	});
+
+	// ========================================================================
+	// Hierarchy Tests
+	// ========================================================================
+
+	Describe("Hierarchy", [this]
+	{
+		It(TEXT("AddScopeRow"), EAsyncExecution::TaskGraphMainTick, [this]()
+		{
+			RowHandle Row = Storage->AddScopeRow();
+			CreatedRows.Add(Row);
+			TestTrue(TEXT("AddScopeRow should return valid handle"),
+				Row != InvalidRowHandle);
+		});
+
+		It(TEXT("ParentChild"), EAsyncExecution::TaskGraphMainTick, [this]()
+		{
+			RowHandle Parent = CreateRow();
+			RowHandle Child = CreateRow();
+			Storage->SetParentScope(Child, Parent);
+
+			RowHandle FoundParent = Storage->GetParentScope(Child);
+			TestEqual(TEXT("GetParentScope should return the parent"),
+				FoundParent, Parent);
+		});
+	});
+
+	// ========================================================================
+	// Scope Data Lookup Tests
+	// ========================================================================
+
+	Describe("Lookup", [this]
+	{
+		It(TEXT("DirectColumn"), EAsyncExecution::TaskGraphMainTick, [this]()
+		{
+			RowHandle Row = CreateRow();
+			Storage->SetScopeData<FTestColumnInt>(Row, FTestColumnInt{.TestInt = 42});
+
+			const FTestColumnInt* Data = Storage->GetScopeData<FTestColumnInt>(Row);
+			TestNotNull(TEXT("Should find column on same row"), Data);
+			if (Data)
+			{
+				TestEqual(TEXT("Column value should match"), Data->TestInt, 42);
+			}
+		});
+
+		It(TEXT("InheritedColumn"), EAsyncExecution::TaskGraphMainTick, [this]()
+		{
+			RowHandle Parent = CreateRow();
+			RowHandle Child = CreateRow();
+			SetupParentChild(Parent, Child);
+
+			Storage->SetScopeData<FTestColumnInt>(Parent, FTestColumnInt{.TestInt = 100});
+
+			const FTestColumnInt* Data = Storage->GetScopeData<FTestColumnInt>(Child);
+			TestNotNull(TEXT("Should find inherited column from parent"), Data);
+			if (Data)
+			{
+				TestEqual(TEXT("Inherited value should match parent"), Data->TestInt, 100);
+			}
+		});
+
+		It(TEXT("OverrideColumn"), EAsyncExecution::TaskGraphMainTick, [this]()
+		{
+			RowHandle Parent = CreateRow();
+			RowHandle Child = CreateRow();
+			SetupParentChild(Parent, Child);
+
+			Storage->SetScopeData<FTestColumnInt>(Parent, FTestColumnInt{.TestInt = 100});
+			Storage->SetScopeData<FTestColumnInt>(Child, FTestColumnInt{.TestInt = 200});
+
+			const FTestColumnInt* Data = Storage->GetScopeData<FTestColumnInt>(Child);
+			TestNotNull(TEXT("Should find overridden column"), Data);
+			if (Data)
+			{
+				TestEqual(TEXT("Should return child value, not parent"), Data->TestInt, 200);
+			}
+		});
+
+		It(TEXT("MissingColumn"), EAsyncExecution::TaskGraphMainTick, [this]()
+		{
+			RowHandle Row = CreateRow();
+			const FTestColumnString* Data = Storage->GetScopeData<FTestColumnString>(Row);
+			TestNull(TEXT("Missing column should return nullptr"), Data);
+		});
+
+		It(TEXT("UsesCurrentScope"), EAsyncExecution::TaskGraphMainTick, [this]()
+		{
+			RowHandle Row = CreateRow();
+			Storage->SetScopeData<FTestColumnInt>(Row, FTestColumnInt{.TestInt = 55});
+
+			SetCurrentScope(Row);
+			const FTestColumnInt* Data = GetScopeData<FTestColumnInt>(*Storage);
+			TestNotNull(TEXT("Should find data via current context"), Data);
+			if (Data)
+			{
+				TestEqual(TEXT("Value should match"), Data->TestInt, 55);
+			}
+			SetCurrentScope(InvalidRowHandle);
+		});
+
+		It(TEXT("ThreeLevelInheritance"), EAsyncExecution::TaskGraphMainTick, [this]()
+		{
+			RowHandle Grandparent = CreateRow();
+			RowHandle Parent = CreateRow();
+			RowHandle Child = CreateRow();
+			SetupParentChild(Grandparent, Parent);
+			SetupParentChild(Parent, Child);
+
+			Storage->SetScopeData<FTestColumnInt>(Grandparent, FTestColumnInt{.TestInt = 999});
+
+			const FTestColumnInt* Data = Storage->GetScopeData<FTestColumnInt>(Child);
+			TestNotNull(TEXT("Should find column from grandparent"), Data);
+			if (Data)
+			{
+				TestEqual(TEXT("Should be grandparent value"), Data->TestInt, 999);
+			}
+		});
+	});
+
+	// ========================================================================
+	// Versioning Tests
+	// ========================================================================
+
+	Describe("Version", [this]
+	{
+		It(TEXT("SetReturnsValid"), EAsyncExecution::TaskGraphMainTick, [this]()
+		{
+			RowHandle Row = CreateRow();
+			FScopeDataVersion V = Storage->SetScopeData<FTestColumnInt>(Row, FTestColumnInt{.TestInt = 1});
+			TestTrue(TEXT("Version from SetScopeData should be valid"), V.IsValid());
+		});
+
+		It(TEXT("SetIncrements"), EAsyncExecution::TaskGraphMainTick, [this]()
+		{
+			RowHandle Row = CreateRow();
+			FScopeDataVersion V1 = Storage->SetScopeData<FTestColumnInt>(Row, FTestColumnInt{.TestInt = 1});
+			FScopeDataVersion V2 = Storage->SetScopeData<FTestColumnInt>(Row, FTestColumnInt{.TestInt = 2});
+			TestTrue(TEXT("Two updates should return different versions"), V1 != V2);
+		});
+
+		It(TEXT("ReadVersionMatches"), EAsyncExecution::TaskGraphMainTick, [this]()
+		{
+			RowHandle Row = CreateRow();
+			FScopeDataVersion V1 = Storage->SetScopeData<FTestColumnInt>(Row, FTestColumnInt{.TestInt = 1});
+			FScopeDataVersion V2 = Storage->GetScopeDataVersion<FTestColumnInt>(Row);
+			TestTrue(TEXT("GetScopeDataVersion should match last update"), V1 == V2);
+		});
+
+		It(TEXT("DeleteReturnsTrue"), EAsyncExecution::TaskGraphMainTick, [this]()
+		{
+			RowHandle Row = CreateRow();
+			Storage->SetScopeData<FTestColumnInt>(Row, FTestColumnInt{.TestInt = 1});
+			bool Deleted = Storage->RemoveScopeData<FTestColumnInt>(Row);
+			TestTrue(TEXT("RemoveScopeData should return true when column existed"), Deleted);
+		});
+
+		It(TEXT("DeleteReturnsFalse"), EAsyncExecution::TaskGraphMainTick, [this]()
+		{
+			RowHandle Row = CreateRow();
+			bool Deleted = Storage->RemoveScopeData<FTestColumnInt>(Row);
+			TestFalse(TEXT("RemoveScopeData should return false when column absent"), Deleted);
+		});
+
+		It(TEXT("DeleteClearsVersion"), EAsyncExecution::TaskGraphMainTick, [this]()
+		{
+			RowHandle Row = CreateRow();
+			Storage->SetScopeData<FTestColumnInt>(Row, FTestColumnInt{.TestInt = 1});
+			Storage->RemoveScopeData<FTestColumnInt>(Row);
+
+			FScopeDataVersion V = Storage->GetScopeDataVersion<FTestColumnInt>(Row);
+			TestFalse(TEXT("Version should be invalid after delete (column gone, walk finds nothing)"), V.IsValid());
+		});
+
+		It(TEXT("DeleteVersionDiffers"), EAsyncExecution::TaskGraphMainTick, [this]()
+		{
+			RowHandle Row = CreateRow();
+			FScopeDataVersion VUpdate = Storage->SetScopeData<FTestColumnInt>(Row, FTestColumnInt{.TestInt = 1});
+			Storage->RemoveScopeData<FTestColumnInt>(Row);
+			FScopeDataVersion VDelete = Storage->GetScopeDataVersion<FTestColumnInt>(Row);
+
+			TestTrue(TEXT("Version after delete should differ from update version"), VUpdate != VDelete);
+		});
+
+		It(TEXT("SetAfterDelete"), EAsyncExecution::TaskGraphMainTick, [this]()
+		{
+			RowHandle Row = CreateRow();
+			FScopeDataVersion V1 = Storage->SetScopeData<FTestColumnInt>(Row, FTestColumnInt{.TestInt = 1});
+			Storage->RemoveScopeData<FTestColumnInt>(Row);
+			FScopeDataVersion V2 = Storage->GetScopeDataVersion<FTestColumnInt>(Row);
+			FScopeDataVersion V3 = Storage->SetScopeData<FTestColumnInt>(Row, FTestColumnInt{.TestInt = 2});
+
+			TestTrue(TEXT("V1 != V2 (update vs delete)"), V1 != V2);
+			TestTrue(TEXT("V2 != V3 (delete vs re-update)"), V2 != V3);
+			TestTrue(TEXT("V3 should be valid"), V3.IsValid());
+
+			// Verify the data is actually back
+			const FTestColumnInt* Data = Storage->GetScopeData<FTestColumnInt>(Row);
+			TestNotNull(TEXT("Column should exist after re-update"), Data);
+			if (Data)
+			{
+				TestEqual(TEXT("Re-updated value should match"), Data->TestInt, 2);
+			}
+		});
+
+		It(TEXT("NeverSetIsInvalid"), EAsyncExecution::TaskGraphMainTick, [this]()
+		{
+			RowHandle Row = CreateRow();
+			FScopeDataVersion V = Storage->GetScopeDataVersion<FTestColumnInt>(Row);
+			TestFalse(TEXT("Version for never-set column should be invalid"), V.IsValid());
+		});
+	});
+
+	// ========================================================================
+	// Hierarchy + Versioning Integration Tests
+	// ========================================================================
+
+	Describe("HierVersion", [this]
+	{
+		It(TEXT("InheritedVersion"), EAsyncExecution::TaskGraphMainTick, [this]()
+		{
+			RowHandle Parent = CreateRow();
+			RowHandle Child = CreateRow();
+			SetupParentChild(Parent, Child);
+
+			FScopeDataVersion VParent = Storage->SetScopeData<FTestColumnInt>(Parent, FTestColumnInt{.TestInt = 10});
+			FScopeDataVersion VFromChild = Storage->GetScopeDataVersion<FTestColumnInt>(Child);
+
+			TestTrue(TEXT("Hierarchy version from child should be valid"), VFromChild.IsValid());
+			TestTrue(TEXT("Hierarchy version should match parent's version"), VParent == VFromChild);
+		});
+
+		It(TEXT("OverriddenVersion"), EAsyncExecution::TaskGraphMainTick, [this]()
+		{
+			RowHandle Parent = CreateRow();
+			RowHandle Child = CreateRow();
+			SetupParentChild(Parent, Child);
+
+			Storage->SetScopeData<FTestColumnInt>(Parent, FTestColumnInt{.TestInt = 10});
+			FScopeDataVersion VChild = Storage->SetScopeData<FTestColumnInt>(Child, FTestColumnInt{.TestInt = 20});
+
+			FScopeDataVersion VFromChild = Storage->GetScopeDataVersion<FTestColumnInt>(Child);
+			TestTrue(TEXT("Should get child's version, not parent's"), VChild == VFromChild);
+		});
+
+		It(TEXT("DeleteOnChildFallsToParent"), EAsyncExecution::TaskGraphMainTick, [this]()
+		{
+			RowHandle Parent = CreateRow();
+			RowHandle Child = CreateRow();
+			SetupParentChild(Parent, Child);
+
+			FScopeDataVersion VParent = Storage->SetScopeData<FTestColumnInt>(Parent, FTestColumnInt{.TestInt = 10});
+			Storage->SetScopeData<FTestColumnInt>(Child, FTestColumnInt{.TestInt = 20});
+			Storage->RemoveScopeData<FTestColumnInt>(Child);
+
+			// The child's version map entry was removed on delete, so the hierarchy walk
+			// falls through to the parent.
+			FScopeDataVersion VFromChild = Storage->GetScopeDataVersion<FTestColumnInt>(Child);
+			TestTrue(TEXT("After delete, hierarchy version should still be valid"), VFromChild.IsValid());
+
+			// The actual data lookup should now fall through to parent
+			const FTestColumnInt* Data = Storage->GetScopeData<FTestColumnInt>(Child);
+			TestNotNull(TEXT("GetScopeData should fall through to parent"), Data);
+			if (Data)
+			{
+				TestEqual(TEXT("Data should come from parent"), Data->TestInt, 10);
+			}
+		});
+	});
+
+	// ========================================================================
+	// Integration Tests
+	// ========================================================================
+
+	Describe("Integration", [this]
+	{
+		It(TEXT("FullWorkflow"), EAsyncExecution::TaskGraphMainTick, [this]()
+		{
+			// Build 3-level hierarchy: Root -> Mid -> Leaf
+			RowHandle Root = CreateRow();
+			RowHandle Mid = CreateRow();
+			RowHandle Leaf = CreateRow();
+			SetupParentChild(Root, Mid);
+			SetupParentChild(Mid, Leaf);
+
+			// Publish on root
+			FScopeDataVersion V1 = Storage->SetScopeData<FTestColumnInt>(Root, FTestColumnInt{.TestInt = 42});
+			TestTrue(TEXT("V1 should be valid"), V1.IsValid());
+
+			// Read from leaf
+			const FTestColumnInt* Data = Storage->GetScopeData<FTestColumnInt>(Leaf);
+			TestNotNull(TEXT("Leaf should inherit from root"), Data);
+			if (Data)
+			{
+				TestEqual(TEXT("Inherited value should be 42"), Data->TestInt, 42);
+			}
+
+			// Detect version change on update
+			FScopeDataVersion V2 = Storage->SetScopeData<FTestColumnInt>(Root, FTestColumnInt{.TestInt = 99});
+			TestTrue(TEXT("Version should change after update"), V1 != V2);
+
+			// Verify new value propagates
+			Data = Storage->GetScopeData<FTestColumnInt>(Leaf);
+			TestNotNull(TEXT("Leaf should still inherit"), Data);
+			if (Data)
+			{
+				TestEqual(TEXT("Updated value should be 99"), Data->TestInt, 99);
+			}
+		});
+
+		It(TEXT("ScopedCallStack"), EAsyncExecution::TaskGraphMainTick, [this]()
+		{
+			RowHandle RowA = CreateRow();
+			RowHandle RowB = CreateRow();
+			Storage->SetScopeData<FTestColumnInt>(RowA, FTestColumnInt{.TestInt = 1});
+			Storage->SetScopeData<FTestColumnInt>(RowB, FTestColumnInt{.TestInt = 2});
+
+			{
+				FPushScopeGuard GuardA(RowA);
+				const FTestColumnInt* DataA = GetScopeData<FTestColumnInt>(*Storage);
+				TestNotNull(TEXT("Level A: should find data"), DataA);
+				if (DataA)
+				{
+					TestEqual(TEXT("Level A value"), DataA->TestInt, 1);
+				}
+
+				{
+					FPushScopeGuard GuardB(RowB);
+					const FTestColumnInt* DataB = GetScopeData<FTestColumnInt>(*Storage);
+					TestNotNull(TEXT("Level B: should find data"), DataB);
+					if (DataB)
+					{
+						TestEqual(TEXT("Level B value"), DataB->TestInt, 2);
+					}
+				}
+
+				const FTestColumnInt* DataAfter = GetScopeData<FTestColumnInt>(*Storage);
+				TestNotNull(TEXT("After B scope: should find A's data"), DataAfter);
+				if (DataAfter)
+				{
+					TestEqual(TEXT("Should be A's value again"), DataAfter->TestInt, 1);
+				}
+			}
+		});
+
+		It(TEXT("MultipleColumnTypes"), EAsyncExecution::TaskGraphMainTick, [this]()
+		{
+			RowHandle Parent = CreateRow();
+			RowHandle Child = CreateRow();
+			SetupParentChild(Parent, Child);
+
+			Storage->SetScopeData<FTestColumnInt>(Parent, FTestColumnInt{.TestInt = 10});
+			FTestColumnString StrColumn;
+			StrColumn.TestString = TEXT("hello");
+			Storage->SetScopeData<FTestColumnString>(Child, MoveTemp(StrColumn));
+
+			const FTestColumnInt* IntData = Storage->GetScopeData<FTestColumnInt>(Child);
+			const FTestColumnString* StrData = Storage->GetScopeData<FTestColumnString>(Child);
+
+			TestNotNull(TEXT("Int from parent should be found"), IntData);
+			TestNotNull(TEXT("String from child should be found"), StrData);
+			if (IntData)
+			{
+				TestEqual(TEXT("Int value from parent"), IntData->TestInt, 10);
+			}
+			if (StrData)
+			{
+				TestEqual(TEXT("String value from child"), StrData->TestString, FString(TEXT("hello")));
+			}
+		});
+
+		It(TEXT("DeleteAndReAdd"), EAsyncExecution::TaskGraphMainTick, [this]()
+		{
+			RowHandle Row = CreateRow();
+			FScopeDataVersion V1 = Storage->SetScopeData<FTestColumnInt>(Row, FTestColumnInt{.TestInt = 1});
+			bool Deleted = Storage->RemoveScopeData<FTestColumnInt>(Row);
+			TestTrue(TEXT("Delete should succeed"), Deleted);
+			FScopeDataVersion V2 = Storage->GetScopeDataVersion<FTestColumnInt>(Row);
+			FScopeDataVersion V3 = Storage->SetScopeData<FTestColumnInt>(Row, FTestColumnInt{.TestInt = 3});
+
+			TestTrue(TEXT("V1 != V2"), V1 != V2);
+			TestTrue(TEXT("V2 != V3"), V2 != V3);
+			TestTrue(TEXT("V1 != V3"), V1 != V3);
+		});
+	});
+
+	// ========================================================================
+	// FEditingVerseScope Tests
+	// ========================================================================
+
+	Describe("EditingVerseScope", [this]
+	{
+		It(TEXT("SetAndReadOnSameRow"), EAsyncExecution::TaskGraphMainTick, [this]()
+		{
+			RowHandle Row = CreateRow();
+			FEditingVerseScope Scope;
+			Scope.EditingScopeVersePath = TEXT("/TestDomain/TestScope");
+			Storage->SetScopeData<FEditingVerseScope>(Row, MoveTemp(Scope));
+
+			const FEditingVerseScope* Data = Storage->GetScopeData<FEditingVerseScope>(Row);
+			TestNotNull(TEXT("Should find FEditingVerseScope on same row"), Data);
+			if (Data)
+			{
+				TestEqual(TEXT("VersePath should match"), Data->EditingScopeVersePath, FString(TEXT("/TestDomain/TestScope")));
+			}
+		});
+
+		It(TEXT("InheritedFromParent"), EAsyncExecution::TaskGraphMainTick, [this]()
+		{
+			RowHandle Parent = CreateRow();
+			RowHandle Child = CreateRow();
+			SetupParentChild(Parent, Child);
+
+			FEditingVerseScope Scope;
+			Scope.EditingScopeVersePath = TEXT("/Game/Maps/TestLevel");
+			Storage->SetScopeData<FEditingVerseScope>(Parent, MoveTemp(Scope));
+
+			const FEditingVerseScope* Data = Storage->GetScopeData<FEditingVerseScope>(Child);
+			TestNotNull(TEXT("Child should inherit FEditingVerseScope from parent"), Data);
+			if (Data)
+			{
+				TestEqual(TEXT("Inherited path should match parent"),
+					Data->EditingScopeVersePath, FString(TEXT("/Game/Maps/TestLevel")));
+			}
+		});
+
+		It(TEXT("ChildOverridesParent"), EAsyncExecution::TaskGraphMainTick, [this]()
+		{
+			RowHandle Root = CreateRow();
+			RowHandle Child = CreateRow();
+			SetupParentChild(Root, Child);
+
+			FEditingVerseScope RootScope;
+			RootScope.EditingScopeVersePath = TEXT("/Game/Maps/TestLevel");
+			Storage->SetScopeData<FEditingVerseScope>(Root, MoveTemp(RootScope));
+
+			FEditingVerseScope ChildScope;
+			ChildScope.EditingScopeVersePath = TEXT("/Game/Prefabs/TestPrefab");
+			Storage->SetScopeData<FEditingVerseScope>(Child, MoveTemp(ChildScope));
+
+			const FEditingVerseScope* Data = Storage->GetScopeData<FEditingVerseScope>(Child);
+			TestNotNull(TEXT("Should find overridden FEditingVerseScope"), Data);
+			if (Data)
+			{
+				TestEqual(TEXT("Should return child scope, not parent"),
+					Data->EditingScopeVersePath, FString(TEXT("/Game/Prefabs/TestPrefab")));
+			}
+
+			// Parent should still have its own scope
+			const FEditingVerseScope* ParentData = Storage->GetScopeData<FEditingVerseScope>(Root);
+			TestNotNull(TEXT("Parent should still have its scope"), ParentData);
+			if (ParentData)
+			{
+				TestEqual(TEXT("Parent scope should be unchanged"),
+					ParentData->EditingScopeVersePath, FString(TEXT("/Game/Maps/TestLevel")));
+			}
+		});
+
+		It(TEXT("ReadViaTLS"), EAsyncExecution::TaskGraphMainTick, [this]()
+		{
+			RowHandle Row = CreateRow();
+			FEditingVerseScope Scope;
+			Scope.EditingScopeVersePath = TEXT("/TLS/TestScope");
+			Storage->SetScopeData<FEditingVerseScope>(Row, MoveTemp(Scope));
+
+			FPushScopeGuard Guard(Row);
+			const FEditingVerseScope* Data = GetScopeData<FEditingVerseScope>(*Storage);
+			TestNotNull(TEXT("Should find FEditingVerseScope via TLS current context"), Data);
+			if (Data)
+			{
+				TestEqual(TEXT("TLS-based read should return correct path"),
+					Data->EditingScopeVersePath, FString(TEXT("/TLS/TestScope")));
+			}
+		});
+
+		It(TEXT("ThreeLevelHierarchy"), EAsyncExecution::TaskGraphMainTick, [this]()
+		{
+			// Simulate Root -> LevelEditor -> PrefabEditor hierarchy
+			RowHandle Root = CreateRow();
+			RowHandle LevelEditor = CreateRow();
+			RowHandle PrefabEditor = CreateRow();
+			SetupParentChild(Root, LevelEditor);
+			SetupParentChild(LevelEditor, PrefabEditor);
+
+			FEditingVerseScope LevelScope;
+			LevelScope.EditingScopeVersePath = TEXT("/Game/Maps/MyLevel");
+			Storage->SetScopeData<FEditingVerseScope>(LevelEditor, MoveTemp(LevelScope));
+
+			FEditingVerseScope PrefabScope;
+			PrefabScope.EditingScopeVersePath = TEXT("/Game/Prefabs/MyPrefab");
+			Storage->SetScopeData<FEditingVerseScope>(PrefabEditor, MoveTemp(PrefabScope));
+
+			// PrefabEditor should see its own scope
+			const FEditingVerseScope* PrefabData = Storage->GetScopeData<FEditingVerseScope>(PrefabEditor);
+			TestNotNull(TEXT("PrefabEditor should have its own scope"), PrefabData);
+			if (PrefabData)
+			{
+				TestEqual(TEXT("PrefabEditor scope"), PrefabData->EditingScopeVersePath, FString(TEXT("/Game/Prefabs/MyPrefab")));
+			}
+
+			// LevelEditor should see its own scope
+			const FEditingVerseScope* LevelData = Storage->GetScopeData<FEditingVerseScope>(LevelEditor);
+			TestNotNull(TEXT("LevelEditor should have its own scope"), LevelData);
+			if (LevelData)
+			{
+				TestEqual(TEXT("LevelEditor scope"), LevelData->EditingScopeVersePath, FString(TEXT("/Game/Maps/MyLevel")));
+			}
+
+			// Root has no scope set -- should return nullptr
+			const FEditingVerseScope* RootData = Storage->GetScopeData<FEditingVerseScope>(Root);
+			TestNull(TEXT("Root should have no FEditingVerseScope"), RootData);
+		});
+	});
+
+	// ========================================================================
+	// Edge Cases
+	// ========================================================================
+
+	Describe("Edge", [this]
+	{
+		It(TEXT("RowNotInHierarchy"), EAsyncExecution::TaskGraphMainTick, [this]()
+		{
+			RowHandle Row = CreateRow();
+			Storage->SetScopeData<FTestColumnInt>(Row, FTestColumnInt{.TestInt = 77});
+
+			// Row is not part of the scope hierarchy, should still check its own columns
+			const FTestColumnInt* Data = Storage->GetScopeData<FTestColumnInt>(Row);
+			TestNotNull(TEXT("Should find data on the row itself"), Data);
+			if (Data)
+			{
+				TestEqual(TEXT("Value should match"), Data->TestInt, 77);
+			}
+		});
+
+		It(TEXT("EmptyHierarchy"), EAsyncExecution::TaskGraphMainTick, [this]()
+		{
+			RowHandle Row = CreateRow();
+			// No parent, no data
+			const FTestColumnInt* Data = Storage->GetScopeData<FTestColumnInt>(Row);
+			TestNull(TEXT("No data on isolated row should return nullptr"), Data);
+		});
+	});
+
+	// ========================================================================
+	// InvalidRowHandle Robustness
+	// ========================================================================
+
+	Describe("InvalidRow", [this]
+	{
+		It(TEXT("GetScopeData"), EAsyncExecution::TaskGraphMainTick, [this]()
+		{
+			const FTestColumnInt* Data = Storage->GetScopeData<FTestColumnInt>(InvalidRowHandle);
+			TestNull(TEXT("GetScopeData on InvalidRowHandle should return nullptr"), Data);
+		});
+
+		It(TEXT("GetScopeDataVersion"), EAsyncExecution::TaskGraphMainTick, [this]()
+		{
+			FScopeDataVersion V = Storage->GetScopeDataVersion<FTestColumnInt>(InvalidRowHandle);
+			TestFalse(TEXT("GetScopeDataVersion on InvalidRowHandle should return invalid version"), V.IsValid());
+		});
+
+		It(TEXT("SetScopeData"), EAsyncExecution::TaskGraphMainTick, [this]()
+		{
+			FScopeDataVersion V = Storage->SetScopeData<FTestColumnInt>(InvalidRowHandle, FTestColumnInt{.TestInt = 1});
+			TestFalse(TEXT("SetScopeData on InvalidRowHandle should return invalid version"), V.IsValid());
+		});
+
+		It(TEXT("RemoveScopeData"), EAsyncExecution::TaskGraphMainTick, [this]()
+		{
+			bool Deleted = Storage->RemoveScopeData<FTestColumnInt>(InvalidRowHandle);
+			TestFalse(TEXT("RemoveScopeData on InvalidRowHandle should return false"), Deleted);
+		});
+
+		It(TEXT("GetParentScope"), EAsyncExecution::TaskGraphMainTick, [this]()
+		{
+			RowHandle Parent = Storage->GetParentScope(InvalidRowHandle);
+			TestEqual(TEXT("GetParentScope on InvalidRowHandle should return InvalidRowHandle"),
+				Parent, InvalidRowHandle);
+		});
+
+		It(TEXT("SetParentScope"), EAsyncExecution::TaskGraphMainTick, [this]()
+		{
+			// Should not crash
+			Storage->SetParentScope(InvalidRowHandle, InvalidRowHandle);
+		});
+
+		It(TEXT("RemoveScopeRow"), EAsyncExecution::TaskGraphMainTick, [this]()
+		{
+			// Should not crash
+			Storage->RemoveScopeRow(InvalidRowHandle);
+		});
+
+		It(TEXT("GetAllVisibleScopeColumns"), EAsyncExecution::TaskGraphMainTick, [this]()
+		{
+			TArray<const UScriptStruct*> Columns = Storage->GetAllVisibleScopeColumns(InvalidRowHandle);
+			TestEqual(TEXT("GetAllVisibleScopeColumns on InvalidRowHandle should return empty"),
+				Columns.Num(), 0);
+		});
+
+		It(TEXT("ReadViaTLSWhenUnset"), EAsyncExecution::TaskGraphMainTick, [this]()
+		{
+			SetCurrentScope(InvalidRowHandle);
+			const FTestColumnInt* Data = GetScopeData<FTestColumnInt>(*Storage);
+			TestNull(TEXT("GetScopeData via unset TLS should return nullptr"), Data);
+		});
+	});
+
+	// ========================================================================
+	// Destroyed Row Robustness
+	// ========================================================================
+
+	Describe("DestroyedRow", [this]
+	{
+		It(TEXT("GetScopeData"), EAsyncExecution::TaskGraphMainTick, [this]()
+		{
+			RowHandle Row = Storage->AddScopeRow();
+			Storage->SetScopeData<FTestColumnInt>(Row, FTestColumnInt{.TestInt = 42});
+			Storage->RemoveScopeRow(Row);
+
+			const FTestColumnInt* Data = Storage->GetScopeData<FTestColumnInt>(Row);
+			TestNull(TEXT("GetScopeData on destroyed row should return nullptr"), Data);
+		});
+
+		It(TEXT("GetScopeDataVersion"), EAsyncExecution::TaskGraphMainTick, [this]()
+		{
+			RowHandle Row = Storage->AddScopeRow();
+			Storage->SetScopeData<FTestColumnInt>(Row, FTestColumnInt{.TestInt = 1});
+			Storage->RemoveScopeRow(Row);
+
+			FScopeDataVersion V = Storage->GetScopeDataVersion<FTestColumnInt>(Row);
+			TestFalse(TEXT("GetScopeDataVersion on destroyed row should return invalid version"), V.IsValid());
+		});
+
+		It(TEXT("SetScopeData"), EAsyncExecution::TaskGraphMainTick, [this]()
+		{
+			RowHandle Row = Storage->AddScopeRow();
+			Storage->RemoveScopeRow(Row);
+
+			FScopeDataVersion V = Storage->SetScopeData<FTestColumnInt>(Row, FTestColumnInt{.TestInt = 1});
+			TestFalse(TEXT("SetScopeData on destroyed row should return invalid version"), V.IsValid());
+		});
+
+		It(TEXT("RemoveScopeData"), EAsyncExecution::TaskGraphMainTick, [this]()
+		{
+			RowHandle Row = Storage->AddScopeRow();
+			Storage->SetScopeData<FTestColumnInt>(Row, FTestColumnInt{.TestInt = 1});
+			Storage->RemoveScopeRow(Row);
+
+			bool Deleted = Storage->RemoveScopeData<FTestColumnInt>(Row);
+			TestFalse(TEXT("RemoveScopeData on destroyed row should return false"), Deleted);
+		});
+
+		It(TEXT("GetParentScope"), EAsyncExecution::TaskGraphMainTick, [this]()
+		{
+			RowHandle Row = Storage->AddScopeRow();
+			Storage->RemoveScopeRow(Row);
+
+			RowHandle Parent = Storage->GetParentScope(Row);
+			TestEqual(TEXT("GetParentScope on destroyed row should return InvalidRowHandle"),
+				Parent, InvalidRowHandle);
+		});
+
+		It(TEXT("SetParentScope"), EAsyncExecution::TaskGraphMainTick, [this]()
+		{
+			RowHandle Row = Storage->AddScopeRow();
+			RowHandle Other = CreateRow();
+			Storage->RemoveScopeRow(Row);
+
+			// Should not crash
+			Storage->SetParentScope(Row, Other);
+		});
+
+		It(TEXT("RemoveScopeRowTwice"), EAsyncExecution::TaskGraphMainTick, [this]()
+		{
+			RowHandle Row = Storage->AddScopeRow();
+			Storage->RemoveScopeRow(Row);
+
+			// Should not crash
+			Storage->RemoveScopeRow(Row);
+		});
+
+		It(TEXT("GetAllVisibleScopeColumns"), EAsyncExecution::TaskGraphMainTick, [this]()
+		{
+			RowHandle Row = Storage->AddScopeRow();
+			Storage->SetScopeData<FTestColumnInt>(Row, FTestColumnInt{.TestInt = 1});
+			Storage->RemoveScopeRow(Row);
+
+			TArray<const UScriptStruct*> Columns = Storage->GetAllVisibleScopeColumns(Row);
+			TestEqual(TEXT("GetAllVisibleScopeColumns on destroyed row should return empty"),
+				Columns.Num(), 0);
+		});
+	});
+}
+
+} // namespace UE::Editor::DataStorage::Scope::Tests
+
+PRAGMA_ENABLE_EXPERIMENTAL_WARNINGS
+
+#endif // WITH_TESTS
