[TEDS] Refactor query capability forwarding to use explicit context environments

TImplementationForwarder has been updated from a 2-parameter form <Interface, Capabilities...> to a 3-parameter form <Environment, Interface, Capabilities...>. The environment type is now an explicit first argument stored in the terminal base, rather than being implicitly tied to the interface. Capabilities are now implemented as static functions that can be implemeted using a generic interface for the environment, or be overloaded on specific calls. This reduces the amount of overloaded with the same implementation as well as be a slight performance boost since the compiler can inline these functions better. That gain however balaneces out the extra pointer that now needs to happen in some cases.

Two concrete environment types are introduced: FDirectContextEnvironment for direct (non-Mass) query execution and FProcessorContextEnvironment for Mass processor query execution. All call sites for the direct API, processor adaptors, and query context implementations have been updated to construct and pass the appropriate environment type.

Capability implementations for row batch, single-row access and dynamic columns have been factored out into separate files, with environment-typed overloads where needed.

#Synced-CL 52378109

diff --git a/Engine/Plugins/Experimental/EditorDataStorage/Source/TedsCore/Private/Processors/TypedElementProcessorAdaptors.cpp b/Engine/Plugins/Experimental/EditorDataStorage/Source/TedsCore/Private/Processors/TypedElementProcessorAdaptors.cpp
--- a/Engine/Plugins/Experimental/EditorDataStorage/Source/TedsCore/Private/Processors/TypedElementProcessorAdaptors.cpp
+++ b/Engine/Plugins/Experimental/EditorDataStorage/Source/TedsCore/Private/Processors/TypedElementProcessorAdaptors.cpp
@@ -1081,10 +1081,11 @@ namespace UE::Editor::DataStorage
 			const FTedsHierarchyAccessInterface* HierarchyAccessInterface = nullptr;
 		};
 
-		struct FMassDirectContextForwarder final : Queries::Private::TImplementationForwarder<IDirectQueryContext, Queries::FDynamicColumnInfoContextCapability>
+		struct FMassDirectContextForwarder final : Queries::Private::TImplementationForwarder<Queries::FProcessorContextEnvironment, IDirectQueryContext, Queries::FDynamicColumnInfoContextCapability>
 		{
 			FMassDirectContextForwarder(FMassExecutionContext& InContext, FEnvironment& InEnvironment, FQueryDescription& Description)
-				: Implementation(InContext, InEnvironment)
+				: TImplementationForwarder(InEnvironment, InContext)
+				, Implementation(InContext, InEnvironment)
 			{
 				FHierarchyHandle Handle = InEnvironment.GetHierarchyRegistrar().FindHierarchyByName(Description.Hierarchy);
 				HierarchyAccessInterface = InEnvironment.GetHierarchyRegistrar().GetAccessInterface(Handle);
@@ -1113,18 +1114,17 @@ namespace UE::Editor::DataStorage
 		protected:
 			void* EmplaceObjectInScratch(const FEmplaceObjectParams& Params) override { return Implementation.EmplaceObjectInScratch(Params.ObjectSize, Params.Alignment, Params.Construct, Params.Destroy, Params.SourceObject); }
 
-			virtual FEnvironment& GetEnvironment() override { return Implementation.Environment; }
-			virtual const FEnvironment& GetEnvironment() const override { return Implementation.Environment; }
 		public:
 
 			FMassContextCommon Implementation;
 			const FTedsHierarchyAccessInterface* HierarchyAccessInterface = nullptr;
 		};
 
-		struct FMassSubqueryContextForwarder : Queries::Private::TImplementationForwarder<ISubqueryContext, Queries::FDynamicColumnInfoContextCapability>
+		struct FMassSubqueryContextForwarder : Queries::Private::TImplementationForwarder<Queries::FProcessorContextEnvironment, ISubqueryContext, Queries::FDynamicColumnInfoContextCapability>
 		{
 			FMassSubqueryContextForwarder(FMassExecutionContext& InContext, FEnvironment& InEnvironment, FQueryDescription& Description)
-				: Implementation(InContext, InEnvironment, Description.bShouldBatchModifications)
+				: TImplementationForwarder(InEnvironment, InContext)
+				, Implementation(InContext, InEnvironment, Description.bShouldBatchModifications)
 			{
 				FHierarchyHandle Handle = InEnvironment.GetHierarchyRegistrar().FindHierarchyByName(Description.Hierarchy);
 				Implementation.HierarchyAccessInterface = InEnvironment.GetHierarchyRegistrar().GetAccessInterface(Handle);
@@ -1167,9 +1167,6 @@ namespace UE::Editor::DataStorage
 		protected:
 			void* EmplaceObjectInScratch(const FEmplaceObjectParams& Params) override { return Implementation.EmplaceObjectInScratch(Params.ObjectSize, Params.Alignment, Params.Construct, Params.Destroy, Params.SourceObject); }
 			
-			virtual FEnvironment& GetEnvironment() override { return Implementation.Environment; }
-			virtual const FEnvironment& GetEnvironment() const override { return Implementation.Environment; }
-
 			FMassWithEnvironmentContextCommon Implementation;
 		};
 
@@ -1280,14 +1277,15 @@ namespace UE::Editor::DataStorage
 			FExtendedQueryStore& QueryStore;
 		};
 
-		struct FMassContextForwarder final : Queries::Private::TImplementationForwarder<IQueryContext, Queries::FDynamicColumnInfoContextCapability>
+		struct FMassContextForwarder final : Queries::Private::TImplementationForwarder<Queries::FProcessorContextEnvironment, IQueryContext, Queries::FDynamicColumnInfoContextCapability>
 		{
 			FMassContextForwarder(
 				FQueryDescription& InQueryDescription,
-				FMassExecutionContext& InContext, 
+				FMassExecutionContext& InContext,
 				FExtendedQueryStore& InQueryStore,
 				FEnvironment& InEnvironment)
-					: Implementation(InQueryDescription, InContext, InQueryStore, InEnvironment)
+					: TImplementationForwarder(InEnvironment, InContext)
+					, Implementation(InQueryDescription, InContext, InQueryStore, InEnvironment)
 			{
 			}
 	
@@ -1332,9 +1330,6 @@ namespace UE::Editor::DataStorage
 		protected:
 			void* EmplaceObjectInScratch(const FEmplaceObjectParams& Params) override { return Implementation.EmplaceObjectInScratch(Params.ObjectSize, Params.Alignment, Params.Construct, Params.Destroy, Params.SourceObject); }
 
-			virtual FEnvironment& GetEnvironment() override { return Implementation.Environment; }
-			virtual const FEnvironment& GetEnvironment() const override { return Implementation.Environment; }
-		
 		protected:
 			FMassQueryContextImplementation Implementation;
 		};
@@ -1635,9 +1630,10 @@ UE::Editor::DataStorage::FQueryResult FTypedElementQueryProcessorData::Execute(
 			{
 				using namespace Processors::Private;
 				// No need to cache any subsystem dependencies as these are not accessible from a direct query.
-				QueryContext_MassForwarder QueryContext(Environment, Context);
+				FProcessorContextEnvironment Env(Environment, Context);
+				QueryContext_ProcessorApi QueryContext(Env);
 				// This assumes that there already has been a check that the provided function is compatible with the query.
-				Callback.Call(QueryContext, QueryContext.GetContextImplementation());
+				Callback.Call(QueryContext, Env);
 				Result.Count += Context.GetNumEntities();
 			};
 		if (EnumHasAnyFlags(ExecutionFlags, ERunQueryFlags::ParallelizeChunks))
diff --git a/Engine/Plugins/Experimental/EditorDataStorage/Source/TedsCore/Private/Processors/TypedElementQueryCapabilityImplemenationForwarder.h b/Engine/Plugins/Experimental/EditorDataStorage/Source/TedsCore/Private/Processors/TypedElementQueryCapabilityImplemenationForwarder.h
--- a/Engine/Plugins/Experimental/EditorDataStorage/Source/TedsCore/Private/Processors/TypedElementQueryCapabilityImplemenationForwarder.h
+++ b/Engine/Plugins/Experimental/EditorDataStorage/Source/TedsCore/Private/Processors/TypedElementQueryCapabilityImplemenationForwarder.h
@@ -3,6 +3,7 @@
 #pragma once
 
 #include "DataStorage/CommonTypes.h"
+#include "DataStorage/Queries/Types.h"
 #include "Elements/Framework/TypedElementQueryCapabilities.h"
 #include "Elements/Framework/TypedElementQueryContract.h"
 #include "Elements/Framework/TypedElementQueryFunctionArgumentTypes.h"
@@ -13,47 +14,85 @@ namespace UE::Editor::DataStorage::Queries::Private
 	// be directly used in the various IQueryContext implementation. This wrapper can be used to bridge the two until
 	// the new query context can fully replace the original query context.
 
-	template<typename Interface, typename... Implementations> struct TImplementationForwarder
+	template<typename Environment, typename Interface, typename... Implementations> struct TImplementationForwarder
 	{
 		static_assert(sizeof...(Implementations) < 0,
 			"Specialization for query context forwards missing. The most likely reason is the addition of a new query "
 			"context capability that wasn't specialized.");
 	};
 
-	template<typename Interface>
-	struct TImplementationForwarder<Interface> : Interface
+	template<typename Environment, typename Interface>
+	struct TImplementationForwarder<Environment, Interface> : Interface
 	{
+		TImplementationForwarder() = default;
+
+		template<typename... TArgs>
+		explicit TImplementationForwarder(TArgs&&... Args)
+			: ContextEnvironment(Forward<TArgs>(Args)...)
+		{
+		}
+
+		Environment ContextEnvironment;
 	};
 
+	template<typename Implementation, typename Capability>
+	concept ImplementsCapability = std::is_base_of_v<ImplementsContextCapability<Capability>, Implementation>;
+
 #define ArgTypeName(Type, Name) Type Name
 #define ArgName(Type, Name) Name
+#define ForwardArg(Type, Name) Forward< Type >( Name )
+	
+#define CallFunction(Capability, ReturnValue, FunctionName, ...) \
+	if constexpr (ImplementsCapability<Implementation, Capability>) \
+	{ \
+		static_assert( requires() \
+			{ \
+				{ Implementation:: FunctionName (this->ContextEnvironment __VA_OPT__(,) __VA_ARGS__ ) } -> UE::CSameAs< ReturnValue >; \
+			}, "Function '" #FunctionName "' in capability '" #Capability "' is missing from the forwarder." ); \
+		return Implementation:: FunctionName ( this->ContextEnvironment __VA_OPT__(,) __VA_ARGS__ ); \
+	} \
+	else \
+	{ \
+		static_assert(std::is_same_v<void, Capability>, \
+			"Function '" #FunctionName "' in capability '" #Capability "' is missing from the implementation."); \
+		return UE::MakeDefault< ReturnValue >(); \
+	}
+
+#define CapabilityStart(Capability, Flags) \
+	template<typename Environment, typename Interface, typename Implementation, typename... Implementations> \
+		requires ImplementsCapability<Implementation, Queries:: Capability > \
+	struct TImplementationForwarder< Environment, Interface, Implementation, Implementations...> : \
+		Implementation, \
+		TImplementationForwarder<Environment, Interface, Implementations...> \
+	{ \
+		using Super = TImplementationForwarder<Environment, Interface, Implementations...>; \
+		TImplementationForwarder() = default; \
+		template<typename... TArgs> \
+		explicit TImplementationForwarder(TArgs&&... Args) \
+			: Super(Forward<TArgs>(Args)...) \
+		{}
 
 #define Function0(Capability, Return, Function) virtual Return Function() \
-	{ return Implementation:: Function (); }
+	{ CallFunction(Capability, Return, Function ) }
 #define Function1(Capability, Return, Function, Arg1) virtual Return Function(ArgTypeName Arg1 ) override \
-	{ return Implementation:: Function (ArgName Arg1 ); }
+	{ CallFunction(Capability, Return, Function, ForwardArg Arg1 ) }
 #define Function2(Capability, Return, Function, Arg1, Arg2) virtual Return Function(ArgTypeName Arg1 , ArgTypeName Arg2 ) override \
-	{ return Implementation:: Function (ArgName Arg1 , ArgName Arg2 ); }
+	{ CallFunction(Capability, Return, Function, ForwardArg Arg1, ForwardArg Arg2 ) }
 #define Function3(Capability, Return, Function, Arg1, Arg2, Arg3) virtual Return Function(ArgTypeName Arg1 , ArgTypeName Arg2 , ArgTypeName Arg3 ) override \
-	{ return Implementation:: Function (ArgName Arg1 , ArgName Arg2 , ArgName Arg3); }
+	{ CallFunction(Capability, Return, Function, ForwardArg Arg1, ForwardArg Arg2, ForwardArg Arg3 ) }
 #define Function4(Capability, Return, Function, Arg1, Arg2, Arg3, Arg4) virtual Return Function(ArgTypeName Arg1 , ArgTypeName Arg2 , ArgTypeName Arg3 , ArgTypeName Arg4 ) override \
-	{ return Implementation:: Function (ArgName Arg1 , ArgName Arg2 , ArgName Arg3 , ArgName Arg4); }
+	{ CallFunction(Capability, Return, Function, ForwardArg Arg1, ForwardArg Arg2, ForwardArg Arg3, ForwardArg Arg4 ) }
 
 #define ConstFunction0(Capability, Return, Function) virtual Return Function() const \
-	{ return Implementation:: Function (); }
+	{ CallFunction(Capability, Return, Function ) }
 #define ConstFunction1(Capability, Return, Function, Arg1) virtual Return Function(ArgTypeName Arg1 ) const override \
-	{ return Implementation:: Function (ArgName Arg1 ); }
+	{ CallFunction(Capability, Return, Function, ForwardArg Arg1 ) }
 #define ConstFunction2(Capability, Return, Function, Arg1, Arg2) virtual Return Function(ArgTypeName Arg1 , ArgTypeName Arg2 ) const override \
-	{ return Implementation:: Function (ArgName Arg1 , ArgName Arg2 ); }
+	{ CallFunction(Capability, Return, Function, ForwardArg Arg1, ForwardArg Arg2 ) }
 #define ConstFunction3(Capability, Return, Function, Arg1, Arg2, Arg3) virtual Return Function(ArgTypeName Arg1 , ArgTypeName Arg2 , ArgTypeName Arg3 ) const override \
-	{ return Implementation:: Function (ArgName Arg1 , ArgName Arg2 , ArgName Arg3); }
+	{ CallFunction(Capability, Return, Function, ForwardArg Arg1, ForwardArg Arg2, ForwardArg Arg3 ) }
 #define ConstFunction4(Capability, Return, Function, Arg1, Arg2, Arg3, Arg4) virtual Return Function(ArgTypeName Arg1 , ArgTypeName Arg2 , ArgTypeName Arg3 , ArgTypeName Arg4 ) const override \
-	{ return Implementation:: Function (ArgName Arg1 , ArgName Arg2 , ArgName Arg3 , ArgName Arg4); }
-
-#define CapabilityStart(Capability, Flags) \
-	template<typename Interface, typename Implementation, typename... Implementations> requires std::is_base_of_v< Queries:: Capability , Implementation> \
-	struct TImplementationForwarder< Interface, Implementation, Implementations...> : Implementation, TImplementationForwarder<Interface, Implementations...> \
-	{ 
+	{ CallFunction(Capability, Return, Function, ForwardArg Arg1, ForwardArg Arg2, ForwardArg Arg3, ForwardArg Arg4 ) }
 
 #define CapabilityEnd(Capability) \
 	}; \
@@ -62,6 +101,8 @@ namespace UE::Editor::DataStorage::Queries::Private
 
 #undef ArgTypeName
 #undef ArgName
+#undef ForwardArg
+#undef CallFunction
 #undef Function0
 #undef Function1
 #undef Function2
diff --git a/Engine/Plugins/Experimental/EditorDataStorage/Source/TedsCore/Private/Queries/QueryContextCapabilities/EditorDataStorageContextEnvironment.h b/Engine/Plugins/Experimental/EditorDataStorage/Source/TedsCore/Private/Queries/QueryContextCapabilities/EditorDataStorageContextEnvironment.h
new file mode 100644
--- /dev/null
+++ b/Engine/Plugins/Experimental/EditorDataStorage/Source/TedsCore/Private/Queries/QueryContextCapabilities/EditorDataStorageContextEnvironment.h
@@ -0,0 +1,22 @@
+// Copyright Epic Games, Inc. All Rights Reserved.
+
+#pragma once
+
+#include "DataStorage/CommonTypes.h"
+#include "Elements/Framework/TypedElementQueryFunctions.h"
+
+namespace UE::Editor::DataStorage
+{
+	class FEnvironment;
+} // namespace UE::Editor::DataStorage
+
+namespace UE::Editor::DataStorage::Queries
+{
+	struct IContextEnvironment : IQueryFunctionResponse
+	{
+		virtual ~IContextEnvironment() = default;
+
+		virtual FEnvironment* GetEnvironment() = 0;
+		virtual const FEnvironment* GetEnvironment() const = 0;
+	};
+} // namespace UE::Editor::DataStorage::Queries
diff --git a/Engine/Plugins/Experimental/EditorDataStorage/Source/TedsCore/Private/Queries/QueryContextCapabilities/EditorDataStorageDynamicColumnInfo.cpp b/Engine/Plugins/Experimental/EditorDataStorage/Source/TedsCore/Private/Queries/QueryContextCapabilities/EditorDataStorageDynamicColumnInfo.cpp
--- a/Engine/Plugins/Experimental/EditorDataStorage/Source/TedsCore/Private/Queries/QueryContextCapabilities/EditorDataStorageDynamicColumnInfo.cpp
+++ b/Engine/Plugins/Experimental/EditorDataStorage/Source/TedsCore/Private/Queries/QueryContextCapabilities/EditorDataStorageDynamicColumnInfo.cpp
@@ -2,17 +2,19 @@
 
 #include "Queries/QueryContextCapabilities/EditorDataStorageDynamicColumnInfo.h"
 
+#include "Queries/QueryContextCapabilities/EditorDataStorageContextEnvironment.h"
+
 namespace UE::Editor::DataStorage::Queries
 {
 	const UScriptStruct* FDynamicColumnInfoContextCapability::FindDynamicColumnType(
-		const UScriptStruct& TemplateType, const FName& Identifier) const
+		const IContextEnvironment& Environment, const UScriptStruct& TemplateType, const FName& Identifier)
 	{
-		return GetEnvironment().FindDynamicColumn(TemplateType, Identifier);
+		return Environment.GetEnvironment()->FindDynamicColumn(TemplateType, Identifier);
 	}
 
 	void FDynamicColumnInfoContextCapability::ForEachDynamicColumnType(
-		const UScriptStruct& Template, TFunctionRef<void(const UScriptStruct& Type)> Callback) const
+		const IContextEnvironment& Environment, const UScriptStruct& TemplateType, TFunctionRef<void(const UScriptStruct& Type)> Callback)
 	{
-		GetEnvironment().ForEachDynamicColumn(Template, Callback);
+		Environment.GetEnvironment()->ForEachDynamicColumn(TemplateType, Callback);
 	}
 } // namespace UE::Editor::DataStorage::Queries
diff --git a/Engine/Plugins/Experimental/EditorDataStorage/Source/TedsCore/Private/Queries/QueryContextCapabilities/EditorDataStorageDynamicColumnInfo.h b/Engine/Plugins/Experimental/EditorDataStorage/Source/TedsCore/Private/Queries/QueryContextCapabilities/EditorDataStorageDynamicColumnInfo.h
--- a/Engine/Plugins/Experimental/EditorDataStorage/Source/TedsCore/Private/Queries/QueryContextCapabilities/EditorDataStorageDynamicColumnInfo.h
+++ b/Engine/Plugins/Experimental/EditorDataStorage/Source/TedsCore/Private/Queries/QueryContextCapabilities/EditorDataStorageDynamicColumnInfo.h
@@ -2,7 +2,8 @@
 
 #pragma once
 
-#include "Elements/Framework/TypedElementQueryContextImplementation.h"
+#include "DataStorage/Queries/Types.h"
+#include "Elements/Framework/TypedElementQueryCapabilities.h"
 #include "TypedElementDatabaseEnvironment.h"
 
 class FName;
@@ -10,19 +11,20 @@ class UScriptStruct;
 
 namespace UE::Editor::DataStorage::Queries
 {
+	struct IContextEnvironment;
+
 	/**
 	 * Implementation for the DynamicColumnInfo context capability.
 	 * This provides read access to the dynamic columns in order to retrieve information about them from any thread.
 	 * This class provides an implementation block and is not meant to be directly used but should be inherited from
 	 * by a class that's compatible with TQueryContextImpl.
 	 */
-	class FDynamicColumnInfoContextCapability : public DynamicColumnInfo, public IQueryEnvironmentRetrieval
+	class FDynamicColumnInfoContextCapability : public ImplementsContextCapability<DynamicColumnInfo>
 	{
 	public:
-		virtual ~FDynamicColumnInfoContextCapability() override = default;
-
-		virtual const UScriptStruct* FindDynamicColumnType(const UScriptStruct& TemplateType, const FName& Identifier) const override;
-		virtual void ForEachDynamicColumnType(
-			const UScriptStruct& Template, TFunctionRef<void(const UScriptStruct& Type)> Callback) const override;
+		static const UScriptStruct* FindDynamicColumnType(
+			const IContextEnvironment& Environment, const UScriptStruct& TemplateType, const FName& Identifier);
+		static void ForEachDynamicColumnType(
+			const IContextEnvironment& Environment, const UScriptStruct& TemplateType, TFunctionRef<void(const UScriptStruct& Type)> Callback);
 	};
 } // namespace UE::Editor::DataStorage::Queries
diff --git a/Engine/Plugins/Experimental/EditorDataStorage/Source/TedsCore/Private/Queries/QueryContextCapabilities/RowBatch.cpp b/Engine/Plugins/Experimental/EditorDataStorage/Source/TedsCore/Private/Queries/QueryContextCapabilities/RowBatch.cpp
new file mode 100644
--- /dev/null
+++ b/Engine/Plugins/Experimental/EditorDataStorage/Source/TedsCore/Private/Queries/QueryContextCapabilities/RowBatch.cpp
@@ -0,0 +1,112 @@
+// Copyright Epic Games, Inc. All Rights Reserved.
+
+#pragma once
+
+#include "Queries/QueryContextCapabilities/RowBatch.h"
+
+#include "Mass/EntityHandle.h"
+#include "Queries/TypedElementQueryContext_DirectSingleApi.h"
+#include "Queries/TypedElementQueryContext_MassForwarder.h"
+#include "TypedElementDatabase.h"
+#include "TypedElementDatabaseEnvironment.h"
+#include "TypedElementUtils.h"
+
+namespace UE::Editor::DataStorage::Queries
+{
+	//
+	// IContextEnvironment
+	//
+
+	uint32 FRowBatchInfoContextCapability::GetBatchRowCount(const IContextEnvironment& Environment)
+	{
+		return Environment.GetBatchRows().Num();
+	}
+
+	FRowHandleArrayView FRowBatchInfoContextCapability::GetBatchRowHandles(const IContextEnvironment& Environment)
+	{
+		return Environment.GetBatchRows();
+	}
+
+	bool FRowBatchInfoContextCapability::CurrentBatchTableHasColumns(
+		const IContextEnvironment& Environment, TConstArrayView<const UScriptStruct*> Columns)
+	{
+		// Pick the first row since all rows in a batch are in the same table.
+		if (FRowHandleArrayView Rows = Environment.GetBatchRows(); !Rows.IsEmpty())
+		{
+			FMassEntityHandle Entity = FMassEntityHandle::FromNumber(Rows[0]);
+			FMassArchetypeHandle Archetype = Environment.GetEnvironment()->GetMassEntityManager().GetArchetypeForEntity(Entity);
+			const FMassElementBitSet& ArchetypeBits = Archetype.GetCompositionBitSetChecked();
+			for (const UScriptStruct* Column : Columns)
+			{
+				if (!ArchetypeBits.Contains(Column))
+				{
+					return false;
+				}
+			}
+			return true;
+		}
+		return false;
+	}
+
+
+
+
+	//
+	// FProcessorContextEnvironment
+	// 
+
+	bool FRowBatchInfoContextCapability::CurrentBatchTableHasColumns(
+		const FProcessorContextEnvironment& Environment, TConstArrayView<const UScriptStruct*> Columns)
+	{
+		const FMassExecutionContext& Context = Environment.MassContext;
+		for (const UScriptStruct* Column : Columns)
+		{
+			if (!Context.DoesArchetypeHaveElement(Column))
+			{
+				return false;
+			}
+		}
+		return true;
+	}
+
+	const void* FRowBatchInfoContextCapability::GetColumnBatchAddress(
+		const FProcessorContextEnvironment& Environment, const UScriptStruct* ColumnType)
+	{
+		checkf(UE::Mass::IsA<FMassFragment>(ColumnType), TEXT("Trying to retrieve a data pointer for column '%s' which doesn't have data."),
+			ColumnType ? *ColumnType->GetName() : TEXT("<null>"));
+		return Environment.MassContext.GetFragmentView(ColumnType).GetData();
+	}
+
+	void* FRowBatchInfoContextCapability::GetMutableColumnBatchAddress(
+		FProcessorContextEnvironment& Environment, const UScriptStruct* ColumnType)
+	{
+		checkf(UE::Mass::IsA<FMassFragment>(ColumnType), TEXT("Trying to retrieve a data pointer for column '%s' which doesn't have data."),
+			ColumnType ? *ColumnType->GetName() : TEXT("<null>"));
+		return Environment.MassContext.GetMutableFragmentView(ColumnType).GetData();
+	}
+
+	
+	
+	//
+	// FDirectContextEnvironment
+	//
+
+	const void* FRowBatchInfoContextCapability::GetColumnBatchAddress(const FDirectContextEnvironment& Environment, const UScriptStruct* ColumnType)
+	{
+		checkf(UE::Mass::IsA<FMassFragment>(ColumnType), TEXT("Trying to retrieve a data pointer for column '%s' which doesn't have data."),
+			ColumnType ? *ColumnType->GetName() : TEXT("<null>"));
+		
+		return !Environment.Rows.IsEmpty()
+			? Environment.DirectContext->GetColumnData(*Environment.Rows.GetData(), ColumnType)
+			: nullptr;
+	}
+
+	void* FRowBatchInfoContextCapability::GetMutableColumnBatchAddress(FDirectContextEnvironment& Environment, const UScriptStruct* ColumnType)
+	{
+		checkf(UE::Mass::IsA<FMassFragment>(ColumnType), TEXT("Trying to retrieve a data pointer for column '%s' which doesn't have data."),
+			ColumnType ? *ColumnType->GetName() : TEXT("<null>"));
+		return !Environment.Rows.IsEmpty()
+			? Environment.DirectContext->GetColumnData(*Environment.Rows.GetData(), ColumnType)
+			: nullptr;
+	}
+} // namespace UE::Editor::DataStorage::Queries
diff --git a/Engine/Plugins/Experimental/EditorDataStorage/Source/TedsCore/Private/Queries/QueryContextCapabilities/RowBatch.h b/Engine/Plugins/Experimental/EditorDataStorage/Source/TedsCore/Private/Queries/QueryContextCapabilities/RowBatch.h
new file mode 100644
--- /dev/null
+++ b/Engine/Plugins/Experimental/EditorDataStorage/Source/TedsCore/Private/Queries/QueryContextCapabilities/RowBatch.h
@@ -0,0 +1,32 @@
+// Copyright Epic Games, Inc. All Rights Reserved.
+
+#pragma once
+
+#include "DataStorage/CommonTypes.h"
+#include "DataStorage/Queries/Types.h"
+#include "Elements/Framework/TypedElementQueryCapabilities.h"
+#include "Queries/QueryContextCapabilities/EditorDataStorageContextEnvironment.h"
+
+namespace UE::Editor::DataStorage::Queries
+{
+	struct FConstDirectContextEnvironment;
+	struct FConstProcessorContextEnvironment;
+	struct FDirectContextEnvironment;
+	struct FProcessorContextEnvironment;
+
+	// For documentation on these functions see "TypedElementQueryCapabilities.inl"
+	class FRowBatchInfoContextCapability : public ImplementsContextCapability<RowBatchInfo>
+	{
+	public:
+		static uint32 GetBatchRowCount(const IContextEnvironment& Environment);
+		static FRowHandleArrayView GetBatchRowHandles(const IContextEnvironment& Environment);
+		static bool CurrentBatchTableHasColumns(const IContextEnvironment& Environment, TConstArrayView<const UScriptStruct*> Columns);
+
+		static bool CurrentBatchTableHasColumns(const FProcessorContextEnvironment& Environment, TConstArrayView<const UScriptStruct*> Columns);
+		static const void* GetColumnBatchAddress(const FProcessorContextEnvironment& Environment, const UScriptStruct* ColumnType);
+		static void* GetMutableColumnBatchAddress(FProcessorContextEnvironment& Environment, const UScriptStruct* ColumnType);
+
+		static const void* GetColumnBatchAddress(const FDirectContextEnvironment& Environment, const UScriptStruct* ColumnType);
+		static void* GetMutableColumnBatchAddress(FDirectContextEnvironment& Environment, const UScriptStruct* ColumnType);
+	};
+} // namespace UE::Editor::DataStorage::Queries
\ No newline at end of file
diff --git a/Engine/Plugins/Experimental/EditorDataStorage/Source/TedsCore/Private/Queries/QueryContextCapabilities/SingleRow.cpp b/Engine/Plugins/Experimental/EditorDataStorage/Source/TedsCore/Private/Queries/QueryContextCapabilities/SingleRow.cpp
new file mode 100644
--- /dev/null
+++ b/Engine/Plugins/Experimental/EditorDataStorage/Source/TedsCore/Private/Queries/QueryContextCapabilities/SingleRow.cpp
@@ -0,0 +1,48 @@
+// Copyright Epic Games, Inc. All Rights Reserved.
+
+#pragma once
+
+#include "Queries/QueryContextCapabilities/SingleRow.h"
+
+#include "Mass/EntityHandle.h"
+#include "Queries/QueryContextCapabilities/EditorDataStorageContextEnvironment.h"
+#include "Queries/TypedElementQueryContext_MassForwarder.h"
+#include "TypedElementDatabaseEnvironment.h"
+
+namespace UE::Editor::DataStorage::Queries
+{
+	bool FSingleRowInfoContextCapability::CurrentRowHasColumns(
+		const IContextEnvironment& Environment, TConstArrayView<const UScriptStruct*> Columns)
+	{
+		// Pick the first row since all rows in a batch are in the same table.
+		if (FRowHandleArrayView Rows = Environment.GetBatchRows(); !Rows.IsEmpty())
+		{
+			FMassEntityHandle Entity = FMassEntityHandle::FromNumber(Rows[0]);
+			FMassArchetypeHandle Archetype = Environment.GetEnvironment()->GetMassEntityManager().GetArchetypeForEntity(Entity);
+			const FMassElementBitSet& ArchetypeBits = Archetype.GetCompositionBitSetChecked();
+			for (const UScriptStruct* Column : Columns)
+			{
+				if (!ArchetypeBits.Contains(Column))
+				{
+					return false;
+				}
+			}
+			return true;
+		}
+		return false;
+	}
+
+	bool FSingleRowInfoContextCapability::CurrentRowHasColumns(
+		const FProcessorContextEnvironment& Environment, TConstArrayView<const UScriptStruct*> Columns)
+	{
+		const FMassExecutionContext& Context = Environment.MassContext;
+		for (const UScriptStruct* Column : Columns)
+		{
+			if (!Context.DoesArchetypeHaveElement(Column))
+			{
+				return false;
+			}
+		}
+		return true;
+	}
+} // namespace UE::Editor::DataStorage::Queries
diff --git a/Engine/Plugins/Experimental/EditorDataStorage/Source/TedsCore/Private/Queries/QueryContextCapabilities/SingleRow.h b/Engine/Plugins/Experimental/EditorDataStorage/Source/TedsCore/Private/Queries/QueryContextCapabilities/SingleRow.h
new file mode 100644
--- /dev/null
+++ b/Engine/Plugins/Experimental/EditorDataStorage/Source/TedsCore/Private/Queries/QueryContextCapabilities/SingleRow.h
@@ -0,0 +1,35 @@
+// Copyright Epic Games, Inc. All Rights Reserved.
+
+#pragma once
+
+#include "Concepts/ConvertibleTo.h"
+#include "DataStorage/CommonTypes.h"
+#include "DataStorage/Queries/Types.h"
+#include "Elements/Framework/TypedElementQueryCapabilities.h"
+
+namespace UE::Editor::DataStorage::Queries
+{
+	struct IContextEnvironment;
+	struct FConstProcessorContextEnvironment;
+	struct FProcessorContextEnvironment;
+
+	template<typename T>
+	concept CHasCurrentRow = requires(const T & Type)
+	{
+		{ Type.CurrentRow } -> UE::CConvertibleTo<RowHandle>;
+	};
+
+	// For documentation on these functions see "TypedElementQueryCapabilities.inl"
+	class FSingleRowInfoContextCapability : public ImplementsContextCapability<SingleRowInfo>
+	{
+	public:
+		static bool CurrentRowHasColumns(const IContextEnvironment& Environment, TConstArrayView<const UScriptStruct*> Columns);
+		static bool CurrentRowHasColumns(const FProcessorContextEnvironment& Environment, TConstArrayView<const UScriptStruct*> Columns);
+		
+		template<CHasCurrentRow EnvironmentType>
+		static RowHandle GetCurrentRow(const EnvironmentType& Environment)
+		{ 
+			return Environment.CurrentRow;
+		}
+	};
+} // namespace UE::Editor::DataStorage::Queries
diff --git a/Engine/Plugins/Experimental/EditorDataStorage/Source/TedsCore/Private/Queries/TypedElementQueryContext_DirectSingleApi.cpp b/Engine/Plugins/Experimental/EditorDataStorage/Source/TedsCore/Private/Queries/TypedElementQueryContext_DirectSingleApi.cpp
--- a/Engine/Plugins/Experimental/EditorDataStorage/Source/TedsCore/Private/Queries/TypedElementQueryContext_DirectSingleApi.cpp
+++ b/Engine/Plugins/Experimental/EditorDataStorage/Source/TedsCore/Private/Queries/TypedElementQueryContext_DirectSingleApi.cpp
@@ -5,90 +5,68 @@
 
 namespace UE::Editor::DataStorage::Queries
 {
-	const FName FQueryContext_DirectSingleApi::Name = "FQueryContext_DirectSingleApi";
+	//
+	// FDirectContextEnvironment
+	//
 
-	FQueryContext_DirectSingleApi::FQueryContext_DirectSingleApi(FEnvironment& Environment, UEditorDataStorage& DataStorage)
-		: Environment(Environment)
-		, DataStorage(DataStorage)
+	FDirectContextEnvironment::FDirectContextEnvironment(FEnvironment& Environment, UEditorDataStorage& DirectContext)
+		: Environment(&Environment)
+		, DirectContext(&DirectContext)
 	{
 	}
 
-	bool FQueryContext_DirectSingleApi::CurrentRowHasColumns(TConstArrayView<const UScriptStruct*> ColumnTypes) const
+	void FDirectContextEnvironment::SetBatch(FRowHandleArrayView InRows)
 	{
-		return DataStorage.HasColumns(CurrentRow, ColumnTypes);
+		Rows = InRows;
 	}
 
-	RowHandle FQueryContext_DirectSingleApi::GetCurrentRow() const
+	FEnvironment* FDirectContextEnvironment::GetEnvironment()
 	{
-		return CurrentRow;
-	}
-
-	void FQueryContext_DirectSingleApi::SetCurrentRow(RowHandle Row)
-	{
-		CurrentRow = Row;
+		return Environment;
 	}
 
-	bool FQueryContext_DirectSingleApi::NextBatch()
+	const FEnvironment* FDirectContextEnvironment::GetEnvironment() const
 	{
-		return false; // Only operates on a single row, so there are no more batches.
+		return Environment;
 	}
 
-	bool FQueryContext_DirectSingleApi::NextRow()
+	FRowHandleArrayView FDirectContextEnvironment::GetBatchRows() const
 	{
-		return false; // Only operates on a single row, so there are no more rows.
+		return Rows;
 	}
 
-	void FQueryContext_DirectSingleApi::GetConstColumns(TArrayView<const void*> ColumnsData, TConstArrayView<const UScriptStruct*> ColumnTypes)
+	void FDirectContextEnvironment::GetConstColumns(
+		TArrayView<const void*> ColumnsData, TConstArrayView<const UScriptStruct*> ColumnTypes) const
 	{
-		const void** ColumnsDataIt = ColumnsData.GetData();
-		const UEditorDataStorage& ConstDataStorage = DataStorage;
-		for (const UScriptStruct* ColumnType : ColumnTypes)
+		if (!Rows.IsEmpty())
 		{
-			*ColumnsDataIt++ = DataStorage.GetColumnData(CurrentRow, ColumnType);
+			RowHandle Root = *Rows.GetData();
+			const void** ColumnsDataIt = ColumnsData.GetData();
+			for (const UScriptStruct* ColumnType : ColumnTypes)
+			{
+				*ColumnsDataIt++ = DirectContext->GetColumnData(Root, ColumnType);
+			}
 		}
 	}
 
-	void FQueryContext_DirectSingleApi::GetMutableColumns(TArrayView<void*> ColumnsData, TConstArrayView<const UScriptStruct*> ColumnTypes)
+	void FDirectContextEnvironment::GetMutableColumns(
+		TArrayView<void*> ColumnsData, TConstArrayView<const UScriptStruct*> ColumnTypes)
 	{
-		void** ColumnsDataIt = ColumnsData.GetData();
-		for (const UScriptStruct* ColumnType : ColumnTypes)
+		if (!Rows.IsEmpty())
 		{
-			*ColumnsDataIt++ = DataStorage.GetColumnData(CurrentRow, ColumnType);
+			// At the time this is called there are guaranteed to be rows assigned.
+			RowHandle Root = *Rows.GetData();
+			void** ColumnsDataIt = ColumnsData.GetData();
+			for (const UScriptStruct* ColumnType : ColumnTypes)
+			{
+				*ColumnsDataIt++ = DirectContext->GetColumnData(Root, ColumnType);
+			}
 		}
 	}
 
-	uint32 FQueryContext_DirectSingleApi::GetBatchRowCount() const
-	{
-		return 1;
-	}
-
-	TConstArrayView<RowHandle> FQueryContext_DirectSingleApi::GetBatchRowHandles() const
-	{
-		return TConstArrayView<RowHandle>(&CurrentRow, 1);
-	}
-
-	bool FQueryContext_DirectSingleApi::CurrentBatchTableHasColumns(TConstArrayView<const UScriptStruct*> Columns) const
-	{
-		return CurrentRowHasColumns(Columns);
-	}
-
-	const void* FQueryContext_DirectSingleApi::GetColumnBatchAddress(const UScriptStruct* ColumnType) const
+	void FDirectContextEnvironment::SetCurrentRow(RowHandle Row)
 	{
-		return DataStorage.GetColumnData(CurrentRow, ColumnType);
-	}
-
-	void* FQueryContext_DirectSingleApi::GetMutableColumnBatchAddress(const UScriptStruct* ColumnType)
-	{
-		return DataStorage.GetColumnData(CurrentRow, ColumnType);
-	}
-
-	FEnvironment& FQueryContext_DirectSingleApi::GetEnvironment()
-	{
-		return Environment;
+		CurrentRow = Row;
 	}
 
-	const FEnvironment& FQueryContext_DirectSingleApi::GetEnvironment() const
-	{
-		return Environment;
-	}
 } // namespace UE::Editor::DataStorage::Queries
\ No newline at end of file
diff --git a/Engine/Plugins/Experimental/EditorDataStorage/Source/TedsCore/Private/Queries/TypedElementQueryContext_DirectSingleApi.h b/Engine/Plugins/Experimental/EditorDataStorage/Source/TedsCore/Private/Queries/TypedElementQueryContext_DirectSingleApi.h
--- a/Engine/Plugins/Experimental/EditorDataStorage/Source/TedsCore/Private/Queries/TypedElementQueryContext_DirectSingleApi.h
+++ b/Engine/Plugins/Experimental/EditorDataStorage/Source/TedsCore/Private/Queries/TypedElementQueryContext_DirectSingleApi.h
@@ -4,67 +4,40 @@
 
 #include "Elements/Framework/TypedElementQueryContextImplementation.h"
 #include "Elements/Framework/TypedElementQueryFunctions.h"
+#include "Elements/Framework/TypedElementRowHandleArray.h"
 #include "TypedElementDatabaseEnvironment.h"
+#include "Queries/QueryContextCapabilities/EditorDataStorageContextEnvironment.h"
 #include "Queries/QueryContextCapabilities/EditorDataStorageDynamicColumnInfo.h"
+#include "Queries/QueryContextCapabilities/RowBatch.h"
+#include "Queries/QueryContextCapabilities/SingleRow.h"
 
 class UEditorDataStorage;
 
 namespace UE::Editor::DataStorage::Queries
 {
-	/** 
-	 * Query context implementation for single row queries that call into TEDS Direct API.
-	 * Note: batch operations, for instance RowBatchInfo, are supported by this context, but
-	 * will in practice only run a single row. Batch mode was added for compatibility and offers
-	 * no performance benefits when using this context.
-	 */
-	class FQueryContext_DirectSingleApi final : 
-		public SingleRowInfo, 
-		public RowBatchInfo, 
-		public FDynamicColumnInfoContextCapability,
-		public IQueryFunctionResponse
+	struct FDirectContextEnvironment final : IContextEnvironment
 	{
-	public:
-		static const FName Name;
+		FDirectContextEnvironment() = default;
+		FDirectContextEnvironment(FEnvironment& Environment, UEditorDataStorage& DirectContext);
 
-		FQueryContext_DirectSingleApi(FEnvironment& Environment, UEditorDataStorage& DataStorage);
+		void SetBatch(FRowHandleArrayView InRows);
 
-		/**
-		 * @section ISingleRowInfo
-		 */
+		virtual FEnvironment* GetEnvironment() override;
+		virtual const FEnvironment* GetEnvironment() const override;
 
-		virtual bool CurrentRowHasColumns(TConstArrayView<const UScriptStruct*> ColumnTypes) const override;
-		virtual RowHandle GetCurrentRow() const override;
-
-		/**
-		 * @section IRowBatchInfo
-		 */
-		virtual uint32 GetBatchRowCount() const override;
-		virtual TConstArrayView<RowHandle> GetBatchRowHandles() const override;
-		virtual bool CurrentBatchTableHasColumns(TConstArrayView<const UScriptStruct*> Columns) const override;
-		virtual const void* GetColumnBatchAddress(const UScriptStruct* ColumnType) const override;
-		virtual void* GetMutableColumnBatchAddress(const UScriptStruct* ColumnType) override;
-
-		/**
-		 * @section IQueryFunctionResponse
-		 */
-		virtual bool NextBatch() override;
-		virtual bool NextRow() override;
-		virtual void GetConstColumns(TArrayView<const void*> ColumnsData, TConstArrayView<const UScriptStruct*> ColumnTypes) override;
+		virtual FRowHandleArrayView GetBatchRows() const override;
+		virtual void GetConstColumns(TArrayView<const void*> ColumnsData, TConstArrayView<const UScriptStruct*> ColumnTypes) const override;
 		virtual void GetMutableColumns(TArrayView<void*> ColumnsData, TConstArrayView<const UScriptStruct*> ColumnTypes) override;
+		virtual void SetCurrentRow(RowHandle Row) override;
 
-		void SetCurrentRow(RowHandle Row);
-
-		/**
-		 * @section IQueryEnvironmentRetrieval
-		 */
-		virtual FEnvironment& GetEnvironment() override;
-		virtual const FEnvironment& GetEnvironment() const override;
-
-	private:
-		FEnvironment& Environment;
-		UEditorDataStorage& DataStorage;
+		FRowHandleArrayView Rows;
+		FEnvironment* Environment = nullptr;
+		UEditorDataStorage* DirectContext = nullptr;
 		RowHandle CurrentRow = InvalidRowHandle;
 	};
-	
-	using QueryContext_DirectSingleApi = TQueryContextImpl<FQueryContext_DirectSingleApi>;
+
+	using QueryContext_DirectApi = TQueryContextImpl<FDirectContextEnvironment,
+		FSingleRowInfoContextCapability,
+		FRowBatchInfoContextCapability,
+		FDynamicColumnInfoContextCapability>;
 } // namespace UE::Editor::DataStorage::Queries
diff --git a/Engine/Plugins/Experimental/EditorDataStorage/Source/TedsCore/Private/Queries/TypedElementQueryContext_MassForwarder.cpp b/Engine/Plugins/Experimental/EditorDataStorage/Source/TedsCore/Private/Queries/TypedElementQueryContext_MassForwarder.cpp
--- a/Engine/Plugins/Experimental/EditorDataStorage/Source/TedsCore/Private/Queries/TypedElementQueryContext_MassForwarder.cpp
+++ b/Engine/Plugins/Experimental/EditorDataStorage/Source/TedsCore/Private/Queries/TypedElementQueryContext_MassForwarder.cpp
@@ -6,92 +6,34 @@
 
 namespace UE::Editor::DataStorage::Queries
 {
-	const FName FQueryContext_MassForwarder::Name = "FQueryContext_MassForwarder";
+	//
+	// FProcessorContextEnvironment
+	//
 
-	FQueryContext_MassForwarder::FQueryContext_MassForwarder(FEnvironment& Environment, FMassExecutionContext& Context)
+	FProcessorContextEnvironment::FProcessorContextEnvironment(FEnvironment& Environment, FMassExecutionContext& MassContext)
 		: Environment(Environment)
-		, Context(Context)
+		, MassContext(MassContext)
 	{
-		TConstArrayView<RowHandle> Rows = MassEntitiesToRowsConversion(Context.GetEntities());
-		CurrentRow = Rows.GetData();
-		EndRow = Rows.GetData() + Rows.Num();
 	}
 
-	bool FQueryContext_MassForwarder::CurrentRowHasColumns(TConstArrayView<const UScriptStruct*> ColumnTypes) const
+	FEnvironment* FProcessorContextEnvironment::GetEnvironment()
 	{
-		bool Result = true;
-		const UScriptStruct* const* It = ColumnTypes.GetData();
-		const UScriptStruct* const* End = ColumnTypes.GetData() + ColumnTypes.Num();
-
-		for (; Result && It != End; ++It)
-		{
-			const UScriptStruct* ColumnType = *It;
-			if (UE::Mass::IsA<FMassTag>(ColumnType))
-			{
-				Result = Context.DoesArchetypeHaveTag(ColumnType);
-			}
-			else if (UE::Mass::IsA<FMassFragment>(ColumnType))
-			{
-				Result = Context.DoesArchetypeHaveFragment(ColumnType);
-			}
-			else
-			{
-				const bool bIsTagOrFragment = false;
-				checkf(bIsTagOrFragment, TEXT("Attempting to check for a column type that is not a column or tag."));
-				return false;
-			}
-		}
-		return Result;
-	}
-
-	RowHandle FQueryContext_MassForwarder::GetCurrentRow() const
-	{
-		checkfSlow(CurrentRow != nullptr, TEXT("GetCurrentRow shouldn't be called if there are no rows."));
-		return *CurrentRow;
-	}
-
-	uint32 FQueryContext_MassForwarder::GetBatchRowCount() const
-	{
-		return static_cast<uint32>(Context.GetNumEntities());
-	}
-
-	TConstArrayView<RowHandle> FQueryContext_MassForwarder::GetBatchRowHandles() const
-	{
-		return MassEntitiesToRowsConversion(Context.GetEntities());
+		return &Environment;
 	}
 
-	bool FQueryContext_MassForwarder::CurrentBatchTableHasColumns(TConstArrayView<const UScriptStruct*> ColumnTypes) const
+	const FEnvironment* FProcessorContextEnvironment::GetEnvironment() const
 	{
-		return CurrentRowHasColumns(ColumnTypes);
+		return &Environment;
 	}
 
-	const void* FQueryContext_MassForwarder::GetColumnBatchAddress(const UScriptStruct* ColumnType) const
+	FRowHandleArrayView FProcessorContextEnvironment::GetBatchRows() const
 	{
-		checkf(UE::Mass::IsA<FMassFragment>(ColumnType), TEXT("Trying to retrieve a data pointer for column '%s' which doesn't have data."),
-			ColumnType ? *ColumnType->GetName() : TEXT("<null>"));
-		return Context.GetFragmentView(ColumnType).GetData();
+		TConstArrayView<RowHandle> Rows = MassEntitiesToRowsConversion(MassContext.GetEntities());
+		return FRowHandleArrayView(Rows, FRowHandleArrayView::EFlags::IsUnique);
 	}
 
-	void* FQueryContext_MassForwarder::GetMutableColumnBatchAddress(const UScriptStruct* ColumnType)
-	{
-		checkf(UE::Mass::IsA<FMassFragment>(ColumnType), TEXT("Trying to retrieve a data pointer for column '%s' which doesn't have data."),
-			ColumnType ? *ColumnType->GetName() : TEXT("<null>"));
-		return Context.GetMutableFragmentView(ColumnType).GetData();
-	}
-
-	bool FQueryContext_MassForwarder::NextBatch()
-	{
-		return false; // This only ever contains a single batch.
-	}
-
-	bool FQueryContext_MassForwarder::NextRow()
-	{
-		checkfSlow(CurrentRow != nullptr, TEXT("NextRow shouldn't be called if there are no rows."));
-		CurrentRow++;
-		return CurrentRow < EndRow;
-	}
-
-	void FQueryContext_MassForwarder::GetConstColumns(TArrayView<const void*> ColumnsData, TConstArrayView<const UScriptStruct*> ColumnTypes)
+	void FProcessorContextEnvironment::GetConstColumns(
+		TArrayView<const void*> ColumnsData, TConstArrayView<const UScriptStruct*> ColumnTypes) const
 	{
 		checkfSlow(ColumnsData.Num() >= ColumnTypes.Num(), TEXT("The list size of requested const column types (%i) doesn't match the "
 			"available space for the results (%i)"), ColumnTypes.Num(), ColumnsData.Num());
@@ -101,11 +43,12 @@ namespace UE::Editor::DataStorage::Queries
 		{
 			checkfSlow(UE::Mass::IsA<FMassFragment>(ColumnType),
 				TEXT("Unable to retrieve const column data for a column type without data: %s"), *ColumnType->GetFullName());
-			*ColumnsDataIt++ = Context.GetFragmentView(ColumnType).GetData();
+			*ColumnsDataIt++ = MassContext.GetFragmentView(ColumnType).GetData();
 		}
 	}
 
-	void FQueryContext_MassForwarder::GetMutableColumns(TArrayView<void*> ColumnsData, TConstArrayView<const UScriptStruct*> ColumnTypes)
+	void FProcessorContextEnvironment::GetMutableColumns(
+		TArrayView<void*> ColumnsData, TConstArrayView<const UScriptStruct*> ColumnTypes)
 	{
 		checkfSlow(ColumnsData.Num() >= ColumnTypes.Num(), TEXT("The list size of requested mutable column types (%i) doesn't match the "
 			"available space for the results (%i)"), ColumnTypes.Num(), ColumnsData.Num());
@@ -115,17 +58,12 @@ namespace UE::Editor::DataStorage::Queries
 		{
 			checkfSlow(UE::Mass::IsA<FMassFragment>(ColumnType),
 				TEXT("Unable to retrieve mutable column data for a column type without data: %s"), *ColumnType->GetFullName());
-			*ColumnsDataIt++ = Context.GetMutableFragmentView(ColumnType).GetData();
+			*ColumnsDataIt++ = MassContext.GetMutableFragmentView(ColumnType).GetData();
 		}
 	}
 
-	FEnvironment& FQueryContext_MassForwarder::GetEnvironment()
-	{
-		return Environment;
-	}
-
-	const FEnvironment& FQueryContext_MassForwarder::GetEnvironment() const
+	void FProcessorContextEnvironment::SetCurrentRow(RowHandle Row)
 	{
-		return Environment;
+		CurrentRow = Row;
 	}
 } // namespace UE::Editor::DataStorage::Queries
\ No newline at end of file
diff --git a/Engine/Plugins/Experimental/EditorDataStorage/Source/TedsCore/Private/Queries/TypedElementQueryContext_MassForwarder.h b/Engine/Plugins/Experimental/EditorDataStorage/Source/TedsCore/Private/Queries/TypedElementQueryContext_MassForwarder.h
--- a/Engine/Plugins/Experimental/EditorDataStorage/Source/TedsCore/Private/Queries/TypedElementQueryContext_MassForwarder.h
+++ b/Engine/Plugins/Experimental/EditorDataStorage/Source/TedsCore/Private/Queries/TypedElementQueryContext_MassForwarder.h
@@ -6,59 +6,32 @@
 #include "Elements/Framework/TypedElementQueryFunctions.h"
 #include "TypedElementDatabaseEnvironment.h"
 #include "Queries/QueryContextCapabilities/EditorDataStorageDynamicColumnInfo.h"
+#include "Queries/QueryContextCapabilities/RowBatch.h"
+#include "Queries/QueryContextCapabilities/SingleRow.h"
 
 struct FMassExecutionContext;
 
 namespace UE::Editor::DataStorage::Queries
 {
-	/** Query context implementation for single row and batch queries that call into TEDS Direct API. */
-	class FQueryContext_MassForwarder final : 
-		public SingleRowInfo, 
-		public RowBatchInfo,
-		public FDynamicColumnInfoContextCapability,
-		public IQueryFunctionResponse
+	struct FProcessorContextEnvironment final : IContextEnvironment
 	{
-	public:
-		static const FName Name;
+		FProcessorContextEnvironment(FEnvironment& Environment, FMassExecutionContext& MassContext);
 
-		FQueryContext_MassForwarder(FEnvironment& Environment, FMassExecutionContext& Context);
-
-		/**
-		 * @section ISingleRowInfo
-		 */
-
-		virtual bool CurrentRowHasColumns(TConstArrayView<const UScriptStruct*> ColumnTypes) const override;
-		virtual RowHandle GetCurrentRow() const override;
-
-		/**
-		 * @section IRowBatchInfo
-		 */
-		virtual uint32 GetBatchRowCount() const override;
-		virtual TConstArrayView<RowHandle> GetBatchRowHandles() const override;
-		virtual bool CurrentBatchTableHasColumns(TConstArrayView<const UScriptStruct*> ColumnTypes) const override;
-		virtual const void* GetColumnBatchAddress(const UScriptStruct* ColumnType) const override;
-		virtual void* GetMutableColumnBatchAddress(const UScriptStruct* ColumnType) override;
-
-		/**
-		 * @section IQueryFunctionResponse
-		 */
-		virtual bool NextBatch() override;
-		virtual bool NextRow() override;
-		virtual void GetConstColumns(TArrayView<const void*> ColumnsData, TConstArrayView<const UScriptStruct*> ColumnTypes) override;
+		virtual FEnvironment* GetEnvironment() override;
+		virtual const FEnvironment* GetEnvironment() const override;
+		
+		virtual FRowHandleArrayView GetBatchRows() const override;
+		virtual void GetConstColumns(TArrayView<const void*> ColumnsData, TConstArrayView<const UScriptStruct*> ColumnTypes) const override;
 		virtual void GetMutableColumns(TArrayView<void*> ColumnsData, TConstArrayView<const UScriptStruct*> ColumnTypes) override;
-
-		/**
-		 * @section IQueryEnvironmentRetrieval
-		 */
-		virtual FEnvironment& GetEnvironment() override;
-		virtual const FEnvironment& GetEnvironment() const override;
-
-	private:
+		virtual void SetCurrentRow(RowHandle Row) override;
+		
 		FEnvironment& Environment;
-		FMassExecutionContext& Context;
-		const RowHandle* CurrentRow;
-		const RowHandle* EndRow;
+		FMassExecutionContext& MassContext;
+		RowHandle CurrentRow = InvalidRowHandle;
 	};
-	
-	using QueryContext_MassForwarder = TQueryContextImpl<FQueryContext_MassForwarder>;
+
+	using QueryContext_ProcessorApi = TQueryContextImpl<FProcessorContextEnvironment,
+		FSingleRowInfoContextCapability,
+		FRowBatchInfoContextCapability,
+		FDynamicColumnInfoContextCapability>;
 } // namespace UE::Editor::DataStorage::Queries
diff --git a/Engine/Plugins/Experimental/EditorDataStorage/Source/TedsCore/Private/TypedElementDatabase.cpp b/Engine/Plugins/Experimental/EditorDataStorage/Source/TedsCore/Private/TypedElementDatabase.cpp
--- a/Engine/Plugins/Experimental/EditorDataStorage/Source/TedsCore/Private/TypedElementDatabase.cpp
+++ b/Engine/Plugins/Experimental/EditorDataStorage/Source/TedsCore/Private/TypedElementDatabase.cpp
@@ -149,7 +149,7 @@ namespace UE::Editor::DataStorage::Private
 		}
 		return Result;
 	}
-	
+
 	constexpr int32 ConvertTableHandleToIndex(const TableHandle InTableHandle)
 	{
 		return static_cast<int32>(InTableHandle);
@@ -192,7 +192,9 @@ void UEditorDataStorage::Initialize()
 	if (ActiveEditorEntityManager && ActiveEditorPhaseManager)
 	{
 		Environment = MakeShared<FEnvironment>(*this, *ActiveEditorEntityManager, *ActiveEditorPhaseManager);
-
+		ContextEnvironment.Environment = Environment.Get();
+		ContextEnvironment.DirectContext = this;
+		
 		using PhaseType = std::underlying_type_t<EQueryTickPhase>;
 		for (PhaseType PhaseId = 0; PhaseId < static_cast<PhaseType>(EQueryTickPhase::Max); ++PhaseId)
 		{
@@ -414,7 +416,7 @@ UE::Editor::DataStorage::TableHandle UEditorDataStorage::FindTable(const FName&
 bool UEditorDataStorage::IsTableValid(UE::Editor::DataStorage::TableHandle Table) const
 {
 	using namespace UE::Editor::DataStorage;
-	
+
 	return Table != InvalidTableHandle && Table < Tables.Num();
 }
 
@@ -454,10 +456,10 @@ UE::Editor::DataStorage::RowHandle UEditorDataStorage::AddRow(UE::Editor::DataSt
 	using namespace UE::Editor::DataStorage;
 
 	FActiveCallScope Scope(*this);
-	
+
 	checkf(IsTableValid(Table), TEXT("Attempting to add a row to a non-existing table."));
 	const int32 TableHandleAsIndex = Private::ConvertTableHandleToIndex(Table);
-	return ActiveEditorEntityManager 
+	return ActiveEditorEntityManager
 		? ActiveEditorEntityManager->CreateEntity(Tables[TableHandleAsIndex]).AsNumber()
 		: UE::Editor::DataStorage::InvalidRowHandle;
 }
@@ -470,10 +472,10 @@ UE::Editor::DataStorage::RowHandle UEditorDataStorage::AddRow(UE::Editor::DataSt
 	if (ActiveEditorEntityManager)
 	{
 		FActiveCallScope Scope(*this);
-		
+
 		checkf(IsTableValid(Table), TEXT("Attempting to a row to a non-existing table."));
 		const int32 TableHandleAsIndex = Private::ConvertTableHandleToIndex(Table);
-		
+
 		TArray<FMassEntityHandle> Entity;
 		Entity.Reserve(1);
 		TSharedRef<FMassEntityManager::FEntityCreationContext> Context =
@@ -497,7 +499,7 @@ bool UEditorDataStorage::AddRow(UE::Editor::DataStorage::RowHandle ReservedRow,
 	if (ActiveEditorEntityManager)
 	{
 		FActiveCallScope Scope(*this);
-		
+
 		ActiveEditorEntityManager->BuildEntity(FMassEntityHandle::FromNumber(ReservedRow), Tables[TableHandleAsIndex]);
 		return true;
 	}
@@ -515,10 +517,10 @@ bool UEditorDataStorage::AddRow(UE::Editor::DataStorage::RowHandle ReservedRow,
 	if (ActiveEditorEntityManager)
 	{
 		FActiveCallScope Scope(*this);
-		
+
 		checkf(IsTableValid(Table), TEXT("Attempting to add a row to a non-existing table."));
 		const int32 TableHandleAsIndex = Private::ConvertTableHandleToIndex(Table);
-		
+
 		TSharedRef<FMassEntityManager::FEntityCreationContext> Context =
 			ActiveEditorEntityManager->BatchCreateReservedEntities(Tables[TableHandleAsIndex], { FMassEntityHandle::FromNumber(ReservedRow) });
 
@@ -552,7 +554,7 @@ bool UEditorDataStorage::BatchAddRow(
 		Entities.Reserve(Count);
 		TSharedRef<FMassEntityManager::FEntityCreationContext> Context =
 			ActiveEditorEntityManager->BatchCreateEntities(Tables[TableHandleAsIndex], Count, Entities);
-		
+
 		for (FMassEntityHandle Entity : Entities)
 		{
 			OnCreated(Entity.AsNumber());
@@ -572,7 +574,7 @@ bool UEditorDataStorage::BatchAddRow(UE::Editor::DataStorage::TableHandle Table,
 	if (ActiveEditorEntityManager)
 	{
 		FActiveCallScope Scope(*this);
-		
+
 		if (!IsTableValid(Table))
 		{
 			checkf(
@@ -582,7 +584,7 @@ bool UEditorDataStorage::BatchAddRow(UE::Editor::DataStorage::TableHandle Table,
 			);
 			return false;
 		}
-		
+
 		const int32 TableHandleAsIndex = Private::ConvertTableHandleToIndex(Table);
 		TSharedRef<FMassEntityManager::FEntityCreationContext> Context =
 			ActiveEditorEntityManager->BatchCreateReservedEntities(Tables[TableHandleAsIndex], RowsToMassEntitiesConversion(ReservedHandles));
@@ -691,7 +693,7 @@ bool UEditorDataStorage::IsRowAssignedUnsafe(UE::Editor::DataStorage::RowHandle
 	return UE::Editor::DataStorage::Legacy::FCommandBuffer::Execute_IsRowAssigned(*ActiveEditorEntityManager, Row);
 }
 
-void UEditorDataStorage::FilterRowsBy(UE::Editor::DataStorage::FRowHandleArray& Result,  UE::Editor::DataStorage::FRowHandleArrayView Input, 
+void UEditorDataStorage::FilterRowsBy(UE::Editor::DataStorage::FRowHandleArray& Result, UE::Editor::DataStorage::FRowHandleArrayView Input,
 	EFilterOptions Options, UE::Editor::DataStorage::Queries::TQueryFunction<bool>& Filter)
 {
 	using namespace UE::Editor::DataStorage;
@@ -704,7 +706,8 @@ void UEditorDataStorage::FilterRowsBy(UE::Editor::DataStorage::FRowHandleArray&
 			bool bResult;
 			explicit FQueryResult(const bool bResultValue)
 				: bResult(bResultValue)
-			{}
+			{
+			}
 
 			virtual void Add(bool bResultValue) override
 			{
@@ -712,17 +715,19 @@ void UEditorDataStorage::FilterRowsBy(UE::Editor::DataStorage::FRowHandleArray&
 			}
 		};
 
-		QueryContext_DirectSingleApi Context(*Environment, *this);
-		Context.CheckCompatiblity(Filter);
-
+		FDirectContextEnvironment LocalContextEnvironment(*Environment, *this);
+		QueryContext_DirectApi Context(LocalContextEnvironment);
+		
 		const bool bVerifyColumns = !EnumHasAnyFlags(Options, EFilterOptions::IncludeRowsMissingColumns);
 		const bool bAcceptedResult = !EnumHasAnyFlags(Options, EFilterOptions::InvertFilter);
 
 		for (RowHandle Row : Input)
 		{
-			Context.GetContextImplementation().SetCurrentRow(Row);
+			LocalContextEnvironment.SetBatch(FRowHandleArrayView(&Row, 1,
+				FRowHandleArrayView::EFlags::IsSorted | FRowHandleArrayView::EFlags::IsUnique));
+
 			FQueryResult QueryResult(!bVerifyColumns);
-			Filter.Call<EFunctionCallConfig::VerifyColumns>(QueryResult, Context, Context.GetContextImplementation());
+			Filter.Call<EFunctionCallConfig::VerifyColumns>(QueryResult, Context, LocalContextEnvironment);
 			if (QueryResult.bResult == bAcceptedResult)
 			{
 				Result.Add(Row);
@@ -1787,7 +1792,7 @@ void UEditorDataStorage::Reset()
 	}
 	Tables.Reset();
 	TableNameLookup.Reset();
-
+	
 	checkf(Environment.IsUnique(), TEXT("UEditorDataStorage should hold the last reference to the Environment and be the one to delete it."));
 	Environment.Reset();
 
@@ -1832,6 +1837,7 @@ const UE::Editor::DataStorage::FEnvironment& UEditorDataStorage::GetEnvironment(
 	return *Environment;
 }
 
+
 FMassArchetypeHandle UEditorDataStorage::LookupArchetype(UE::Editor::DataStorage::TableHandle InTableHandle) const
 {
 	using namespace UE::Editor::DataStorage;
diff --git a/Engine/Plugins/Experimental/EditorDataStorage/Source/TedsCore/Private/TypedElementDatabase.h b/Engine/Plugins/Experimental/EditorDataStorage/Source/TedsCore/Private/TypedElementDatabase.h
--- a/Engine/Plugins/Experimental/EditorDataStorage/Source/TedsCore/Private/TypedElementDatabase.h
+++ b/Engine/Plugins/Experimental/EditorDataStorage/Source/TedsCore/Private/TypedElementDatabase.h
@@ -11,6 +11,7 @@
 #include "TypedElementDatabaseCommandBuffer.h"
 #include "UObject/ObjectMacros.h"
 #include "Queries/QueryContextCapabilities/EditorDataStorageDynamicColumnInfo.h"
+#include "Queries/TypedElementQueryContext_DirectSingleApi.h"
 
 #include "TypedElementDatabase.generated.h"
 
@@ -27,12 +28,15 @@ namespace UE::Editor::DataStorage
 	class FEnvironment;
 } // namespace UE::Editor::DataStorage
 
+using ImplementationForwarder = UE::Editor::DataStorage::Queries::Private::TImplementationForwarder<
+	UE::Editor::DataStorage::Queries::FDirectContextEnvironment,
+	UE::Editor::DataStorage::ICoreProvider,
+	UE::Editor::DataStorage::Queries::FDynamicColumnInfoContextCapability>;
+
 UCLASS(MinimalAPI)
 class UEditorDataStorage final
 	: public UObject
-	, public UE::Editor::DataStorage::Queries::Private::TImplementationForwarder<
-		UE::Editor::DataStorage::ICoreProvider, 
-		UE::Editor::DataStorage::Queries::FDynamicColumnInfoContextCapability>
+	, public ImplementationForwarder
 {
 	GENERATED_BODY()
 
@@ -237,8 +241,8 @@ public:
 	UE_API TSharedPtr<UE::Editor::DataStorage::FEnvironment> GetEnvironmentPtr();
 	UE_API TSharedPtr<const UE::Editor::DataStorage::FEnvironment> GetEnvironmentPtr() const;
 
-	virtual UE::Editor::DataStorage::FEnvironment& GetEnvironment() override;
-	virtual const UE::Editor::DataStorage::FEnvironment& GetEnvironment() const override;
+	UE::Editor::DataStorage::FEnvironment& GetEnvironment();
+	const UE::Editor::DataStorage::FEnvironment& GetEnvironment() const;
 
 	UE_API FMassArchetypeHandle LookupArchetype(UE::Editor::DataStorage::TableHandle InTableHandle) const;
 
@@ -293,7 +297,7 @@ private:
 	TArray<FFactoryTypePair> Factories;
 
 	TSharedPtr<UE::Editor::DataStorage::FEnvironment> Environment;
-	
+
 	UE::Editor::DataStorage::FTypedElementOnDataStorageUpdate OnUpdateDelegate;
 	UE::Editor::DataStorage::FTypedElementOnDataStorageUpdate OnUpdateCompletedDelegate;
 	FDelegateHandle OnPreMassTickHandle;
diff --git a/Engine/Source/Runtime/TypedElementFramework/Public/DataStorage/CommonTypes.h b/Engine/Source/Runtime/TypedElementFramework/Public/DataStorage/CommonTypes.h
--- a/Engine/Source/Runtime/TypedElementFramework/Public/DataStorage/CommonTypes.h
+++ b/Engine/Source/Runtime/TypedElementFramework/Public/DataStorage/CommonTypes.h
@@ -66,6 +66,12 @@ namespace UE
 		using FColumn = FEditorDataStorageColumn;
 		using FTag = FEditorDataStorageTag;
 
+		template<typename T>
+		concept FunctionType =
+			requires { &std::remove_reference_t<T>::operator(); }			// non-generic lambdas & single-operator() functors
+			|| std::is_function_v<std::remove_pointer_t<std::decay_t<T>>>	// functions, function pointers & function references
+			|| std::is_member_function_pointer_v<std::decay_t<T>>;			// member function pointers
+
 		/**
 		 * Defines a dynamic type for a value tag
 		 * Example:
@@ -96,6 +102,7 @@ namespace UE
 			TYPEDELEMENTFRAMEWORK_API friend uint32 GetTypeHash(const FDynamicColumnDescription& Descriptor);
 			bool operator==(const FDynamicColumnDescription&) const = default;
 		};
+
 		// Standard callbacks.
 
 		using RowCreationCallbackRef = TFunctionRef<void(RowHandle Row)>;
diff --git a/Engine/Source/Runtime/TypedElementFramework/Public/Elements/Framework/TypedElementQueryCapabilities.h b/Engine/Source/Runtime/TypedElementFramework/Public/Elements/Framework/TypedElementQueryCapabilities.h
--- a/Engine/Source/Runtime/TypedElementFramework/Public/Elements/Framework/TypedElementQueryCapabilities.h
+++ b/Engine/Source/Runtime/TypedElementFramework/Public/Elements/Framework/TypedElementQueryCapabilities.h
@@ -3,10 +3,13 @@
 #pragma once
 
 #include <type_traits>
+#include "Concepts/SameAs.h"
 #include "Containers/ContainersFwd.h"
+#include "DataStorage/CommonTypes.h"
 #include "DataStorage/Handles.h"
 #include "DataStorage/Queries/Conditions.h"
 #include "Elements/Framework/TypedElementQueryFunctionArgumentTypes.h"
+#include "Elements/Framework/TypedElementRowHandleArrayView.h"
 #include "Misc/EnumClassFlags.h"
 #include "UObject/NameTypes.h"
 
@@ -33,9 +36,163 @@ namespace UE::Editor::DataStorage::Queries
 		virtual ~IContextCapability() = default;
 	};
 
-#define WithWrappers 1	
+	template<typename T>
+	concept ContextCapability = std::is_base_of_v<IContextCapability, T> && !std::is_same_v<IContextCapability, T>;
+
+	template<ContextCapability Capability>
+	struct ImplementsContextCapability 
+	{
+		static constexpr int32 Id = Capability::Id;
+		inline static FName Name = Capability::Name;
+	};
+
+	template<typename ReturnType>
+	class TQueryFunction;
+
+	struct IContextContract;
+	struct IQueryFunctionResponse;
+
+	template<typename Return, FunctionType Function>
+	TQueryFunction<Return> BuildQueryFunction(Function&& Callback); // forward declaration
+
+	// Here only to make sure that the argument shows up in compiler errors.
+
+	template <typename T>
+	constexpr T& VarType(T& Value);
+	
+	template <typename T>
+	constexpr const T& VarType(const T& Value);
+	
+	template <typename T>
+	constexpr T&& VarType(T&& Value);
+
+#define ArgType(Type, Name) Type
+#define ArgName(Type, Name) Name
 #define ArgTypeName(Type, Name) Type Name
 
+	// Fill out the FCapabilityInfo
+
+	template<typename Implementation, typename Capability>
+	struct FCapabilityCompatibilityCheck
+	{
+	};
+
+#define CapabilityStart(InCapability, InFlags) \
+	template<typename T> \
+	struct I##InCapability ; \
+	\
+	template<typename Implementation> \
+	struct FCapabilityCompatibilityCheck<Implementation, I##InCapability <IContextCapability> > \
+	{ \
+		static constexpr bool Value =
+
+#define Function0(Capability, Return, Function) \
+	requires(Implementation& obj) \
+	{ \
+		{ obj. Function () } -> UE::CSameAs< Return >; \
+	} &&  
+#define Function1(Capability, Return, Function, Arg1) \
+	requires(Implementation& obj, ArgTypeName Arg1) \
+	{ \
+		{ obj. Function ( VarType< ArgType Arg1 >( ArgName Arg1 ) ) } -> UE::CSameAs< Return >; \
+	} && 
+#define Function2(Capability, Return, Function, Arg1, Arg2) \
+	requires(Implementation& obj, ArgTypeName Arg1, ArgTypeName Arg2) \
+	{ \
+		{ obj. Function ( \
+			VarType< ArgType Arg1 >( ArgName Arg1 ), \
+			VarType< ArgType Arg2 >( ArgName Arg2 ) ) } -> UE::CSameAs< Return >; \
+	} && 
+#define Function3(Capability, Return, Function, Arg1, Arg2, Arg3) \
+	requires(Implementation& obj, ArgTypeName Arg1, ArgTypeName Arg2, ArgTypeName Arg3) \
+	{ \
+		{ obj. Function ( \
+			VarType< ArgType Arg1 >(ArgName Arg1), \
+			VarType< ArgType Arg2 >(ArgName Arg2), \
+			VarType< ArgType Arg3 >(ArgName Arg3)) } -> UE::CSameAs< Return >; \
+	} &&
+#define Function4(Capability, Return, Function, Arg1, Arg2, Arg3, Arg4) \
+	requires(Implementation& obj, ArgTypeName Arg1, ArgTypeName Arg2, ArgTypeName Arg3, ArgTypeName Arg4) \
+	{ \
+		{ obj. Function ( \
+			VarType< ArgType Arg1 >(ArgName Arg1), \
+			VarType< ArgType Arg2 >(ArgName Arg2), \
+			VarType< ArgType Arg3 >(ArgName Arg3), \
+			VarType< ArgType Arg4 >(ArgName Arg4)) } -> UE::CSameAs< Return >; \
+	} &&
+
+#define ConstFunction0(Capability, Return, Function) \
+	requires(const Implementation& obj) \
+	{ \
+		{ obj. Function () } -> UE::CSameAs< Return >; \
+	} && 
+#define ConstFunction1(Capability, Return, Function, Arg1) \
+	requires(const Implementation& obj, ArgTypeName Arg1) \
+	{ \
+		{ obj. Function ( VarType< ArgType Arg1 >( ArgName Arg1 ) ) } -> UE::CSameAs< Return >; \
+	} && 
+#define ConstFunction2(Capability, Return, Function, Arg1, Arg2) \
+	requires(const Implementation& obj, ArgTypeName Arg1, ArgTypeName Arg2) \
+	{ \
+		{ obj. Function ( \
+			VarType< ArgType Arg1 >( ArgName Arg1 ), \
+			VarType< ArgType Arg2 >( ArgName Arg2 ) ) } -> UE::CSameAs< Return >; \
+	} && 
+#define ConstFunction3(Capability, Return, Function, Arg1, Arg2, Arg3) \
+	requires(const Implementation& obj, ArgTypeName Arg1, ArgTypeName Arg2, ArgTypeName Arg3) \
+	{ \
+		{ obj. Function ( \
+			VarType< ArgType Arg1 >(ArgName Arg1), \
+			VarType< ArgType Arg2 >(ArgName Arg2), \
+			VarType< ArgType Arg3 >(ArgName Arg3)) } -> UE::CSameAs< Return >; \
+	} &&
+#define ConstFunction4(Capability, Return, Function, Arg1, Arg2, Arg3, Arg4) \
+	requires(const Implementation& obj, ArgTypeName Arg1, ArgTypeName Arg2, ArgTypeName Arg3, ArgTypeName Arg4) \
+	{ \
+		{ obj. Function ( \
+			VarType< ArgType Arg1 >(ArgName Arg1), \
+			VarType< ArgType Arg2 >(ArgName Arg2), \
+			VarType< ArgType Arg3 >(ArgName Arg3), \
+			VarType< ArgType Arg4 >(ArgName Arg4)) } -> UE::CSameAs< Return >; \
+	} &&
+
+#define CapabilityEnd(InCapability) \
+			true; \
+	};
+
+#include "Elements/Framework/TypedElementQueryCapabilities.inl"
+
+#undef Function0
+#undef Function1
+#undef Function2
+#undef Function3
+#undef Function4
+#undef ConstFunction0
+#undef ConstFunction1
+#undef ConstFunction2
+#undef ConstFunction3
+#undef ConstFunction4
+#undef CapabilityStart
+#undef CapabilityEnd
+
+	// Fill in the ability interface.
+
+	// Assign a unique id to each capability, starting at 0. Since __COUNTER__ is a preprocessor defined macro other headers are free
+	// to use it as well, leading to the starting address not being 0, hence recording the first occurrence and subtracting it from
+	// the ids.
+	static constexpr int32 CapabilityBaseId = __COUNTER__;
+
+#define WithWrappers 1	
+
+#define CapabilityStart(InCapability, InFlags) \
+	template<typename Base> \
+	struct I##InCapability : Base \
+	{ \
+		inline static const FName Name = #InCapability ; \
+		static constexpr int32 Id = __COUNTER__ - CapabilityBaseId - 1; \
+		static constexpr EContextCapabilityFlags Flags = InFlags; \
+		virtual ~I##InCapability () = default; \
+
 #define Function0(Capability, Return, Function) virtual Return Function() = 0;
 #define Function1(Capability, Return, Function, Arg1) virtual Return Function(ArgTypeName Arg1 ) = 0;
 #define Function2(Capability, Return, Function, Arg1, Arg2) virtual Return Function(ArgTypeName Arg1 , ArgTypeName Arg2 ) = 0;
@@ -48,21 +205,16 @@ namespace UE::Editor::DataStorage::Queries
 #define ConstFunction3(Capability, Return, Function, Arg1, Arg2, Arg3) virtual Return Function(ArgTypeName Arg1 , ArgTypeName Arg2 , ArgTypeName Arg3 ) const = 0;
 #define ConstFunction4(Capability, Return, Function, Arg1, Arg2, Arg3, Arg4) virtual Return Function(ArgTypeName Arg1 , ArgTypeName Arg2 , ArgTypeName Arg3 , ArgTypeName Arg4 ) const = 0;
 
-#define CapabilityStart(InCapability, InFlags) \
-	template<typename Base> \
-	struct I##InCapability : Base \
-	{ \
-		inline static const FName Name = #InCapability ; \
-		static constexpr EContextCapabilityFlags Flags = InFlags; \
-		virtual ~I##InCapability () = default;
-
-#define CapabilityEnd(Capability) \
+#define CapabilityEnd(InCapability) \
 	}; \
-	using Capability = I##Capability <IContextCapability>;
+	using InCapability = I##InCapability <IContextCapability>; \
+
 
 #include "Elements/Framework/TypedElementQueryCapabilities.inl"
 
 #undef ArgTypeName
+#undef ArgType
+#undef ArgName
 #undef Function0
 #undef Function1
 #undef Function2
diff --git a/Engine/Source/Runtime/TypedElementFramework/Public/Elements/Framework/TypedElementQueryCapabilities.inl b/Engine/Source/Runtime/TypedElementFramework/Public/Elements/Framework/TypedElementQueryCapabilities.inl
--- a/Engine/Source/Runtime/TypedElementFramework/Public/Elements/Framework/TypedElementQueryCapabilities.inl
+++ b/Engine/Source/Runtime/TypedElementFramework/Public/Elements/Framework/TypedElementQueryCapabilities.inl
@@ -9,7 +9,7 @@
  * Capability to provide information about the currently active row.
  */
 CapabilityStart(SingleRowInfo, EContextCapabilityFlags::SupportsSingle)
-	/*
+	/**
 	 * Returns whether a column matches the requested type or not. This version only applies to current row. This version is faster than
 	 * querying for an arbitrary row.
 	 */
@@ -18,7 +18,7 @@ CapabilityStart(SingleRowInfo, EContextCapabilityFlags::SupportsSingle)
 	ConstFunction0(SingleRowInfo, RowHandle, GetCurrentRow)
 
 #if defined(WithWrappers)
-	/*
+	/**
 	 * Returns whether a column matches the requested type or not. This version only applies to current row. This version is faster than
 	 * querying for an arbitrary row.
 	 */
@@ -37,7 +37,7 @@ CapabilityStart(RowBatchInfo, EContextCapabilityFlags::SupportsBatch)
 	/** Returns the number of rows in the current batch. */
 	ConstFunction0(RowBatchInfo, uint32, GetBatchRowCount)
 	/** Returns an view with the rows used by this batch. */
-	ConstFunction0(RowBatchInfo, TConstArrayView<RowHandle>, GetBatchRowHandles)
+	ConstFunction0(RowBatchInfo, FRowHandleArrayView, GetBatchRowHandles)
 	/** Checks if the rows in the current batch have the requested columns. */
 	ConstFunction1(RowBatchInfo, bool, CurrentBatchTableHasColumns, (TConstArrayView<const UScriptStruct*>, Columns))
 	/** Direct access to the const addresses of a column or null if not found in this query. */
@@ -107,7 +107,7 @@ CapabilityStart(DynamicColumnInfo, EContextCapabilityFlags::SupportsSingle | ECo
 	/**
 	 * Executes the given callback for each known dynamic column that derives from the base template provided.
 	 */
-	ConstFunction2(DynamicColumnInfo, void, ForEachDynamicColumnType, (const UScriptStruct&, Template), (TFunctionRef<void(const UScriptStruct& Type)>, Callback))
+	ConstFunction2(DynamicColumnInfo, void, ForEachDynamicColumnType, (const UScriptStruct&, TemplateType), (TFunctionRef<void(const UScriptStruct& Type)>, Callback))
 
 #if defined(WithWrappers)
 	/**
diff --git a/Engine/Source/Runtime/TypedElementFramework/Public/Elements/Framework/TypedElementQueryContext.h b/Engine/Source/Runtime/TypedElementFramework/Public/Elements/Framework/TypedElementQueryContext.h
--- a/Engine/Source/Runtime/TypedElementFramework/Public/Elements/Framework/TypedElementQueryContext.h
+++ b/Engine/Source/Runtime/TypedElementFramework/Public/Elements/Framework/TypedElementQueryContext.h
@@ -10,9 +10,6 @@ namespace UE::Editor::DataStorage::Queries
 {
 	struct IContextContract;
 
-	template<typename T>
-	concept ContextCapability = std::is_base_of_v<IContextCapability, T> && !std::is_same_v<IContextCapability, T>;
-
 	template<ContextCapability... Capabilities>
 	struct TCapabilityStore : Capabilities... 
 	{
@@ -28,7 +25,7 @@ namespace UE::Editor::DataStorage::Queries
 	 * query callback to improve re-usability or use a predefined query context if the query callback doesn't need to be reused.
 	 */
 	template<ContextCapability... CapabilityTypes>
-	struct TQueryContext final : Private::TForwarder<CapabilityTypes...>
+	struct TQueryContext final : Private::TForwarder< CapabilityTypes...>
 	{
 		using Capabilities = TCapabilityStore<CapabilityTypes...>;
 
@@ -38,14 +35,15 @@ namespace UE::Editor::DataStorage::Queries
 		TQueryContext() = default;
 		TQueryContext(IContextContract& Contract);
 	};
-	
+
+
 	//
 	// Implementations.
 	//
 	
 	template<ContextCapability... CapabilityTypes>
 	TQueryContext<CapabilityTypes...>::TQueryContext(IContextContract& Contract)
-		: Private::TForwarder<CapabilityTypes...>(Contract)
+		: Private::TForwarder< CapabilityTypes...>(Contract)
 	{
 	}
 } // namespace UE::Editor::DataStorage::Queries
diff --git a/Engine/Source/Runtime/TypedElementFramework/Public/Elements/Framework/TypedElementQueryContextImplementation.h b/Engine/Source/Runtime/TypedElementFramework/Public/Elements/Framework/TypedElementQueryContextImplementation.h
--- a/Engine/Source/Runtime/TypedElementFramework/Public/Elements/Framework/TypedElementQueryContextImplementation.h
+++ b/Engine/Source/Runtime/TypedElementFramework/Public/Elements/Framework/TypedElementQueryContextImplementation.h
@@ -11,6 +11,8 @@
 
 #define ArgTypeName(Type, Name) Type Name
 
+#define CapabilityStart(Capability, Flags)
+
 #define Function0(Capability, Return, Function) virtual Return Function() override;
 #define Function1(Capability, Return, Function, Arg1) virtual Return Function(ArgTypeName Arg1 ) override;
 #define Function2(Capability, Return, Function, Arg1, Arg2) virtual Return Function(ArgTypeName Arg1 , ArgTypeName Arg2 ) override;
@@ -23,7 +25,6 @@
 #define ConstFunction3(Capability, Return, Function, Arg1, Arg2, Arg3) virtual Return Function(ArgTypeName Arg1 , ArgTypeName Arg2 , ArgTypeName Arg3 ) const override;
 #define ConstFunction4(Capability, Return, Function, Arg1, Arg2, Arg3, Arg4) virtual Return Function(ArgTypeName Arg1 , ArgTypeName Arg2 , ArgTypeName Arg3 , ArgTypeName Arg4 ) const override;
 
-#define CapabilityStart(Capability, Flags)
 #define CapabilityEnd(Capability)
 
 #define WithWrappers 1
@@ -33,36 +34,32 @@ namespace UE::Editor::DataStorage::Queries
 	/** 
 	 * Base used for context implementations. This will verify if a context is correctly implemented.
 	 * Classes that implement capabilities should use this template to automatically satisfy the requirements
-	 * of a query contract. Any missing functionality will be filled in by this template with a placeholder
-	 * function that asserts. The supported capabilities are automatically extracted from the implementation type
-	 * but can be overruled by providing a list of capabilities as extra template arguments. Note that any
-	 * capabilities provided this way do need to be implemented by the implementation type.
+	 * of a query contract. Any missing functionality in a capability implementation will be reported as a compilation 
+	 * failure. This template will satisfy the IContextContract, but will implement placeholder functions for the
+	 * functions of capabilities that are not supported.
 	 */
-	template<typename ImplementationType, ContextCapability... SupportedCapabilities>
+	template<typename EnvironmentType, typename... ImplementationTypes>
 	class TQueryContextImpl final : public IContextContract
 	{
 	public:
-		template<typename... TArgs>
-		explicit TQueryContextImpl(TArgs&&... Args);
+		explicit TQueryContextImpl(EnvironmentType& Environment);
 
 		virtual ~TQueryContextImpl() override = default;
 
-		static bool SupportsCapabilities(TConstArrayView<FName> Capabilities);
-
 #include "Elements/Framework/TypedElementQueryCapabilities.inl"
 
 		// Generic access
-		template<typename ReturnType>
-		bool CheckCompatiblity(const TQueryFunction<ReturnType>& Function) const;
-		const ImplementationType& GetContextImplementation() const;
-		ImplementationType& GetContextImplementation();
-	
-	private:
+
 		template<ContextCapability RequestedCapability>
 		static constexpr bool SupportsCapability();
 		static bool SupportsCapability(const FName& Capability);
+		static constexpr bool SupportsCapability(int32 CapabilityId);
+
+		static TConstArrayView<FName> StaticGetSupportedCapabilities();
 
-		ImplementationType Implementation;
+		virtual TConstArrayView<FName> GetSupportedCapabilities() const override;
+		
+		EnvironmentType& Environment;
 	};
 }
 // namespace UE::Editor::DataStorage::Queries
diff --git a/Engine/Source/Runtime/TypedElementFramework/Public/Elements/Framework/TypedElementQueryContextImplementation.inl b/Engine/Source/Runtime/TypedElementFramework/Public/Elements/Framework/TypedElementQueryContextImplementation.inl
--- a/Engine/Source/Runtime/TypedElementFramework/Public/Elements/Framework/TypedElementQueryContextImplementation.inl
+++ b/Engine/Source/Runtime/TypedElementFramework/Public/Elements/Framework/TypedElementQueryContextImplementation.inl
@@ -1,63 +1,71 @@
 // Copyright Epic Games, Inc. All Rights Reserved.
 
 #include <type_traits>
+#include "Containers/StringFwd.h"
 #include "DataStorage/CommonTypes.h"
 
 namespace UE::Editor::DataStorage::Queries
 {
-	template<typename QueryContext, ContextCapability... SupportedCapabilities>
-	template<typename... TArgs>
-	TQueryContextImpl<QueryContext, SupportedCapabilities...>::TQueryContextImpl(TArgs&&... Args)
-		: Implementation(Forward<TArgs>(Args)...)
+	namespace Private
+	{
+		template<ContextCapability RequestedCapability, typename... ImplementationTypes>
+		struct FindImplementationType
+		{ 
+			using Type = void;
+		};
+
+		template<ContextCapability RequestedCapability, typename ImplementationType, typename... ImplementationTypes>
+		struct FindImplementationType<RequestedCapability, ImplementationType, ImplementationTypes...> {
+			using Type = std::conditional_t<
+				std::is_base_of_v<ImplementsContextCapability<RequestedCapability>, ImplementationType>,
+				ImplementationType,
+				typename FindImplementationType<RequestedCapability, ImplementationTypes...>::Type
+			>;
+		};
+
+		template<ContextCapability RequestedCapability, typename... ImplementationTypes>
+		using ImplementationType = typename FindImplementationType<RequestedCapability, ImplementationTypes...>::Type;
+	}
+
+	template<typename EnvironmentType, typename... ImplementationTypes>
+	TQueryContextImpl<EnvironmentType, ImplementationTypes...>::TQueryContextImpl(EnvironmentType& Environment)
+		: Environment(Environment)
 	{
 	}
 
-	template<typename QueryContext, ContextCapability... SupportedCapabilities>
+
+	template<typename EnvironmentType, typename... ImplementationTypes>
 	template<ContextCapability RequestedCapability>
-	constexpr bool TQueryContextImpl<QueryContext, SupportedCapabilities...>::SupportsCapability()
+	constexpr bool TQueryContextImpl<EnvironmentType, ImplementationTypes...>::SupportsCapability()
 	{
-		if constexpr (sizeof...(SupportedCapabilities) > 0)
-		{
-			return std::is_base_of_v<RequestedCapability, QueryContext> && (std::is_same_v<RequestedCapability, SupportedCapabilities> || ...);
-		}
-		else
-		{
-			return std::is_base_of_v<RequestedCapability, QueryContext>;
-		}
-	}
-
-	template<typename QueryContext, ContextCapability... SupportedCapabilities>
-	bool TQueryContextImpl<QueryContext, SupportedCapabilities...>::SupportsCapability(const FName& Capability)
+		return (std::is_base_of_v<ImplementsContextCapability<RequestedCapability>, ImplementationTypes> || ...);
+	}
+
+	template<typename EnvironmentType, typename... ImplementationTypes>
+	bool TQueryContextImpl<EnvironmentType, ImplementationTypes...>::SupportsCapability(const FName& Capability)
 	{
-		if constexpr (sizeof...(SupportedCapabilities) > 0)
-		{
-			return IContextContract::SupportsCapability<QueryContext>(Capability) && ((SupportedCapabilities::Name == Capability) || ...);
-		}
-		else
-		{
-			return IContextContract::SupportsCapability<QueryContext>(Capability);
-		}
-	}
-
-	template<typename QueryContext, ContextCapability... SupportedCapabilities>
-	bool TQueryContextImpl<QueryContext, SupportedCapabilities...>::SupportsCapabilities(TConstArrayView<FName> Capabilities)
+		return ((Capability == ImplementationTypes::Name) || ...);
+	}
+
+	template<typename EnvironmentType, typename... ImplementationTypes>
+	constexpr bool TQueryContextImpl<EnvironmentType, ImplementationTypes...>::SupportsCapability(int32 CapabilityId)
 	{
-		for (const FName& Capability : Capabilities)
-		{
-			if (!SupportsCapability(Capability))
-			{
-				return false;
-			}
-		}
-		return true;
+		return ((CapabilityId == ImplementationTypes::Id) || ...);
 	}
 
+
 	// Use a macro for this since it's impossible to do with a template since the function pointer that would be needed might not exist.
 	// The macro also has the added benefit that it can print the function name for better reporting.
 #define CallFunction(Capability, ReturnValue, FunctionName, ...) \
-	if constexpr (SupportsCapability< Capability >()) \
+	if constexpr (SupportsCapability<Capability>()) \
 	{ \
-		return Implementation.FunctionName( __VA_ARGS__ ); \
+		using ImplementationType = Private::ImplementationType<Capability, ImplementationTypes...>; \
+		constexpr bool bFunctionFound = requires() \
+				{ \
+					{ ImplementationType:: FunctionName (Environment __VA_OPT__(,) __VA_ARGS__ ) } -> UE::CSameAs< ReturnValue >; \
+				}; \
+		static_assert(bFunctionFound, "Function '" #FunctionName "' in capability '" #Capability "' is missing from the implementation." ); \
+		return ImplementationType:: FunctionName ( Environment __VA_OPT__(,) __VA_ARGS__ ); \
 	} \
 	else \
 	{ \
@@ -67,10 +75,13 @@ namespace UE::Editor::DataStorage::Queries
 
 #define ArgTypeName(Type, Name) Type Name
 #define ArgName(Type, Name) Name
+#define ForwardArg(Type, Name) Forward< Type >( Name )
 
 #define FunctionCommon(ReturnType) \
-	template<typename QueryContext, ContextCapability... SupportedCapabilities> \
-	ReturnType TQueryContextImpl<QueryContext, SupportedCapabilities...>::
+	template<typename EnvironmentType, typename... ImplementationTypes> \
+	ReturnType TQueryContextImpl<EnvironmentType, ImplementationTypes...>::
+
+#define CapabilityStart(Capability, Flags)
 
 #define Function0(Capability, Return, Function) \
 	FunctionCommon(Return) Function() \
@@ -80,22 +91,22 @@ namespace UE::Editor::DataStorage::Queries
 #define Function1(Capability, Return, Function, Arg1) \
 	FunctionCommon(Return) Function(ArgTypeName Arg1) \
 	{ \
-		CallFunction(Capability, Return, Function, ArgName Arg1 ) \
+		CallFunction(Capability, Return, Function, ForwardArg Arg1 ) \
 	}
 #define Function2(Capability, Return, Function, Arg1, Arg2) \
 	FunctionCommon(Return) Function(ArgTypeName Arg1, ArgTypeName Arg2) \
 	{ \
-		CallFunction(Capability, Return, Function, ArgName Arg1, ArgName Arg2 ) \
+		CallFunction(Capability, Return, Function, ForwardArg Arg1, ForwardArg Arg2 ) \
 	}
 #define Function3(Capability, Return, Function, Arg1, Arg2, Arg3) \
 	FunctionCommon(Return) Function(ArgTypeName Arg1, ArgTypeName Arg2, ArgTypeName Arg3) \
 	{ \
-		CallFunction(Capability, Return, Function, ArgName Arg1, ArgName Arg2, ArgName Arg3 ) \
+		CallFunction(Capability, Return, Function, ForwardArg Arg1, ForwardArg Arg2, ForwardArg Arg3 ) \
 	}
 #define Function4(Capability, Return, Function, Arg1, Arg2, Arg3, Arg4) \
 	FunctionCommon(Return) Function(ArgTypeName Arg1, ArgTypeName Arg2, ArgTypeName Arg3, ArgTypeName Arg4) \
 	{ \
-		CallFunction(Capability, Return, Function, ArgName Arg1, ArgName Arg2, ArgName Arg3, ArgName Arg4 ) \
+		CallFunction(Capability, Return, Function, ForwardArg Arg1, ForwardArg Arg2, ForwardArg Arg3, ForwardArg Arg4 ) \
 	}
 
 #define ConstFunction0(Capability, Return, Function) \
@@ -106,31 +117,31 @@ namespace UE::Editor::DataStorage::Queries
 #define ConstFunction1(Capability, Return, Function, Arg1) \
 	FunctionCommon(Return) Function(ArgTypeName Arg1) const \
 	{ \
-		CallFunction(Capability, Return, Function, ArgName Arg1 ) \
+		CallFunction(Capability, Return, Function, ForwardArg Arg1 ) \
 	}
 #define ConstFunction2(Capability, Return, Function, Arg1, Arg2) \
 	FunctionCommon(Return) Function(ArgTypeName Arg1, ArgTypeName Arg2) const \
 	{ \
-		CallFunction(Capability, Return, Function, ArgName Arg1, ArgName Arg2 ) \
+		CallFunction(Capability, Return, Function, ForwardArg Arg1, ForwardArg Arg2 ) \
 	}
 #define ConstFunction3(Capability, Return, Function, Arg1, Arg2, Arg3) \
 	FunctionCommon(Return) Function(ArgTypeName Arg1, ArgTypeName Arg2, ArgTypeName Arg3) const \
 	{ \
-		CallFunction(Capability, Return, Function, ArgName Arg1, ArgName Arg2, ArgName Arg3 ) \
+		CallFunction(Capability, Return, Function, ForwardArg Arg1, ForwardArg Arg2, ForwardArg Arg3 ) \
 	}
 #define ConstFunction4(Capability, Return, Function, Arg1, Arg2, Arg3, Arg4) \
 	FunctionCommon(Return) Function(ArgTypeName Arg1, ArgTypeName Arg2, ArgTypeName Arg3, ArgTypeName Arg4) const \
 	{ \
-		CallFunction(Capability, Return, Function, ArgName Arg1, ArgName Arg2, ArgName Arg3, ArgName Arg4 ) \
+		CallFunction(Capability, Return, Function, ForwardArg Arg1, ForwardArg Arg2, ForwardArg Arg3, ForwardArg Arg4 ) \
 	}
 
-#define CapabilityStart(Capability, Flags)
 #define CapabilityEnd(Capability)
 
 #include "Elements/Framework/TypedElementQueryCapabilities.inl"
 
 #undef ArgTypeName
 #undef ArgName
+#undef ForwardArg
 #undef FunctionCommon
 #undef Function0
 #undef Function1
@@ -148,35 +159,21 @@ namespace UE::Editor::DataStorage::Queries
 
 	//
 	// Generic access
-	// 
-
-	template<typename QueryContext, ContextCapability... SupportedCapabilities>
-	template<typename ReturnType>
-	bool TQueryContextImpl<QueryContext, SupportedCapabilities...>::CheckCompatiblity(const TQueryFunction<ReturnType>& Function) const
+	
+	template<typename EnvironmentType, typename... ImplementationTypes>
+	TConstArrayView<FName> TQueryContextImpl<EnvironmentType, ImplementationTypes...>::StaticGetSupportedCapabilities()
 	{
-		bool Result = true;
-		for (const FName& Capability : Function.Capabilities)
-		{
-			if (!SupportsCapability(Capability))
-			{
-				Result = false;
-				checkf(false, TEXT("Requested query context '%s' is not supported by query context implementation."), 
-					*Capability.ToString());
-			}
-		}
-		return Result;
-	}
+		static_assert(sizeof...(ImplementationTypes) > 0,
+			"At least one context capability needs to be implemented to work with a query context contract.");
 
-	template<typename QueryContext, ContextCapability... SupportedCapabilities>
-	const QueryContext& TQueryContextImpl<QueryContext, SupportedCapabilities...>::GetContextImplementation() const
-	{
-		return Implementation;
+		static TArray<FName> Result = { ImplementationTypes::Name... };
+		return Result;
 	}
 
-	template<typename QueryContext, ContextCapability... SupportedCapabilities>
-	QueryContext& TQueryContextImpl<QueryContext, SupportedCapabilities...>::GetContextImplementation()
+	template<typename EnvironmentType, typename... ImplementationTypes>
+	TConstArrayView<FName> TQueryContextImpl<EnvironmentType, ImplementationTypes...>::GetSupportedCapabilities() const
 	{
-		return Implementation;
+		return StaticGetSupportedCapabilities();
 	}
 }
 // namespace UE::Editor::DataStorage::Queries
diff --git a/Engine/Source/Runtime/TypedElementFramework/Public/Elements/Framework/TypedElementQueryContract.h b/Engine/Source/Runtime/TypedElementFramework/Public/Elements/Framework/TypedElementQueryContract.h
--- a/Engine/Source/Runtime/TypedElementFramework/Public/Elements/Framework/TypedElementQueryContract.h
+++ b/Engine/Source/Runtime/TypedElementFramework/Public/Elements/Framework/TypedElementQueryContract.h
@@ -2,6 +2,7 @@
 
 #pragma once
 
+#include "DataStorage/Queries/Types.h"
 #include "Elements/Framework/TypedElementQueryCapabilities.h"
 
 namespace UE::Editor::DataStorage::Queries
@@ -29,43 +30,67 @@ namespace UE::Editor::DataStorage::Queries
 			// Work around for MSVC struggling to correctly expand fold expressions when providing the implementation for a template with 
 			// templated types.
 			template<typename Implementation, template<typename> typename Base>
-			constexpr static bool ImplementsCapability()
+			static bool ImplementsCapability(const FName& Capability)
 			{
-				return std::is_base_of_v<Base<IContextCapability>, Implementation>;
+				return 
+					FCapabilityCompatibilityCheck<Implementation, Base<IContextCapability>>::Value && 
+					Base<IContextCapability>::Name == Capability;
 			}
 
 			template<typename Implementation, template<typename> typename Base>
-			static bool ImplementsCapability(const FName& Capability)
+			constexpr static bool ImplementsCapability(int32 CapabilityId)
 			{
-				return ImplementsCapability<Implementation, Base>() && Base<IContextCapability>::Name == Capability;
+				return 
+					FCapabilityCompatibilityCheck<Implementation, Base<IContextCapability>>::Value && 
+					Base<IContextCapability>::Id == CapabilityId;
 			}
 
 		public:
+			static FName GetCapabilityName(int32 CapabilityId)
+			{
+				FName Result;
+				((Bases<IContextCapability>::Id == CapabilityId && (Result = Bases<IContextCapability>::Name, true)), ...);
+				return Result;
+			}
+
+			template<typename Implementation, ContextCapability RequestedCapability>
+			constexpr static bool SupportsCapability()
+			{
+				return FCapabilityCompatibilityCheck<Implementation, RequestedCapability>::Value;
+			}
+
 			template<typename Implementation>
 			static bool SupportsCapability(const FName& Capability)
 			{
 				return (ImplementsCapability<Implementation, Bases>(Capability) || ...);
 			}
 
+			template<typename Implementation>
+			constexpr static bool SupportsCapability(int32 CapabilityId)
+			{
+				return (ImplementsCapability<Implementation, Bases>(CapabilityId) || ...);
+			}
+
 			template<typename Implementation>
 			constexpr static int32 CountSupportedCapabilities()
 			{
-				return 0 + ((ImplementsCapability<Implementation, Bases>() ? 1 : 0) + ...);
+				return 0 + ((SupportsCapability<Implementation, Bases<IContextCapability>>() ? 1 : 0) + ...);
 			}
 
 			template<typename Implementation>
-			static TConstArrayView<FName> SupportedCapabilitiesList()
+			static TConstArrayView<FName> GetSupportedCapabilityNameList()
 			{
 				static TConstArrayView<FName> Result = []()
 					{
 						constexpr int32 CapabilityCount = CountSupportedCapabilities<Implementation>();
+						static_assert(CapabilityCount > 0, "At least one context capability needs to be implemented to work with a query context contract.");
 						static FName Capabilities[CapabilityCount] = {};
 
 						int32 Index = 0;
 						(
 							[&]
 							{
-								if constexpr (ImplementsCapability<Implementation, Bases>())
+								if constexpr (SupportsCapability<Implementation, Bases<IContextCapability>>())
 								{
 									Capabilities[Index++] = Bases<IContextCapability>::Name;
 								}
@@ -75,6 +100,31 @@ namespace UE::Editor::DataStorage::Queries
 					}();
 				return Result;
 			}
+
+			template<typename Implementation>
+			static TConstArrayView<int32> GetSupportedCapabilityIdList()
+			{
+				static TConstArrayView<int32> Result = []()
+					{
+						constexpr int32 CapabilityCount = CountSupportedCapabilities<Implementation>();
+						static_assert(CapabilityCount > 0, "At least one context capability needs to be implemented to work with a query context contract.");
+						static int32 CapabilityIds[CapabilityCount] = {};
+
+						int32 Index = 0;
+						(
+							[&]
+							{
+								if constexpr (SupportsCapability<Implementation, Bases<IContextCapability>>())
+								{
+									CapabilityIds[Index++] = Bases<IContextCapability>::Id;
+								}
+							}(), ...);
+
+						return TConstArrayView<int32>(CapabilityIds, CapabilityCount);
+					}();
+				return Result;
+			}
+
 		};
 
 		//Used to remove the first template argument as that's a placeholder to allow the X-macros to work with the template.
@@ -112,6 +162,8 @@ namespace UE::Editor::DataStorage::Queries
 	>
 	{
 		virtual ~IContextContract() override = default;
+
+		virtual TConstArrayView<FName> GetSupportedCapabilities() const = 0;
 	};
 
 #undef Function0
diff --git a/Engine/Source/Runtime/TypedElementFramework/Public/Elements/Framework/TypedElementQueryFunctions.h b/Engine/Source/Runtime/TypedElementFramework/Public/Elements/Framework/TypedElementQueryFunctions.h
--- a/Engine/Source/Runtime/TypedElementFramework/Public/Elements/Framework/TypedElementQueryFunctions.h
+++ b/Engine/Source/Runtime/TypedElementFramework/Public/Elements/Framework/TypedElementQueryFunctions.h
@@ -6,6 +6,7 @@
 #include "Containers/ContainersFwd.h"
 #include "Elements/Framework/TypedElementQueryCapabilities.h"
 #include "Elements/Framework/TypedElementQueryFunctionArgumentTypes.h"
+#include "Elements/Framework/TypedElementRowHandleArrayView.h"
 #include "Misc/EnumClassFlags.h"
 #include "Templates/Function.h"
 #include "UObject/NameTypes.h"
@@ -20,11 +21,11 @@ namespace UE::Editor::DataStorage::Queries
 	{
 		virtual ~IQueryFunctionResponse() = default;
 		
-		virtual bool NextBatch() = 0;
-		virtual bool NextRow() = 0;
-
-		virtual void GetConstColumns(TArrayView<const void*> ColumnsData, TConstArrayView<const UScriptStruct*> ColumnTypes) = 0;
+		virtual FRowHandleArrayView GetBatchRows() const = 0;
+		virtual void GetConstColumns(TArrayView<const void*> ColumnsData, TConstArrayView<const UScriptStruct*> ColumnTypes) const = 0;
 		virtual void GetMutableColumns(TArrayView<void*> ColumnsData, TConstArrayView<const UScriptStruct*> ColumnTypes) = 0;
+		// Only called when a batch is broken up in individual rows for the current active call.
+		virtual void SetCurrentRow(RowHandle Row) = 0;
 	};
 
 	enum class EFunctionCallConfig
@@ -35,29 +36,18 @@ namespace UE::Editor::DataStorage::Queries
 	};
 	ENUM_CLASS_FLAGS(EFunctionCallConfig);
 
-	template<typename T>
-	concept FunctionType =
-		requires { &std::remove_reference_t<T>::operator(); }			// non-generic lambdas & single-operator() functors
-		|| std::is_function_v<std::remove_pointer_t<std::decay_t<T>>>	// functions, function pointers & function references
-		|| std::is_member_function_pointer_v<std::decay_t<T>>;			// member function pointers
-
 	/** Storage for a function that can be used as part of a query.  */
 	template<typename ReturnType>
 	class TQueryFunctionBase
 	{
 	public:
-		using FunctionSpecializationCallback = 
-			bool(*)(IQueryFunctionResponse& Response, TArrayView<const void*> ConstColumns, TArrayView<void*> MutableColumns);
+		using ContractType = IContextContract&;
+		using FunctionResponseType = IQueryFunctionResponse&;
+		using FunctionSpecializationCallback = bool(*)(FunctionResponseType Response, TArrayView<const void*> ConstColumns, TArrayView<void*> MutableColumns);
 	
-		using WrapperFunctionType = TFunction<
-			void(
-				TResult<ReturnType>& Result, 
-				IContextContract& Contract, 
-				IQueryFunctionResponse& Response,
-				FunctionSpecializationCallback Specialization)
-		>;
-
-		TConstArrayView<FName> Capabilities;
+		using WrapperFunctionType = TFunction<EFlowControl(TResult<ReturnType>& Result, ContractType Contract,
+				FunctionResponseType Response, FunctionSpecializationCallback Specialization)>;
+
 		TConstArrayView<const UScriptStruct*> ConstColumnTypes;
 		TConstArrayView<const UScriptStruct*> MutableColumnTypes;
 		WrapperFunctionType Function;
@@ -65,22 +55,24 @@ namespace UE::Editor::DataStorage::Queries
 
 	protected:
 		template<EFunctionCallConfig Config>
-		void CallInternal(TResult<ReturnType>& Result, IContextContract& Contract, IQueryFunctionResponse& Response);
+		EFlowControl CallInternal(TResult<ReturnType>& Result, ContractType Contract, FunctionResponseType Response) const;
 	};
 
 	template<typename ReturnType>
 	class TQueryFunction final : public TQueryFunctionBase<ReturnType>
 	{
 	public:
+		using typename TQueryFunctionBase<ReturnType>::FunctionResponseType;
+
 		template<EFunctionCallConfig Config>
-		void Call(TResult<ReturnType>& Result, IContextContract& Contract, IQueryFunctionResponse& Response)
+		EFlowControl Call(TResult<ReturnType>& Result, IContextContract& Contract, FunctionResponseType Response) const
 		{
-			this->template CallInternal<Config>(Result, Contract, Response);
+			return this->template CallInternal<Config>(Result, Contract, Response);
 		}
 
-		void Call(TResult<ReturnType>& Result, IContextContract& Contract, IQueryFunctionResponse& Response)
+		EFlowControl Call(TResult<ReturnType>& Result, IContextContract& Contract, FunctionResponseType Response) const
 		{
-			this->template CallInternal<EFunctionCallConfig::None>(Result, Contract, Response);
+			return this->template CallInternal<EFunctionCallConfig::None>(Result, Contract, Response);
 		}
 	};
 
@@ -88,25 +80,24 @@ namespace UE::Editor::DataStorage::Queries
 	class TQueryFunction<void> final : public TQueryFunctionBase<void>
 	{
 	public:
+		using typename TQueryFunctionBase<void>::FunctionResponseType;
+
 		template<EFunctionCallConfig Config>
-		void Call(IContextContract& Contract, IQueryFunctionResponse& Response)
+		EFlowControl Call(IContextContract& Contract, FunctionResponseType Response) const
 		{
 			TResult<void> Dummy;
-			this->template CallInternal<Config>(Dummy, Contract, Response);
+			return this->template CallInternal<Config>(Dummy, Contract, Response);
 		}
 
-		void Call(IContextContract& Contract, IQueryFunctionResponse& Response)
+		EFlowControl Call(IContextContract& Contract, FunctionResponseType Response) const
 		{
-			this->Call<EFunctionCallConfig::None>(Contract, Response);
+			return this->Call<EFunctionCallConfig::None>(Contract, Response);
 		}
 	};
 
 	template<typename Return, FunctionType Function>
 	TQueryFunction<Return> BuildQueryFunction(Function&& Callback);
 
-	template<typename Return, FunctionType Function>
-	TQueryFunction<void> BuildQueryFunction(TResult<Return>& Result, Function&& Callback);
-
 } // namespace UE::Editor::DataStorage::Queries
 
 #include "Elements/Framework/TypedElementQueryFunctions.inl"
diff --git a/Engine/Source/Runtime/TypedElementFramework/Public/Elements/Framework/TypedElementQueryFunctions.inl b/Engine/Source/Runtime/TypedElementFramework/Public/Elements/Framework/TypedElementQueryFunctions.inl
--- a/Engine/Source/Runtime/TypedElementFramework/Public/Elements/Framework/TypedElementQueryFunctions.inl
+++ b/Engine/Source/Runtime/TypedElementFramework/Public/Elements/Framework/TypedElementQueryFunctions.inl
@@ -130,10 +130,6 @@ namespace UE::Editor::DataStorage::Queries
 			template<typename RequestedReturnType>
 			static constexpr void SetupResult(TQueryFunction<RequestedReturnType>& Result)
 			{
-				if constexpr (!std::is_same_v<void, ContextType>)
-				{
-					Result.Capabilities = IContextContract::SupportedCapabilitiesList<typename ContextType::Capabilities>();
-				}
 				Result.ConstColumnTypes = ArgumentInfo::ListConstColumns();
 				Result.MutableColumnTypes = ArgumentInfo::ListMutableColumns();
 				Result.bIsSingleRowProcessor = bIsSingleRowProcessor;
@@ -141,22 +137,21 @@ namespace UE::Editor::DataStorage::Queries
 		};
 
 		template<typename FunctionInfo, typename ArgumentInfo, typename ReturnType, typename FunctionType>
-		void CallBody(
+		EFlowControl CallBody(
 			TResult<ReturnType>& QueryResult,
-			IContextContract& Contract,
-			IQueryFunctionResponse& Response,
-			typename TQueryFunction<ReturnType>::FunctionSpecializationCallback Specialization,
+			typename TQueryFunctionBase<ReturnType>::ContractType Contract,
+			typename TQueryFunctionBase<ReturnType>::FunctionResponseType Response,
+			typename TQueryFunctionBase<ReturnType>::FunctionSpecializationCallback Specialization,
 			FunctionType&& Callback)
 		{
 			using ArgumentList = typename ArgumentInfo::ArgumentList;
 
 			TConstArrayView<const UScriptStruct*> ConstColumnTypes = ArgumentInfo::ListConstColumns();
-			TConstArrayView<const UScriptStruct*> MutableColumnTypes = ArgumentInfo::ListMutableColumns();
-
 			const void** ConstColumnsData = static_cast<const void**>(FMemory_Alloca(sizeof(void*) * ConstColumnTypes.Num()));
-			void** MutableColumnsData = static_cast<void**>(FMemory_Alloca(sizeof(void*) * MutableColumnTypes.Num()));
-
 			TArrayView<const void*> ConstColumns(ConstColumnsData, ConstColumnTypes.Num());
+
+			TConstArrayView<const UScriptStruct*> MutableColumnTypes = ArgumentInfo::ListMutableColumns();
+			void** MutableColumnsData = static_cast<void**>(FMemory_Alloca(sizeof(void*) * MutableColumnTypes.Num()));
 			TArrayView<void*> MutableColumns(MutableColumnsData, MutableColumnTypes.Num());
 
 			EFlowControl Flow = EFlowControl::Continue;
@@ -166,20 +161,25 @@ namespace UE::Editor::DataStorage::Queries
 			ArgumentInfo::SetContext(Arguments, Contract);
 			ArgumentInfo::SetFlowControl(Arguments, Flow);
 
-			do
+			FRowHandleArrayView Rows = Response.GetBatchRows();
+			if (!Rows.IsEmpty())
 			{
 				Response.GetConstColumns(ConstColumns, ConstColumnTypes);
 				Response.GetMutableColumns(MutableColumns, MutableColumnTypes);
+				bool bSpecializationResult = Specialization(Response, ConstColumns, MutableColumns);
 
-				if (Specialization(Response, ConstColumns, MutableColumns))
+				if (bSpecializationResult)
 				{
 					ArgumentInfo::SetConstColumns(Arguments, ConstColumns);
 					ArgumentInfo::SetMutableColumns(Arguments, MutableColumns);
 
 					if constexpr (FunctionInfo::bIsSingleRowProcessor)
 					{
-						do
+						const RowHandle* It = Rows.GetData();
+						const RowHandle* End = Rows.GetData() + Rows.Num();
+						for (; It != End && Flow == EFlowControl::Continue; ++It)
 						{
+							Response.SetCurrentRow(*It);
 							if constexpr (ArgumentInfo::ResultIndex >= 0 || std::is_same_v<ReturnType, void>)
 							{
 								Arguments.ApplyBefore(Callback);
@@ -189,47 +189,50 @@ namespace UE::Editor::DataStorage::Queries
 								QueryResult.Add(Arguments.ApplyBefore(Callback));
 							}
 							ArgumentInfo::IncrementColumns(Arguments);
-						} while (Flow == EFlowControl::Continue && Response.NextRow());
+						}
 					}
 					else
 					{
 						Arguments.ApplyBefore(Callback);
 					}
+					return Flow;
 				}
-				else
-				{
-					break;
-				}
-			} while (Flow == EFlowControl::Continue && Response.NextBatch());
+			}
+			return EFlowControl::Break;
 		}
 	} // namespace Private
 
+
+
+	//
+	// TQueryFunctionBase
+	//
+
 	template<typename ReturnType>
 	template<EFunctionCallConfig Config>
-	void TQueryFunctionBase<ReturnType>::CallInternal(TResult<ReturnType>& Result, IContextContract& Contract, IQueryFunctionResponse& Response)
+	EFlowControl TQueryFunctionBase<ReturnType>::CallInternal(TResult<ReturnType>& Result, ContractType Contract, FunctionResponseType Response) const
 	{
-		auto Specialization = [](IQueryFunctionResponse& Response, TArrayView<const void*> ConstColumns, TArrayView<void*> MutableColumns)
-			{
-				if constexpr (EnumHasAnyFlags(Config, EFunctionCallConfig::VerifyColumns))
+		auto Specialization = [](FunctionResponseType Response, TArrayView<const void*> ConstColumns, TArrayView<void*> MutableColumns)
 				{
-					bool bMissingColumn = false;
-					for (const void* Column : ConstColumns)
+					if constexpr (EnumHasAnyFlags(Config, EFunctionCallConfig::VerifyColumns))
 					{
-						bMissingColumn = bMissingColumn || Column == nullptr;
+						bool bMissingColumn = false;
+						for (const void* Column : ConstColumns)
+						{
+							bMissingColumn = bMissingColumn || Column == nullptr;
+						}
+						for (void* Column : MutableColumns)
+						{
+							bMissingColumn = bMissingColumn || Column == nullptr;
+						}
+						return !bMissingColumn;
 					}
-					for (void* Column : MutableColumns)
+					else
 					{
-						bMissingColumn = bMissingColumn || Column == nullptr;
+						return true;
 					}
-					return !bMissingColumn;
-				}
-				else
-				{
-					return true;
-				}
-			};
-
-		Function(Result, Contract, Response, Specialization);
+				};
+		return Function(Result, Contract, Response, Specialization);
 	}
 
 	template<typename Return, FunctionType Function>
@@ -245,36 +248,13 @@ namespace UE::Editor::DataStorage::Queries
 		Result.Function =
 			[LocalCallback = Forward<Function>(Callback)](
 				TResult<Return>& QueryResult, 
-				IContextContract& Contract, 
-				IQueryFunctionResponse& Response, 
+				TQueryFunction<Return>::ContractType Contract,
+				TQueryFunction<Return>::FunctionResponseType Response,
 				TQueryFunction<Return>::FunctionSpecializationCallback Specialization)
 			{
-				Private::CallBody<FunctionInfo, ArgumentInfo, Return>(QueryResult, Contract, Response, Specialization, LocalCallback);
+				return Private::CallBody<FunctionInfo, ArgumentInfo, Return>(QueryResult, Contract, Response, Specialization, LocalCallback);
 			};
 		return Result;
 	}
 
-	template<typename Return, FunctionType Function>
-	TQueryFunction<void> BuildQueryFunction(TResult<Return>& Result, Function&& Callback)
-	{
-		using FunctionInfo = Private::TFunctionInfo<Function>;
-		using ArgumentInfo = typename FunctionInfo::ArgumentInfo;
-
-		TQueryFunction<Return> CompositedFunction;
-		FunctionInfo::template StaticValidation<Return>();
-		FunctionInfo::SetupResult(CompositedFunction);
-
-		Result.Function =
-			[&Result, LocalCallback = Forward<Function>(Callback)](
-				TResult<Return>& QueryResult,
-				IContextContract& Contract,
-				IQueryFunctionResponse& Response,
-				TQueryFunction<Return>::FunctionSpecializationCallback Specialization)
-			{
-				// Ignore QueryResult (which will always be a dummy for TQueryFunction<void>) in favor of the captured TResult.
-				Private::CallBody<FunctionInfo, ArgumentInfo, Return>(Result, Contract, Response, Specialization, LocalCallback);
-			};
-
-		return CompositedFunction;
-	}
 } // namespace UE::Editor::DataStorage::Queries
diff --git a/Engine/Source/Runtime/TypedElementFramework/Public/Elements/Interfaces/TypedElementDataStorageInterface.h b/Engine/Source/Runtime/TypedElementFramework/Public/Elements/Interfaces/TypedElementDataStorageInterface.h
--- a/Engine/Source/Runtime/TypedElementFramework/Public/Elements/Interfaces/TypedElementDataStorageInterface.h
+++ b/Engine/Source/Runtime/TypedElementFramework/Public/Elements/Interfaces/TypedElementDataStorageInterface.h
@@ -530,13 +530,13 @@ public:
 	 * Directly runs a query. The callback will be called for batches of matching rows. During a single call to RunQuery the callback
 	 * may be called multiple times. If the query handle is invalid or has been deleted nothing happens and the callback won't be called.
 	 */
-	template<Queries::FunctionType Function>
+	template<FunctionType Function>
 	FQueryResult RunQuery(QueryHandle Query, ERunQueryFlags Flags, Function&& Callback);
 	/**
 	 * Directly runs a query. The callback will be called for batches of matching rows. During a single call to RunQuery the callback
 	 * may be called multiple times. If the query handle is invalid or has been deleted nothing happens and the callback won't be called.
 	 */
-	template<typename ResultType, Queries::FunctionType Function>
+	template<typename ResultType, FunctionType Function>
 	FQueryResult RunQuery(QueryHandle Query, ERunQueryFlags Flags, Queries::TResult<ResultType>& Result, Function&& Callback);
 	/**
 	 * Triggers all queries registered under the activation name to run for one update cycle. The activatable queries will be activated at
@@ -1054,13 +1054,13 @@ void ICoreProvider::ForEachDynamicColumn(TFunctionRef<void(const UScriptStruct&
 	ForEachDynamicColumn(DynamicColumnTemplateType::StaticStruct(), Callback);
 }
 
-template<Queries::FunctionType Function>
+template<FunctionType Function>
 FQueryResult ICoreProvider::RunQuery(QueryHandle Query, ERunQueryFlags Flags, Function&& Callback)
 {
 	return RunQuery(Query, Flags, Queries::BuildQueryFunction<void>(Forward<Function>(Callback)));
 }
 
-template<typename ResultType, Queries::FunctionType Function>
+template<typename ResultType, FunctionType Function>
 FQueryResult ICoreProvider::RunQuery(QueryHandle Query, ERunQueryFlags Flags, Queries::TResult<ResultType>& Result, Function&& Callback)
 {
 	return RunQuery(Query, Flags, Queries::BuildQueryFunction(Result, Forward<Function>(Callback)));
diff --git a/Engine/Source/Runtime/TypedElementFramework/Tests/Elements/Framework/TypedElementQueryContextMock.h b/Engine/Source/Runtime/TypedElementFramework/Tests/Elements/Framework/TypedElementQueryContextMock.h
--- a/Engine/Source/Runtime/TypedElementFramework/Tests/Elements/Framework/TypedElementQueryContextMock.h
+++ b/Engine/Source/Runtime/TypedElementFramework/Tests/Elements/Framework/TypedElementQueryContextMock.h
@@ -155,9 +155,8 @@ namespace UE::Editor::DataStorage::Queries
 #define MockFuncVar Concat(Environment.MockFunction_, __LINE__)
 
 #define CapabilityStart(Capability, Flags) \
-		struct F##Capability##Mock : IContextCapability \
-		{ \
-			inline static const FName Name = #Capability;
+		struct F##Capability##Mock : ImplementsContextCapability< Capability > \
+		{
 
 #define Function0(Capability, Return, Function) \
 			static Return Function (FMockContextEnvironment& Environment) \
@@ -254,84 +253,11 @@ namespace UE::Editor::DataStorage::Queries
 #undef MockFunction
 #undef MockFuncVar
 
-	//
-	// Temporary forwarder that bridges FMockContextEnvironment with the capability interfaces.
-	//
-
-#define ArgTypeName(Type, Name) Type Name
-#define ArgName(Type, Name) Name
-#define CapabilityStart(Capability, Flags)
-#define CapabilityEnd(Capability)
-
-#define Function0(Capability, Return, Function) \
-		virtual Return Function() override \
-		{ return F##Capability##Mock::Function(*this); }
-
-#define Function1(Capability, Return, Function, Arg1) \
-		virtual Return Function(ArgTypeName Arg1) override \
-		{ return F##Capability##Mock::Function(*this, ArgName Arg1); }
-
-#define Function2(Capability, Return, Function, Arg1, Arg2) \
-		virtual Return Function(ArgTypeName Arg1, ArgTypeName Arg2) override \
-		{ return F##Capability##Mock::Function(*this, ArgName Arg1, ArgName Arg2); }
-
-#define Function3(Capability, Return, Function, Arg1, Arg2, Arg3) \
-		virtual Return Function(ArgTypeName Arg1, ArgTypeName Arg2, ArgTypeName Arg3) override \
-		{ return F##Capability##Mock::Function(*this, ArgName Arg1, ArgName Arg2, ArgName Arg3); }
-
-#define Function4(Capability, Return, Function, Arg1, Arg2, Arg3, Arg4) \
-		virtual Return Function(ArgTypeName Arg1, ArgTypeName Arg2, ArgTypeName Arg3, ArgTypeName Arg4) override \
-		{ return F##Capability##Mock::Function(*this, ArgName Arg1, ArgName Arg2, ArgName Arg3, ArgName Arg4); }
-
-#define ConstFunction0(Capability, Return, Function) \
-		virtual Return Function() const override \
-		{ return F##Capability##Mock::Function(*this); }
-
-#define ConstFunction1(Capability, Return, Function, Arg1) \
-		virtual Return Function(ArgTypeName Arg1) const override \
-		{ return F##Capability##Mock::Function(*this, ArgName Arg1); }
-
-#define ConstFunction2(Capability, Return, Function, Arg1, Arg2) \
-		virtual Return Function(ArgTypeName Arg1, ArgTypeName Arg2) const override \
-		{ return F##Capability##Mock::Function(*this, ArgName Arg1, ArgName Arg2); }
-
-#define ConstFunction3(Capability, Return, Function, Arg1, Arg2, Arg3) \
-		virtual Return Function(ArgTypeName Arg1, ArgTypeName Arg2, ArgTypeName Arg3) const override \
-		{ return F##Capability##Mock::Function(*this, ArgName Arg1, ArgName Arg2, ArgName Arg3); }
-
-#define ConstFunction4(Capability, Return, Function, Arg1, Arg2, Arg3, Arg4) \
-		virtual Return Function(ArgTypeName Arg1, ArgTypeName Arg2, ArgTypeName Arg3, ArgTypeName Arg4) const override \
-		{ return F##Capability##Mock::Function(*this, ArgName Arg1, ArgName Arg2, ArgName Arg3, ArgName Arg4); }
-
-	struct FMockContextForwarder
-		: FMockContextEnvironment
-		, ISingleRowInfo<IContextCapability>
-		, IRowBatchInfo<IContextCapability>
-		, IDynamicColumnInfo<IContextCapability>
-	{
-#include "Elements/Framework/TypedElementQueryCapabilities.inl"
-	};
-
-#undef CapabilityStart
-#undef CapabilityEnd
-#undef Function0
-#undef Function1
-#undef Function2
-#undef Function3
-#undef Function4
-#undef ConstFunction0
-#undef ConstFunction1
-#undef ConstFunction2
-#undef ConstFunction3
-#undef ConstFunction4
-#undef ArgTypeName
-#undef ArgName
-
 	//
 	// Combine all mock implementations together into a final implementation.
 	//
 
-using QueryContextMock = TQueryContextImpl<FMockContextForwarder
+using QueryContextMock = TQueryContextImpl<FMockContextEnvironment
 
 #define Function0(Capability, Return, Function)
 #define Function1(Capability, Return, Function, Arg1)
@@ -345,7 +271,7 @@ using QueryContextMock = TQueryContextImpl<FMockContextForwarder
 #define ConstFunction3(Capability, Return, Function, Arg1, Arg2, Arg3)
 #define ConstFunction4(Capability, Return, Function, Arg1, Arg2, Arg3, Arg4)
 
-#define CapabilityStart(Capability, Flags) , Capability
+#define CapabilityStart(Capability, Flags) , F##Capability##Mock
 #define CapabilityEnd(Capability)
 
 #include "Elements/Framework/TypedElementQueryCapabilities.inl"
diff --git a/Engine/Source/Runtime/TypedElementFramework/Tests/Elements/Framework/TypedElementQueryFunctionsTests.cpp b/Engine/Source/Runtime/TypedElementFramework/Tests/Elements/Framework/TypedElementQueryFunctionsTests.cpp
--- a/Engine/Source/Runtime/TypedElementFramework/Tests/Elements/Framework/TypedElementQueryFunctionsTests.cpp
+++ b/Engine/Source/Runtime/TypedElementFramework/Tests/Elements/Framework/TypedElementQueryFunctionsTests.cpp
@@ -6,7 +6,6 @@
 #include "Elements/Framework/TypedElementQueryContext.h"
 #include "Elements/Framework/TypedElementQueryContextMock.h"
 #include "Elements/Framework/TypedElementQueryFunctions.h"
-#include "Elements/Framework/TypedElementRowHandleArrayView.h"
 #include "Elements/Framework/TypedElementTestColumns.h"
 #include "Tests/TestHarnessAdapter.h"
 
@@ -46,18 +45,13 @@ namespace UE::Editor::DataStorage::Queries::Tests
 	{
 		virtual ~FTestQueryFunctionResponse() override = default;
 
-		virtual bool NextBatch() override
+		virtual FRowHandleArrayView GetBatchRows() const override
 		{
-			return false;
-		}
-
-		virtual bool NextRow() override
-		{
-			++CurrentRow;
-			return CurrentRow < static_cast<RowHandle>(FakeRows.Num());
+			CHECK_MESSAGE(TEXT("Row count was not set in test."), !FakeRows.IsEmpty());
+			return FRowHandleArrayView(FakeRows.GetData(), FakeRows.Num(), FRowHandleArrayView::EFlags::None);
 		}
 
-		virtual void GetConstColumns(TArrayView<const void*> ColumnsData, TConstArrayView<const UScriptStruct*> ColumnTypes) override
+		virtual void GetConstColumns(TArrayView<const void*> ColumnsData, TConstArrayView<const UScriptStruct*> ColumnTypes) const override
 		{
 			CHECK_MESSAGE(TEXT("Row count was not set in test."), !FakeRows.IsEmpty());
 			for (int32 Index = 0, Num = ColumnTypes.Num(); Index < Num; ++Index)
@@ -97,7 +91,7 @@ namespace UE::Editor::DataStorage::Queries::Tests
 			}
 		}
 
-		void SetCurrentRow(RowHandle Row)
+		virtual void SetCurrentRow(RowHandle Row) override
 		{
 			CurrentRow = Row;
 		}
@@ -111,7 +105,6 @@ namespace UE::Editor::DataStorage::Queries::Tests
 			{
 				FakeRows.Add(Row);
 			}
-			CurrentRow = FakeRows[0];
 		}
 
 		template<TColumnType... ColumnType>
@@ -140,8 +133,7 @@ namespace UE::Editor::DataStorage::Queries::Tests
 		SECTION("Empty function (void)")
 		{
 			TQueryFunction<void> Result = BuildQueryFunction<void>([](){});
-			
-			CHECK_MESSAGE(TEXT("Empty query still has capabilities."), Result.Capabilities.IsEmpty());
+
 			CHECK_MESSAGE(TEXT("Empty query still has const columns."), Result.ConstColumnTypes.IsEmpty());
 			CHECK_MESSAGE(TEXT("Empty query still has mutable columns."), Result.MutableColumnTypes.IsEmpty());
 		}
@@ -150,7 +142,6 @@ namespace UE::Editor::DataStorage::Queries::Tests
 		{
 			TQueryFunction<bool> Result = BuildQueryFunction<bool>([]() { return false; });
 
-			CHECK_MESSAGE(TEXT("Empty query still has capabilities."), Result.Capabilities.IsEmpty());
 			CHECK_MESSAGE(TEXT("Empty query still has const columns."), Result.ConstColumnTypes.IsEmpty());
 			CHECK_MESSAGE(TEXT("Empty query still has mutable columns."), Result.MutableColumnTypes.IsEmpty());
 		}
@@ -160,7 +151,6 @@ namespace UE::Editor::DataStorage::Queries::Tests
 			bool bResult = false;
 			TQueryFunction<bool> Result = BuildQueryFunction<bool>([&bResult]() { return bResult; });
 
-			CHECK_MESSAGE(TEXT("Empty query still has capabilities."), Result.Capabilities.IsEmpty());
 			CHECK_MESSAGE(TEXT("Empty query still has const columns."), Result.ConstColumnTypes.IsEmpty());
 			CHECK_MESSAGE(TEXT("Empty query still has mutable columns."), Result.MutableColumnTypes.IsEmpty());
 		}
@@ -169,7 +159,6 @@ namespace UE::Editor::DataStorage::Queries::Tests
 		{
 			TQueryFunction<void> Result = BuildQueryFunction<void>([](TQueryContext<SingleRowInfo> Context) {});
 
-			CHECK_MESSAGE(TEXT("Expected exactly one capability."), Result.Capabilities.Num() == 1);
 			CHECK_MESSAGE(TEXT("Query with only context shouldn't have const columns."), Result.ConstColumnTypes.IsEmpty());
 			CHECK_MESSAGE(TEXT("Query with only context shouldn't have mutable columns."), Result.MutableColumnTypes.IsEmpty());
 		}
@@ -181,7 +170,6 @@ namespace UE::Editor::DataStorage::Queries::Tests
 				{
 				});
 
-			CHECK_MESSAGE(TEXT("Query with only columns shouldn't have capabilities."), Result.Capabilities.IsEmpty());
 			CHECK_MESSAGE(TEXT("Expected exactly one const column."), Result.ConstColumnTypes.Num() == 1);
 			CHECK_MESSAGE(TEXT("Expected exactly one mutable column."), Result.MutableColumnTypes.Num() == 1);
 		}
@@ -193,7 +181,6 @@ namespace UE::Editor::DataStorage::Queries::Tests
 				{
 				});
 
-			CHECK_MESSAGE(TEXT("Expected exactly one capability."), Result.Capabilities.Num() == 1);
 			CHECK_MESSAGE(TEXT("Expected exactly one const column."), Result.ConstColumnTypes.Num() == 1);
 			CHECK_MESSAGE(TEXT("Expected exactly one mutable column."), Result.MutableColumnTypes.Num() == 1);
 
@@ -208,7 +195,6 @@ namespace UE::Editor::DataStorage::Queries::Tests
 				{
 				});
 
-			CHECK_MESSAGE(TEXT("Expected exactly one capability."), Result.Capabilities.Num() == 1);
 			CHECK_MESSAGE(TEXT("Expected exactly one const column."), Result.ConstColumnTypes.Num() == 1);
 			CHECK_MESSAGE(TEXT("Expected exactly one mutable column."), Result.MutableColumnTypes.Num() == 1);
 
@@ -223,7 +209,6 @@ namespace UE::Editor::DataStorage::Queries::Tests
 				{
 				});
 
-			CHECK_MESSAGE(TEXT("Expected exactly one capability."), Result.Capabilities.Num() == 1);
 			CHECK_MESSAGE(TEXT("Expected exactly one const column."), Result.ConstColumnTypes.Num() == 1);
 			CHECK_MESSAGE(TEXT("Expected exactly one mutable column."), Result.MutableColumnTypes.Num() == 1);
 
@@ -231,24 +216,6 @@ namespace UE::Editor::DataStorage::Queries::Tests
 			CHECK_MESSAGE(TEXT("Incorrect mutable column set."), Result.MutableColumnTypes[0] == FTestColumnInt::StaticStruct());
 		}
 
-		SECTION("Check context compatibility without arguments")
-		{
-			TQueryFunction<void> Function = BuildQueryFunction<void>([](){});
-
-			QueryContextMock Context;
-			bool bResult = Context.CheckCompatiblity(Function);
-			CHECK_MESSAGE(TEXT("Function without context incorrectly found to not be matching."), bResult);
-		}
-
-		SECTION("Check context compatibility with context")
-		{
-			TQueryFunction<void> Function = BuildQueryFunction<void>([](TQueryContext<SingleRowInfo> Context) {});
-
-			QueryContextMock Context;
-			bool bResult = Context.CheckCompatiblity(Function);
-			CHECK_MESSAGE(TEXT("Function with context incorrectly found to not be matching."), bResult);
-		}
-
 		SECTION("Call without arguments")
 		{
 			TQueryFunction<bool> Function = BuildQueryFunction<bool>(
@@ -257,10 +224,11 @@ namespace UE::Editor::DataStorage::Queries::Tests
 					return true;
 				});
 			
-			QueryContextMock Context;
+			FMockContextEnvironment Environment;
+			QueryContextMock Context(Environment);
 			FTestQueryFunctionResponse Response;
 			Response.SetRowCount(1);
-
+			
 			FBoolOrResultAccumulator Result;
 			Function.Call<EFunctionCallConfig::VerifyColumns>(Result, Context, Response);
 			CHECK_MESSAGE(TEXT("Query function wasn't called."), Result.Value);
@@ -274,7 +242,8 @@ namespace UE::Editor::DataStorage::Queries::Tests
 					Column.TestString = TEXT("Callback");
 				});
 
-			QueryContextMock Context;
+			FMockContextEnvironment Environment;
+			QueryContextMock Context(Environment);
 			FTestColumnString Column;
 			Column.TestString = TEXT("Not set");
 			
@@ -294,7 +263,8 @@ namespace UE::Editor::DataStorage::Queries::Tests
 					ColumnA.TestString = FString::Printf(TEXT("Test: %i"), ColumnB.TestInt);
 				});
 
-			QueryContextMock Context;
+			FMockContextEnvironment Environment;
+			QueryContextMock Context(Environment);
 			FTestColumnString ColumnA;
 			FTestColumnInt ColumnB;
 			ColumnB.TestInt = 42;
@@ -324,9 +294,10 @@ namespace UE::Editor::DataStorage::Queries::Tests
 			FRowHandleArrayView RowsView(Rows.GetData(), Rows.Num(), 
 				FRowHandleArrayView::EFlags::IsSorted | FRowHandleArrayView::EFlags::IsUnique);
 
-			QueryContextMock Context;
-			Context.GetContextImplementation().Assign_GetBatchRowCount_Const([]() { return RowCount; });
-			Context.GetContextImplementation().Assign_GetBatchRowHandles_Const([RowsView](){ return RowsView; });
+			FMockContextEnvironment Environment;
+			QueryContextMock Context(Environment);
+			Environment.Assign_GetBatchRowCount_Const([]() { return RowCount; });
+			Environment.Assign_GetBatchRowHandles_Const([RowsView](){ return RowsView; });
 			FTestColumnString ColumnsA[RowCount];
 			FTestColumnInt ColumnsB[] =
 			{
@@ -353,18 +324,19 @@ namespace UE::Editor::DataStorage::Queries::Tests
 					Context.CurrentRowHasColumns<FTestColumnString>();
 				});
 
-			QueryContextMock Context;
+			FMockContextEnvironment Environment;
+			QueryContextMock Context(Environment);
 
 			bool bGetCurrentRowCalled = false;
-			Context.GetContextImplementation().Assign_GetCurrentRow_Const(
+			Environment.Assign_GetCurrentRow_Const(
 				[&bGetCurrentRowCalled]()
 				{
 					bGetCurrentRowCalled = true;
 					return 0;
 				});
-
+			
 			bool bCurrentRowHasColumnsCalled = false;
-			Context.GetContextImplementation().Assign_CurrentRowHasColumns_Const(
+			Environment.Assign_CurrentRowHasColumns_Const(
 				[&bCurrentRowHasColumnsCalled](TConstArrayView<const UScriptStruct*>)
 				{
 					bCurrentRowHasColumnsCalled = true;
@@ -391,10 +363,11 @@ namespace UE::Editor::DataStorage::Queries::Tests
 					RowCheck.Add(Context.GetCurrentRow());
 				});
 
-			QueryContextMock Context;
+			FMockContextEnvironment Environment;
+			QueryContextMock Context(Environment);
 			FTestQueryFunctionResponse Response;
 
-			Context.GetContextImplementation().Assign_GetCurrentRow_Const(
+			Environment.Assign_GetCurrentRow_Const(
 				[&Response]()
 				{
 					return Response.CurrentRow;
@@ -414,7 +387,8 @@ namespace UE::Editor::DataStorage::Queries::Tests
 		{
 			TQueryFunction<bool> Function = BuildQueryFunction<bool>([]() { return true; });
 
-			QueryContextMock Context;
+			FMockContextEnvironment Environment;
+			QueryContextMock Context(Environment);
 			FTestColumnString Column;
 
 			FBoolOrResultAccumulator Result;
@@ -435,8 +409,9 @@ namespace UE::Editor::DataStorage::Queries::Tests
 					Result.Add(true);
 				});
 
-			QueryContextMock Context;
-			Context.GetContextImplementation().Assign_GetCurrentRow_Const([]() { return RowHandle(1); });
+			FMockContextEnvironment Environment;
+			Environment.Assign_GetCurrentRow_Const([]() { return RowHandle(1); });
+			QueryContextMock Context(Environment);
 			FTestColumnString Column;
 
 			FBoolOrResultAccumulator Result;
@@ -465,9 +440,10 @@ namespace UE::Editor::DataStorage::Queries::Tests
 			FRowHandleArrayView RowsView(Rows.GetData(), Rows.Num(),
 				FRowHandleArrayView::EFlags::IsSorted | FRowHandleArrayView::EFlags::IsUnique);
 
-			QueryContextMock Context;
-			Context.GetContextImplementation().Assign_GetBatchRowCount_Const([]() { return RowCount; });
-			Context.GetContextImplementation().Assign_GetBatchRowHandles_Const([RowsView]() { return RowsView; });
+			FMockContextEnvironment Environment;
+			QueryContextMock Context(Environment);
+			Environment.Assign_GetBatchRowCount_Const([]() { return RowCount; });
+			Environment.Assign_GetBatchRowHandles_Const([RowsView]() { return RowsView; });
 
 			FTestColumnInt Columns[] =
 			{
@@ -497,7 +473,8 @@ namespace UE::Editor::DataStorage::Queries::Tests
 					Flow = EFlowControl::Break;
 				});
 
-			QueryContextMock Context;
+			FMockContextEnvironment Environment;
+			QueryContextMock Context(Environment);
 
 			FTestColumnInt Columns[] =
 			{
@@ -513,6 +490,7 @@ namespace UE::Editor::DataStorage::Queries::Tests
 
 			CHECK_MESSAGE(TEXT("Iteration didn't stop on the first row."), LastValue == 1);
 		}
+
 	}
 } // namespace UE::Editor::DataStorage::Queries::Tests
 
