[TEDS] Add FindForColumnOrGeneric helpers to FMetaDataView

Adds `FindForColumnOrGeneric` and `FindForColumnsOrGeneric` to `FMetaDataView`, which try column-specific lookup first and fall back to `FindGeneric`.
Includes template variants for statically-known column types.

#Synced-CL 52095318

diff --git a/Engine/Source/Runtime/TypedElementFramework/Private/Elements/Framework/TypedElementMetaData.cpp b/Engine/Source/Runtime/TypedElementFramework/Private/Elements/Framework/TypedElementMetaData.cpp
--- a/Engine/Source/Runtime/TypedElementFramework/Private/Elements/Framework/TypedElementMetaData.cpp
+++ b/Engine/Source/Runtime/TypedElementFramework/Private/Elements/Framework/TypedElementMetaData.cpp
@@ -143,6 +143,24 @@ namespace UE::Editor::DataStorage
 		return FMetaDataEntryView();
 	}
 
+	FMetaDataEntryView FMetaDataView::FindForColumnOrGeneric(TWeakObjectPtr<const UScriptStruct> Column, FName AttributeName) const
+	{
+		FMetaDataEntryView ColumnMetaData = FindForColumn(Column, AttributeName);
+		return ColumnMetaData.IsSet() ? ColumnMetaData : FindGeneric(AttributeName);
+	}
+
+	FMetaDataEntryView FMetaDataView::FindForColumnsOrGeneric(const TConstArrayView<TWeakObjectPtr<const UScriptStruct>> Columns, FName AttributeName) const
+	{
+		for (const TWeakObjectPtr<const UScriptStruct>& Column : Columns)
+		{
+			FMetaDataEntryView ColumnMetaData = FindForColumn(Column, AttributeName);
+			if (ColumnMetaData.IsSet())
+			{
+				return ColumnMetaData;
+			}
+		}
+		return FindGeneric(AttributeName);
+	}
 
 
 	//
diff --git a/Engine/Source/Runtime/TypedElementFramework/Public/Elements/Framework/TypedElementMetaData.h b/Engine/Source/Runtime/TypedElementFramework/Public/Elements/Framework/TypedElementMetaData.h
--- a/Engine/Source/Runtime/TypedElementFramework/Public/Elements/Framework/TypedElementMetaData.h
+++ b/Engine/Source/Runtime/TypedElementFramework/Public/Elements/Framework/TypedElementMetaData.h
@@ -120,6 +120,17 @@ namespace UE::Editor::DataStorage
 			TWeakObjectPtr<const UScriptStruct> Column, FName AttributeName) const;
 		template<typename Column>
 		FMetaDataEntryView FindForColumn(FName AttributeName) const;
+
+		/** First tries FindForColumn and if unsuccessful, calls FindGeneric. */
+		TYPEDELEMENTFRAMEWORK_API FMetaDataEntryView FindForColumnOrGeneric(
+			TWeakObjectPtr<const UScriptStruct> Column, FName AttributeName) const;
+		template<typename Column>
+		FMetaDataEntryView FindForColumnOrGeneric(FName AttributeName) const;
+		/** Tries FindForColumn for each given column in order until successful, otherwise calls FindGeneric. */
+		TYPEDELEMENTFRAMEWORK_API FMetaDataEntryView FindForColumnsOrGeneric(
+			const TConstArrayView<TWeakObjectPtr<const UScriptStruct>> Columns, FName AttributeName) const;
+		template<typename... Columns>
+		FMetaDataEntryView FindForColumnsOrGeneric(FName AttributeName) const;
 	};
 
 	/**
diff --git a/Engine/Source/Runtime/TypedElementFramework/Public/Elements/Framework/TypedElementMetaData.inl b/Engine/Source/Runtime/TypedElementFramework/Public/Elements/Framework/TypedElementMetaData.inl
--- a/Engine/Source/Runtime/TypedElementFramework/Public/Elements/Framework/TypedElementMetaData.inl
+++ b/Engine/Source/Runtime/TypedElementFramework/Public/Elements/Framework/TypedElementMetaData.inl
@@ -104,6 +104,19 @@ namespace UE::Editor::DataStorage
 		return FindForColumn(Column::StaticStruct(), AttributeName);
 	}
 
+	template<typename Column>
+	FMetaDataEntryView FMetaDataView::FindForColumnOrGeneric(FName AttributeName) const
+	{
+		return FindForColumnOrGeneric(Column::StaticStruct(), AttributeName);
+	}
+
+	template<typename... Columns>
+	FMetaDataEntryView FMetaDataView::FindForColumnsOrGeneric(FName AttributeName) const
+	{
+		const TWeakObjectPtr<const UScriptStruct> ColumnStructs[] = { Columns::StaticStruct()... };
+		return FindForColumnsOrGeneric(ColumnStructs, AttributeName);
+	}
+
 
 
 	//
