[TEDS Search] Create a factory for overriding Search functions so that items that cannot be defined as UPROPERTIES can still be searched

#Synced-CL 52470375

diff --git a/Engine/Plugins/Experimental/EditorDataStorage/Source/TedsCore/Private/PropertySearchers/DefaultPropertySearcherFactory.cpp b/Engine/Plugins/Experimental/EditorDataStorage/Source/TedsCore/Private/PropertySearchers/DefaultPropertySearcherFactory.cpp
new file mode 100644
--- /dev/null
+++ b/Engine/Plugins/Experimental/EditorDataStorage/Source/TedsCore/Private/PropertySearchers/DefaultPropertySearcherFactory.cpp
@@ -0,0 +1,30 @@
+// Copyright Epic Games, Inc. All Rights Reserved.
+
+#include "PropertySearchers/DefaultPropertySearcherFactory.h"
+
+#include "Elements/Interfaces/TypedElementDataStorageUiInterface.h"
+#include "PropertySearchers/StringSearchers.h"
+#include "UObject/StrProperty.h"
+#include "UObject/TextProperty.h"
+#include "UObject/UnrealType.h"
+
+#include UE_INLINE_GENERATED_CPP_BY_NAME(DefaultPropertySearcherFactory)
+
+void UDefaultPropertySearcherFactory::RegisterPropertySearchers(UE::Editor::DataStorage::IUiProvider& DataStorageUi) const
+{
+	using namespace UE::Editor::DataStorage::Searchers;
+
+	auto RegisterCallback = [&DataStorageUi]<typename Searcher, typename PropertyType>()
+	{
+		DataStorageUi.RegisterSearcherGeneratorForProperty(PropertyType::StaticClass(),
+			[](TWeakObjectPtr<const UScriptStruct> ColumnType, const FProperty& Property)
+			{
+				return MakeShared<Searcher>(ColumnType, CastField<PropertyType>(&Property));
+			});
+	};
+	
+	// Strings
+	RegisterCallback.operator()<FStringSearcher, FStrProperty>();
+	RegisterCallback.operator()<FTextSearcher, FTextProperty>();
+	RegisterCallback.operator()<FNameSearcher, FNameProperty>();
+}
diff --git a/Engine/Plugins/Experimental/EditorDataStorage/Source/TedsCore/Private/PropertySearchers/DefaultPropertySearcherFactory.h b/Engine/Plugins/Experimental/EditorDataStorage/Source/TedsCore/Private/PropertySearchers/DefaultPropertySearcherFactory.h
new file mode 100644
--- /dev/null
+++ b/Engine/Plugins/Experimental/EditorDataStorage/Source/TedsCore/Private/PropertySearchers/DefaultPropertySearcherFactory.h
@@ -0,0 +1,19 @@
+// Copyright Epic Games, Inc. All Rights Reserved.
+
+#pragma once
+
+#include "Elements/Interfaces/TypedElementDataStorageFactory.h"
+
+#include "DefaultPropertySearcherFactory.generated.h"
+
+/**
+ * Factory used to register searcher for text properties.
+ */
+UCLASS(Transient)
+class UDefaultPropertySearcherFactory : public UEditorDataStorageFactory
+{
+	GENERATED_BODY()
+
+public:
+	virtual void RegisterPropertySearchers(UE::Editor::DataStorage::IUiProvider& DataStorageUi) const override;
+};
diff --git a/Engine/Plugins/Experimental/EditorDataStorage/Source/TedsCore/Private/PropertySearchers/PropertySearcherBase.h b/Engine/Plugins/Experimental/EditorDataStorage/Source/TedsCore/Private/PropertySearchers/PropertySearcherBase.h
new file mode 100644
--- /dev/null
+++ b/Engine/Plugins/Experimental/EditorDataStorage/Source/TedsCore/Private/PropertySearchers/PropertySearcherBase.h
@@ -0,0 +1,36 @@
+// Copyright Epic Games, Inc. All Rights Reserved.
+
+#pragma once
+
+#include "Elements/Interfaces/TypedElementDataStorageUiInterface.h"
+#include "UObject/WeakObjectPtrTemplates.h"
+
+class UScriptStruct;
+
+namespace UE::Editor::DataStorage
+{
+	class ICoreProvider;
+}
+
+namespace UE::Editor::DataStorage::Searchers
+{
+	template<typename SearchableType, typename PropertyType>
+	class TPropertySearcherBase : public FColumnSearcherInterface
+	{
+	public:
+		TPropertySearcherBase(TWeakObjectPtr<const UScriptStruct> ColumnType, const PropertyType* Property);
+
+		virtual void GetSearchableString(FString& SearchableString, const ICoreProvider& Storage, RowHandle Row) const override;
+		
+	protected:
+		virtual void GetSearchableString(FString& SearchableString, const ICoreProvider& Storage, const SearchableType& Value) const = 0;
+
+	private:
+		const SearchableType* GetValue(const ICoreProvider& Storage, const UScriptStruct* Column, RowHandle Row) const;
+
+		TWeakObjectPtr<const UScriptStruct> ColumnType;
+		const PropertyType* Property;
+	};
+} // namespace UE::Editor::DataStorage::Searchers
+
+#include "PropertySearchers/PropertySearcherBase.inl"
\ No newline at end of file
diff --git a/Engine/Plugins/Experimental/EditorDataStorage/Source/TedsCore/Private/PropertySearchers/PropertySearcherBase.inl b/Engine/Plugins/Experimental/EditorDataStorage/Source/TedsCore/Private/PropertySearchers/PropertySearcherBase.inl
new file mode 100644
--- /dev/null
+++ b/Engine/Plugins/Experimental/EditorDataStorage/Source/TedsCore/Private/PropertySearchers/PropertySearcherBase.inl
@@ -0,0 +1,38 @@
+// Copyright Epic Games, Inc. All Rights Reserved.
+
+#include "Elements/Interfaces/TypedElementDataStorageInterface.h"
+
+namespace UE::Editor::DataStorage::Searchers
+{
+	template<typename SearchableType, typename PropertyType>
+	TPropertySearcherBase<SearchableType, PropertyType>::TPropertySearcherBase(
+		TWeakObjectPtr<const UScriptStruct> ColumnType, const PropertyType* Property)
+		: ColumnType(MoveTemp(ColumnType))
+		, Property(Property)
+	{
+	}
+
+	template<typename SearchableType, typename PropertyType>
+	void TPropertySearcherBase<SearchableType, PropertyType>::GetSearchableString(
+		FString& SearchableString, const ICoreProvider& Storage, RowHandle Row) const
+	{
+		if (const UScriptStruct* Column = ColumnType.Get())
+		{
+			if (const SearchableType* Value = GetValue(Storage, Column, Row))
+			{
+				GetSearchableString(SearchableString, Storage, *Value);
+			}
+		}
+	}
+
+	template<typename SearchableType, typename PropertyType>
+	const SearchableType* TPropertySearcherBase<SearchableType, PropertyType>::GetValue(
+		const ICoreProvider& Storage, const UScriptStruct* Column, RowHandle Row) const
+	{
+		if (const void* Data = Storage.GetColumnData(Row, Column))
+		{
+			return Property->template ContainerPtrToValuePtr<SearchableType>(Data);
+		}
+		return nullptr;
+	}
+} // namespace UE::Editor::DataStorage::Searchers
\ No newline at end of file
diff --git a/Engine/Plugins/Experimental/EditorDataStorage/Source/TedsCore/Private/PropertySearchers/StringSearchers.cpp b/Engine/Plugins/Experimental/EditorDataStorage/Source/TedsCore/Private/PropertySearchers/StringSearchers.cpp
new file mode 100644
--- /dev/null
+++ b/Engine/Plugins/Experimental/EditorDataStorage/Source/TedsCore/Private/PropertySearchers/StringSearchers.cpp
@@ -0,0 +1,46 @@
+// Copyright Epic Games, Inc. All Rights Reserved.
+
+#include "PropertySearchers/StringSearchers.h"
+
+namespace UE::Editor::DataStorage::Searchers
+{
+	//
+	// FStringSearcher
+	//
+
+	FStringSearcher::FStringSearcher(TWeakObjectPtr<const UScriptStruct> ColumnType, const FStrProperty* Property)
+		: BaseType(MoveTemp(ColumnType), Property)
+	{}
+
+	void FStringSearcher::GetSearchableString(FString& SearchableString, const ICoreProvider& Storage, const FString& Value) const
+	{
+		SearchableString = Value;
+	}
+
+	//
+	// FTextSearcher
+	//
+
+	FTextSearcher::FTextSearcher(TWeakObjectPtr<const UScriptStruct> ColumnType, const FTextProperty* Property)
+		: BaseType(MoveTemp(ColumnType), Property)
+	{}
+
+	void FTextSearcher::GetSearchableString(FString& SearchableString, const ICoreProvider& Storage, const FText& Value) const
+	{
+		SearchableString = Value.ToString();
+	}
+
+	//
+	// FNameSearcher
+	//
+
+	FNameSearcher::FNameSearcher(TWeakObjectPtr<const UScriptStruct> ColumnType, const FNameProperty* Property)
+		: BaseType(MoveTemp(ColumnType), Property)
+	{}
+
+	void FNameSearcher::GetSearchableString(FString& SearchableString, const ICoreProvider& Storage, const FName& Value) const
+	{
+		SearchableString = Value.ToString();
+	}
+
+} // namespace UE::Editor::DataStorage::Searchers
diff --git a/Engine/Plugins/Experimental/EditorDataStorage/Source/TedsCore/Private/PropertySearchers/StringSearchers.h b/Engine/Plugins/Experimental/EditorDataStorage/Source/TedsCore/Private/PropertySearchers/StringSearchers.h
new file mode 100644
--- /dev/null
+++ b/Engine/Plugins/Experimental/EditorDataStorage/Source/TedsCore/Private/PropertySearchers/StringSearchers.h
@@ -0,0 +1,53 @@
+// Copyright Epic Games, Inc. All Rights Reserved.
+
+#pragma once
+
+#include "Elements/Interfaces/TypedElementDataStorageUiInterface.h"
+#include "PropertySearchers/PropertySearcherBase.h"
+#include "UObject/StrProperty.h"
+#include "UObject/TextProperty.h"
+#include "UObject/UnrealType.h"
+#include "UObject/WeakObjectPtrTemplates.h"
+
+class UScriptStruct;
+
+namespace UE::Editor::DataStorage
+{
+	class ICoreProvider;
+}
+
+namespace UE::Editor::DataStorage::Searchers
+{
+	class FStringSearcher final : public TPropertySearcherBase<FString, FStrProperty>
+	{
+		using BaseType = TPropertySearcherBase<FString, FStrProperty>;
+	public:
+		FStringSearcher(TWeakObjectPtr<const UScriptStruct> ColumnType, const FStrProperty* Property);
+
+	protected:
+		virtual void GetSearchableString(FString& SearchableString, const ICoreProvider& Storage, const FString& Value) const override;
+	};
+
+
+	class FTextSearcher final : public TPropertySearcherBase<FText, FTextProperty>
+	{
+		using BaseType = TPropertySearcherBase<FText, FTextProperty>;
+	public:
+		FTextSearcher(TWeakObjectPtr<const UScriptStruct> ColumnType, const FTextProperty* Property);
+
+	protected:
+		virtual void GetSearchableString(FString& SearchableString, const ICoreProvider& Storage, const FText& Value) const override;
+	};
+
+
+	class FNameSearcher final : public TPropertySearcherBase<FName, FNameProperty>
+	{
+		using BaseType = TPropertySearcherBase<FName, FNameProperty>;
+	public:
+		FNameSearcher(TWeakObjectPtr<const UScriptStruct> ColumnType, const FNameProperty* Property);
+
+	protected:
+		virtual void GetSearchableString(FString& SearchableString, const ICoreProvider& Storage, const FName& Value) const override;
+	};
+	
+} // namespace UE::Editor::DataStorage::Searchers
\ No newline at end of file
diff --git a/Engine/Plugins/Experimental/EditorDataStorage/Source/TedsCore/Private/TypedElementDatabaseUI.cpp b/Engine/Plugins/Experimental/EditorDataStorage/Source/TedsCore/Private/TypedElementDatabaseUI.cpp
--- a/Engine/Plugins/Experimental/EditorDataStorage/Source/TedsCore/Private/TypedElementDatabaseUI.cpp
+++ b/Engine/Plugins/Experimental/EditorDataStorage/Source/TedsCore/Private/TypedElementDatabaseUI.cpp
@@ -669,6 +669,55 @@ UE::Editor::DataStorage::TableHandle UEditorDataStorageUi::GetWidgetTable() cons
 	return WidgetTable;
 }
 
+void UEditorDataStorageUi::GeneratePropertySearchers(TArray<TSharedPtr<const UE::Editor::DataStorage::FColumnSearcherInterface>>& Results,
+		TArrayView<TWeakObjectPtr<const UScriptStruct>> Columns) const
+{
+	using namespace UE::Editor::DataStorage;
+
+	TArray<const PropertySearcherConstructorCallback*, TInlineAllocator<8>> UsedConstructors;
+
+	for (TWeakObjectPtr<const UScriptStruct>& ColumnTypePointer : Columns)
+	{
+		if (const UScriptStruct* ColumnType = ColumnTypePointer.Get())
+		{
+			for (TFieldIterator<FProperty> PropIt(ColumnType); PropIt; ++PropIt)
+			{
+				FProperty* Property = *PropIt;
+				if (Property->HasMetaData(TEXT("Searchable")))
+				{
+					if (const PropertySearcherConstructorCallback* Callback = PropertySearcherConstructors.Find(Property->GetClass()))
+					{
+						if (!UsedConstructors.Contains(Callback))
+						{
+							UsedConstructors.Add(Callback);
+							Results.Add((*Callback)(ColumnTypePointer, *Property));
+						}
+					}
+					else
+					{
+						UE_LOG(LogEditorDataStorageUI, Warning, 
+							TEXT("Property '%s' in class '%s' is marked as searchable, but no property searcher for type '%s' was registered."),
+							*Property->GetName(), *Property->GetFullGroupName(true), *Property->GetClass()->GetName());
+					}
+				}
+			}
+		}
+	}
+}
+
+void UEditorDataStorageUi::RegisterSearcherGeneratorForProperty(const FFieldClass* PropertyType, 
+	PropertySearcherConstructorCallback PropertySearcherConstructor)
+{
+	checkf(!PropertySearcherConstructors.Contains(PropertyType), 
+		TEXT("Double registration of property searcher for type '%s'."), *PropertyType->GetName());
+	PropertySearcherConstructors.Add(PropertyType, MoveTemp(PropertySearcherConstructor));
+}
+
+void UEditorDataStorageUi::UnregisterSearcherGeneratorForProperty(const FFieldClass* PropertyType)
+{
+	PropertySearcherConstructors.Remove(PropertyType);
+}
+
 void UEditorDataStorageUi::GeneratePropertySorters(TArray<TSharedPtr<const UE::Editor::DataStorage::FColumnSorterInterface>>& Results,
 	TArrayView<TWeakObjectPtr<const UScriptStruct>> Columns) const
 {
diff --git a/Engine/Plugins/Experimental/EditorDataStorage/Source/TedsCore/Private/TypedElementDatabaseUI.h b/Engine/Plugins/Experimental/EditorDataStorage/Source/TedsCore/Private/TypedElementDatabaseUI.h
--- a/Engine/Plugins/Experimental/EditorDataStorage/Source/TedsCore/Private/TypedElementDatabaseUI.h
+++ b/Engine/Plugins/Experimental/EditorDataStorage/Source/TedsCore/Private/TypedElementDatabaseUI.h
@@ -100,6 +100,14 @@ public:
 	/** Get the table where TEDS UI widgets are stored */
 	UE_API virtual UE::Editor::DataStorage::TableHandle GetWidgetTable() const override;
 
+	UE_API virtual void GeneratePropertySearchers(TArray<TSharedPtr<const UE::Editor::DataStorage::FColumnSearcherInterface>>& Results,
+		TArrayView<TWeakObjectPtr<const UScriptStruct>> Columns) const override;
+
+	UE_API virtual void RegisterSearcherGeneratorForProperty(const FFieldClass* PropertyType, 
+		PropertySearcherConstructorCallback PropertySearcherConstructor) override;
+
+	UE_API virtual void UnregisterSearcherGeneratorForProperty(const FFieldClass* PropertyType) override;
+
 	UE_API virtual void GeneratePropertySorters(TArray<TSharedPtr<const UE::Editor::DataStorage::FColumnSorterInterface>>& Results,
 		TArrayView<TWeakObjectPtr<const UScriptStruct>> Columns) const override;
 
@@ -185,6 +193,7 @@ private:
 	UE_API void RegisterDecoratorWidgetObservers(UE::Editor::DataStorage::RowHandle PurposeRowHandle, const UScriptStruct* DecoratorColumn);
 
 	TMap<const FFieldClass*, PropertySorterConstructorCallback> PropertySorterConstructors;
+	TMap<const FFieldClass*, PropertySearcherConstructorCallback> PropertySearcherConstructors;
 
 	UE::Editor::DataStorage::TableHandle WidgetTable{ UE::Editor::DataStorage::InvalidTableHandle };
 	UE::Editor::DataStorage::TableHandle WidgetPurposeTable{ UE::Editor::DataStorage::InvalidTableHandle };
diff --git a/Engine/Plugins/Experimental/EditorDataStorage/Source/TedsCore/Private/TypedElementsDataStorage.cpp b/Engine/Plugins/Experimental/EditorDataStorage/Source/TedsCore/Private/TypedElementsDataStorage.cpp
--- a/Engine/Plugins/Experimental/EditorDataStorage/Source/TedsCore/Private/TypedElementsDataStorage.cpp
+++ b/Engine/Plugins/Experimental/EditorDataStorage/Source/TedsCore/Private/TypedElementsDataStorage.cpp
@@ -147,6 +147,7 @@ void FEditorDataStorageModule::StartupModule()
 					Factory->RegisterDealiaser(*DataStorageCompatibility);
 					
 					Factory->RegisterWidgetPurposes(*DataStorageUi);
+					Factory->RegisterPropertySearchers(*DataStorageUi);
 					Factory->RegisterPropertySorters(*DataStorageUi);
 				}
 
diff --git a/Engine/Plugins/Experimental/EditorDataStorageFeatures/Source/TedsOutliner/Private/TedsOutlinerColumn.cpp b/Engine/Plugins/Experimental/EditorDataStorageFeatures/Source/TedsOutliner/Private/TedsOutlinerColumn.cpp
--- a/Engine/Plugins/Experimental/EditorDataStorageFeatures/Source/TedsOutliner/Private/TedsOutlinerColumn.cpp
+++ b/Engine/Plugins/Experimental/EditorDataStorageFeatures/Source/TedsOutliner/Private/TedsOutlinerColumn.cpp
@@ -167,8 +167,17 @@ namespace UE::Editor::Outliner
 
 	void FTedsOutlinerUiColumn::PopulateSearchStrings( const ISceneOutlinerTreeItem& Item, TArray< FString >& OutSearchStrings ) const
 	{
-		// TODO: We don't currently have a way to convert TEDS widgets into searchable strings, but we can rely on the fallback column if it exists
-		if (FallbackColumn)
+		using namespace UE::Editor::DataStorage;
+		
+		if (!bHybridMode)
+		{
+			if(const FTedsOutlinerTreeItem* OutlinerTreeItem = Item.CastTo<FTedsOutlinerTreeItem>())
+			{
+				const RowHandle RowHandle = OutlinerTreeItem->GetRowHandle();
+				TableViewerColumnImpl->OnGetSearchStringCallback(RowHandle, OutSearchStrings);
+			}
+		}
+		else if (FallbackColumn)
 		{
 			FallbackColumn->PopulateSearchStrings(Item, OutSearchStrings);
 		}
diff --git a/Engine/Plugins/Experimental/EditorDataStorageFeatures/Source/TedsOutliner/Private/Widgets/OutlinerLevelWidget.cpp b/Engine/Plugins/Experimental/EditorDataStorageFeatures/Source/TedsOutliner/Private/Widgets/OutlinerLevelWidget.cpp
--- a/Engine/Plugins/Experimental/EditorDataStorageFeatures/Source/TedsOutliner/Private/Widgets/OutlinerLevelWidget.cpp
+++ b/Engine/Plugins/Experimental/EditorDataStorageFeatures/Source/TedsOutliner/Private/Widgets/OutlinerLevelWidget.cpp
@@ -28,6 +28,14 @@ namespace UE::Editor::DataStorage
 	{
 		return CreateSortPrefix(ByteIndex, TSortStringView(FSortCaseSensitive{}, Column.Level->GetOuter()->GetName()));
 	}
+	
+	void FLevelWidgetSearcher::GetSearchableString(FString& SearchableString, const ICoreProvider& Storage, const FLevelColumn& Column) const
+	{
+		if (const TWeakObjectPtr<ULevel> Level = Column.Level; Level.IsValid())
+		{
+			SearchableString = Level->GetOuter()->GetName();
+		}
+	}
 }
 
 void UOutlinerLevelWidgetFactory::RegisterWidgetConstructors(UE::Editor::DataStorage::ICoreProvider& DataStorage,
@@ -98,4 +106,16 @@ TArray<TSharedPtr<const UE::Editor::DataStorage::FColumnSorterInterface>> FOutli
 		});
 }
 
+TArray<TSharedPtr<const UE::Editor::DataStorage::FColumnSearcherInterface>> FOutlinerLevelWidgetConstructor::ConstructColumnSearchers(
+	UE::Editor::DataStorage::ICoreProvider* DataStorage,
+	UE::Editor::DataStorage::IUiProvider* DataStorageUi,
+	const UE::Editor::DataStorage::FMetaDataView& Arguments)
+{
+	using namespace UE::Editor::DataStorage;
+	return TArray<TSharedPtr<const FColumnSearcherInterface>>(
+		{
+			MakeShared<FLevelWidgetSearcher>()
+		});
+}
+
 #undef LOCTEXT_NAMESPACE
diff --git a/Engine/Plugins/Experimental/EditorDataStorageFeatures/Source/TedsOutliner/Private/Widgets/OutlinerMobilityWidget.cpp b/Engine/Plugins/Experimental/EditorDataStorageFeatures/Source/TedsOutliner/Private/Widgets/OutlinerMobilityWidget.cpp
--- a/Engine/Plugins/Experimental/EditorDataStorageFeatures/Source/TedsOutliner/Private/Widgets/OutlinerMobilityWidget.cpp
+++ b/Engine/Plugins/Experimental/EditorDataStorageFeatures/Source/TedsOutliner/Private/Widgets/OutlinerMobilityWidget.cpp
@@ -13,6 +13,29 @@
 
 #define LOCTEXT_NAMESPACE "OutlinerMobilityWidget"
 
+FText GetComponentMobility(const EComponentMobility::Type InMobility)
+{
+	switch(InMobility)
+	{
+	case EComponentMobility::Movable:
+		return LOCTEXT("OutlinerMobilityTypeMovable", "Movable");
+	case EComponentMobility::Static:
+		return LOCTEXT("OutlinerMobilityTypeStatic", "Static");
+	case EComponentMobility::Stationary:
+		return LOCTEXT("OutlinerMobilityTypeStationary", "Stationary");
+	default:
+		return LOCTEXT("OutlinerMobilityTypeUnknown", "<unknown>");
+	}
+}
+
+namespace UE::Editor::DataStorage
+{
+	void FMobilityWidgetSearcher::GetSearchableString(FString& SearchableString, const ICoreProvider& Storage, const FTedsActorMobilityColumn& Column) const
+	{
+		SearchableString = GetComponentMobility(Column.Mobility).ToString();
+	}
+}
+
 void UOutlinerMobilityWidgetFactory::RegisterWidgetConstructors(UE::Editor::DataStorage::ICoreProvider& DataStorage,
 	UE::Editor::DataStorage::IUiProvider& DataStorageUi) const
 {
@@ -54,19 +77,21 @@ TSharedPtr<SWidget> FOutlinerMobilityWidgetConstructor::CreateWidget(
 					.ColorAndOpacity(FSlateColor::UseSubduedForeground())
 					.Text(Binder.BindData(&FTedsActorMobilityColumn::Mobility, [](EComponentMobility::Type InMobility)
 						{
-							switch(InMobility)
-							{
-								case EComponentMobility::Movable:
-									return LOCTEXT("OutlinerMobilityTypeMovable", "Movable");
-								case EComponentMobility::Static:
-									return LOCTEXT("OutlinerMobilityTypeStatic", "Static");
-								case EComponentMobility::Stationary:
-									return LOCTEXT("OutlinerMobilityTypeStationary", "Stationary");
-								default:
-									return LOCTEXT("OutlinerMobilityTypeUnknown", "<unknown>");
-							}
+							return GetComponentMobility(InMobility);
 						}))
 			];
 }
 
+TArray<TSharedPtr<const UE::Editor::DataStorage::FColumnSearcherInterface>> FOutlinerMobilityWidgetConstructor::ConstructColumnSearchers(
+	UE::Editor::DataStorage::ICoreProvider* DataStorage,
+	UE::Editor::DataStorage::IUiProvider* DataStorageUi,
+	const UE::Editor::DataStorage::FMetaDataView& Arguments)
+{
+	using namespace UE::Editor::DataStorage;
+	return TArray<TSharedPtr<const FColumnSearcherInterface>>(
+		{
+			MakeShared<FMobilityWidgetSearcher>()
+		});
+}
+
 #undef LOCTEXT_NAMESPACE
diff --git a/Engine/Plugins/Experimental/EditorDataStorageFeatures/Source/TedsOutliner/Private/Widgets/TypeInfoWidget.cpp b/Engine/Plugins/Experimental/EditorDataStorageFeatures/Source/TedsOutliner/Private/Widgets/TypeInfoWidget.cpp
--- a/Engine/Plugins/Experimental/EditorDataStorageFeatures/Source/TedsOutliner/Private/Widgets/TypeInfoWidget.cpp
+++ b/Engine/Plugins/Experimental/EditorDataStorageFeatures/Source/TedsOutliner/Private/Widgets/TypeInfoWidget.cpp
@@ -35,6 +35,15 @@ namespace UE::Editor::DataStorage
 		FString TypeInfoName = TypeInfoPtr ? TypeInfoPtr->GetName() : TEXT("NULL");
 		return CreateSortPrefix(ByteIndex, TSortStringView(FSortCaseInsensitive{}, TypeInfoName));
 	}
+
+	void FTypeInfoWidgetSearcher::GetSearchableString(FString& SearchableString, const ICoreProvider& Storage, const FTypedElementClassTypeInfoColumn& Column) const
+	{
+		if (const TWeakObjectPtr<const UClass> TypeInfo = Column.TypeInfo; TypeInfo.IsValid())
+		{
+			SearchableString = TypeInfo->GetName();
+		}
+	}
+
 }
 
 
@@ -107,6 +116,18 @@ TArray<TSharedPtr<const UE::Editor::DataStorage::FColumnSorterInterface>> FTypeI
 		});
 }
 
+TArray<TSharedPtr<const UE::Editor::DataStorage::FColumnSearcherInterface>> FTypeInfoWidgetConstructor::ConstructColumnSearchers(
+	UE::Editor::DataStorage::ICoreProvider* DataStorage,
+	UE::Editor::DataStorage::IUiProvider* DataStorageUi,
+	const UE::Editor::DataStorage::FMetaDataView& Arguments)
+{
+	using namespace UE::Editor::DataStorage;
+	return TArray<TSharedPtr<const FColumnSearcherInterface>>(
+		{
+			MakeShared<FTypeInfoWidgetSearcher>()
+		});
+}
+
 TSharedRef<SWidget> FTypeInfoWidgetConstructor::CreateTextWidget(UE::Editor::DataStorage::ICoreProvider* DataStorage,
 	UE::Editor::DataStorage::RowHandle TargetRow)
 {
diff --git a/Engine/Plugins/Experimental/EditorDataStorageFeatures/Source/TedsOutliner/Public/Widgets/OutlinerLevelWidget.h b/Engine/Plugins/Experimental/EditorDataStorageFeatures/Source/TedsOutliner/Public/Widgets/OutlinerLevelWidget.h
--- a/Engine/Plugins/Experimental/EditorDataStorageFeatures/Source/TedsOutliner/Public/Widgets/OutlinerLevelWidget.h
+++ b/Engine/Plugins/Experimental/EditorDataStorageFeatures/Source/TedsOutliner/Public/Widgets/OutlinerLevelWidget.h
@@ -21,6 +21,12 @@ namespace UE::Editor::DataStorage
 		virtual int32 Compare(const FLevelColumn& Left, const FLevelColumn& Right) const override;
 		virtual FPrefixInfo CalculatePrefix(const FLevelColumn& Column, uint32 ByteIndex) const override;
 	};
+	
+	class FLevelWidgetSearcher final : public TColumnSearcherInterface<FLevelColumn>
+    {
+    protected:
+    	virtual void GetSearchableString(FString& SearchableString, const ICoreProvider& Storage, const FLevelColumn& Column) const override;
+    };
 }
 
 UCLASS()
@@ -50,10 +56,16 @@ public:
 		UE::Editor::DataStorage::RowHandle TargetRow,
 		UE::Editor::DataStorage::RowHandle WidgetRow,
 		const UE::Editor::DataStorage::FMetaDataView& Arguments) override;
+		
 	UE_API virtual TArray<TSharedPtr<const UE::Editor::DataStorage::FColumnSorterInterface>> ConstructColumnSorters(
 		UE::Editor::DataStorage::ICoreProvider* DataStorage,
 		UE::Editor::DataStorage::IUiProvider* DataStorageUi,
 		const UE::Editor::DataStorage::FMetaDataView& Arguments) override;
+		
+	UE_API virtual TArray<TSharedPtr<const UE::Editor::DataStorage::FColumnSearcherInterface>> ConstructColumnSearchers(
+			UE::Editor::DataStorage::ICoreProvider* DataStorage,
+			UE::Editor::DataStorage::IUiProvider* DataStorageUi,
+			const UE::Editor::DataStorage::FMetaDataView& Arguments) override;
 };
 
 #undef UE_API
diff --git a/Engine/Plugins/Experimental/EditorDataStorageFeatures/Source/TedsOutliner/Public/Widgets/OutlinerMobilityWidget.h b/Engine/Plugins/Experimental/EditorDataStorageFeatures/Source/TedsOutliner/Public/Widgets/OutlinerMobilityWidget.h
--- a/Engine/Plugins/Experimental/EditorDataStorageFeatures/Source/TedsOutliner/Public/Widgets/OutlinerMobilityWidget.h
+++ b/Engine/Plugins/Experimental/EditorDataStorageFeatures/Source/TedsOutliner/Public/Widgets/OutlinerMobilityWidget.h
@@ -5,11 +5,23 @@
 #include "Elements/Interfaces/TypedElementDataStorageFactory.h"
 #include "Elements/Interfaces/TypedElementDataStorageInterface.h"
 #include "Elements/Interfaces/TypedElementDataStorageUiInterface.h"
+#include "Columns/TedsActorMobilityColumns.h"
 
 #include "OutlinerMobilityWidget.generated.h"
 
 #define UE_API TEDSOUTLINER_API
 
+struct FTedsActorMobilityColumn;
+
+namespace UE::Editor::DataStorage
+{
+	class FMobilityWidgetSearcher final : public TColumnSearcherInterface<FTedsActorMobilityColumn>
+	{
+	protected:
+		virtual void GetSearchableString(FString& SearchableString, const ICoreProvider& Storage, const FTedsActorMobilityColumn& Column) const override;
+	};
+}
+
 UCLASS()
 class UOutlinerMobilityWidgetFactory : public UEditorDataStorageFactory
 {
@@ -38,6 +50,10 @@ public:
 		UE::Editor::DataStorage::RowHandle WidgetRow,
 		const UE::Editor::DataStorage::FMetaDataView& Arguments) override;
 
+	UE_API virtual TArray<TSharedPtr<const UE::Editor::DataStorage::FColumnSearcherInterface>> ConstructColumnSearchers(
+		UE::Editor::DataStorage::ICoreProvider* DataStorage, UE::Editor::DataStorage::IUiProvider* DataStorageUi,
+		const UE::Editor::DataStorage::FMetaDataView& Arguments) override;
+
 };
 
 #undef UE_API
diff --git a/Engine/Plugins/Experimental/EditorDataStorageFeatures/Source/TedsOutliner/Public/Widgets/TypeInfoWidget.h b/Engine/Plugins/Experimental/EditorDataStorageFeatures/Source/TedsOutliner/Public/Widgets/TypeInfoWidget.h
--- a/Engine/Plugins/Experimental/EditorDataStorageFeatures/Source/TedsOutliner/Public/Widgets/TypeInfoWidget.h
+++ b/Engine/Plugins/Experimental/EditorDataStorageFeatures/Source/TedsOutliner/Public/Widgets/TypeInfoWidget.h
@@ -21,6 +21,12 @@ namespace UE::Editor::DataStorage
 		virtual int32 Compare(const FTypedElementClassTypeInfoColumn& Left, const FTypedElementClassTypeInfoColumn& Right) const override;
 		virtual FPrefixInfo CalculatePrefix(const FTypedElementClassTypeInfoColumn& Column, uint32 ByteIndex) const override;
 	};
+	
+	class FTypeInfoWidgetSearcher final : public TColumnSearcherInterface<FTypedElementClassTypeInfoColumn>
+    {
+    protected:
+        virtual void GetSearchableString(FString& SearchableString, const ICoreProvider& Storage, const FTypedElementClassTypeInfoColumn& Column) const override;
+    };
 }
 
 UCLASS()
@@ -59,6 +65,11 @@ protected:
 		UE::Editor::DataStorage::IUiProvider* DataStorageUi,
 		const UE::Editor::DataStorage::FMetaDataView& Arguments) override;
 
+	UE_API TArray<TSharedPtr<const UE::Editor::DataStorage::FColumnSearcherInterface>> ConstructColumnSearchers(
+		UE::Editor::DataStorage::ICoreProvider* DataStorage,
+		UE::Editor::DataStorage::IUiProvider* DataStorageUi,
+		const UE::Editor::DataStorage::FMetaDataView& Arguments) override;
+
 protected:
 
 	static TSharedRef<SWidget> CreateTextWidget(UE::Editor::DataStorage::ICoreProvider* DataStorage, UE::Editor::DataStorage::RowHandle TargetRow);
diff --git a/Engine/Plugins/Experimental/EditorDataStorageFeatures/Source/TedsTableViewer/Private/TedsTableViewerColumn.cpp b/Engine/Plugins/Experimental/EditorDataStorageFeatures/Source/TedsTableViewer/Private/TedsTableViewerColumn.cpp
--- a/Engine/Plugins/Experimental/EditorDataStorageFeatures/Source/TedsTableViewer/Private/TedsTableViewerColumn.cpp
+++ b/Engine/Plugins/Experimental/EditorDataStorageFeatures/Source/TedsTableViewer/Private/TedsTableViewerColumn.cpp
@@ -45,11 +45,19 @@ namespace UE::Editor::DataStorage
 		if (HeaderWidgetConstructor)
 		{
 			ColumnSorters = HeaderWidgetConstructor->ConstructColumnSorters(Storage, StorageUi, GetMetaData());
+			ColumnSearchers = HeaderWidgetConstructor->ConstructColumnSearchers(Storage, StorageUi, GetMetaData());
 		}
 		// Prefer to use the sort function in the header, but if it doesn't exist, use the one from the cell.
 		if (CellWidgetConstructor && ColumnSorters.IsEmpty())
 		{
-			ColumnSorters = CellWidgetConstructor->ConstructColumnSorters(Storage, StorageUi, GetMetaData());
+			if (ColumnSorters.IsEmpty())
+			{
+				ColumnSorters = CellWidgetConstructor->ConstructColumnSorters(Storage, StorageUi, GetMetaData());
+			}
+			if (ColumnSearchers.IsEmpty())
+			{
+				ColumnSearchers = CellWidgetConstructor->ConstructColumnSearchers(Storage, StorageUi, GetMetaData());
+			}
 		}
 
 		RegisterQueries();
@@ -234,6 +242,23 @@ namespace UE::Editor::DataStorage
 		IsRowVisibleDelegate = MoveTemp(InIsRowVisibleDelegate);
 	}
 
+	void FTedsTableViewerColumn::OnGetSearchStringCallback(const RowHandle& InRowHandle, TArray< FString >& OutSearchStrings)
+	{
+		if (Storage)
+		{
+			FString OutString;
+			for (const TSharedPtr<const FColumnSearcherInterface>& Searcher : ColumnSearchers )
+			{
+				OutString.Reset(FName::StringBufferSize);
+				Searcher->GetSearchableString(OutString, *Storage, InRowHandle);
+				if (!OutString.IsEmpty())
+				{
+					OutSearchStrings.Add(OutString);
+				}
+			}
+		}
+	}
+
 	bool FTedsTableViewerColumn::HasSortInfo() const
 	{
 		return !ColumnSorters.IsEmpty();
diff --git a/Engine/Plugins/Experimental/EditorDataStorageFeatures/Source/TedsTableViewer/Public/TedsTableViewerColumn.h b/Engine/Plugins/Experimental/EditorDataStorageFeatures/Source/TedsTableViewer/Public/TedsTableViewerColumn.h
--- a/Engine/Plugins/Experimental/EditorDataStorageFeatures/Source/TedsTableViewer/Public/TedsTableViewerColumn.h
+++ b/Engine/Plugins/Experimental/EditorDataStorageFeatures/Source/TedsTableViewer/Public/TedsTableViewerColumn.h
@@ -67,6 +67,7 @@ namespace UE::Editor::DataStorage
 			FIsSorting IsSortingHandler;
 		};
 
+		TEDSTABLEVIEWER_API void OnGetSearchStringCallback(const RowHandle& InRowHandle, TArray< FString >& OutSearchStrings);
 		TEDSTABLEVIEWER_API void SetInitialSortMode(EColumnSortMode::Type InSortMode);
 		TEDSTABLEVIEWER_API void SetSortDelegates(FSortHandler InSortDelegates);
 		TEDSTABLEVIEWER_API bool HasSortInfo() const;
@@ -116,6 +117,8 @@ namespace UE::Editor::DataStorage
 		QueryHandle WidgetQuery;
 		TMap<RowHandle, bool> RowsToUpdate;
 
+		TArray<TSharedPtr<const FColumnSearcherInterface>> ColumnSearchers;
+
 		// Sorting information
 		TArray<TSharedPtr<const FColumnSorterInterface>> ColumnSorters;
 		FSortHandler SortDelegates;
diff --git a/Engine/Source/Editor/UnrealEd/Private/Layers/Widgets/ActorLayersWidget.cpp b/Engine/Source/Editor/UnrealEd/Private/Layers/Widgets/ActorLayersWidget.cpp
--- a/Engine/Source/Editor/UnrealEd/Private/Layers/Widgets/ActorLayersWidget.cpp
+++ b/Engine/Source/Editor/UnrealEd/Private/Layers/Widgets/ActorLayersWidget.cpp
@@ -3,10 +3,28 @@
 #include "ActorLayersWidget.h"
 
 #include "Elements/Framework/TypedElementAttributeBinding.h"
-#include "Layers/Columns/LayersColumns.h"
 #include "Widgets/SBoxPanel.h"
 #include "Widgets/Text/STextBlock.h"
 
+namespace UE::Editor::DataStorage
+{
+	void FActorLayersWidgetSearcher::GetSearchableString(FString& SearchableString, const ICoreProvider& Storage, const Layers::FActorLayersColumn& Column) const
+	{
+		for (const RowHandle LayerRow : Column.Layers.GetRows())
+		{
+			if (const Layers::FLayerNameColumn* LayerNameColumn = Storage.GetColumn<Layers::FLayerNameColumn>(LayerRow))
+			{
+				if (SearchableString.Len())
+				{
+					SearchableString += TEXT(", ");
+				}
+
+				SearchableString += LayerNameColumn->LayerName.ToString();
+			}
+		}
+	}
+}
+
 FActorLayersWidgetConstructor::FActorLayersWidgetConstructor()
 	: FSimpleWidgetConstructor(StaticStruct())
 {
@@ -48,4 +66,16 @@ TSharedPtr<SWidget> FActorLayersWidgetConstructor::CreateWidget(UE::Editor::Data
 					return FText::FromString(Result);
 				}))
 		];
+}
+
+TArray<TSharedPtr<const UE::Editor::DataStorage::FColumnSearcherInterface>> FActorLayersWidgetConstructor::ConstructColumnSearchers(
+	UE::Editor::DataStorage::ICoreProvider* DataStorage,
+	UE::Editor::DataStorage::IUiProvider* DataStorageUi,
+	const UE::Editor::DataStorage::FMetaDataView& Arguments)
+{
+	using namespace UE::Editor::DataStorage;
+	return TArray<TSharedPtr<const FColumnSearcherInterface>>(
+		{
+			MakeShared<FActorLayersWidgetSearcher>()
+		});
 }
\ No newline at end of file
diff --git a/Engine/Source/Editor/UnrealEd/Private/Layers/Widgets/ActorLayersWidget.h b/Engine/Source/Editor/UnrealEd/Private/Layers/Widgets/ActorLayersWidget.h
--- a/Engine/Source/Editor/UnrealEd/Private/Layers/Widgets/ActorLayersWidget.h
+++ b/Engine/Source/Editor/UnrealEd/Private/Layers/Widgets/ActorLayersWidget.h
@@ -4,9 +4,18 @@
 
 #include "Elements/Interfaces/TypedElementDataStorageFactory.h"
 #include "Elements/Interfaces/TypedElementDataStorageUiInterface.h"
+#include "Layers/Columns/LayersColumns.h"
 
 #include "ActorLayersWidget.generated.h"
 
+namespace UE::Editor::DataStorage
+{
+	class FActorLayersWidgetSearcher final : public TColumnSearcherInterface<Layers::FActorLayersColumn>
+	{
+	protected:
+		virtual void GetSearchableString(FString& SearchableString, const ICoreProvider& Storage, const Layers::FActorLayersColumn& Column) const override;
+	};
+}
 
 // Widget to show the layers an actor belongs to
 USTRUCT()
@@ -24,4 +33,9 @@ public:
 		UE::Editor::DataStorage::RowHandle TargetRow,
 		UE::Editor::DataStorage::RowHandle WidgetRow,
 		const UE::Editor::DataStorage::FMetaDataView& Arguments) override;
+		
+	virtual TArray<TSharedPtr<const UE::Editor::DataStorage::FColumnSearcherInterface>> ConstructColumnSearchers(
+		UE::Editor::DataStorage::ICoreProvider* DataStorage,
+		UE::Editor::DataStorage::IUiProvider* DataStorageUi,
+		const UE::Editor::DataStorage::FMetaDataView& Arguments) override;
 };
\ No newline at end of file
diff --git a/Engine/Source/Runtime/TypedElementFramework/Private/Elements/Interfaces/TypedElementDataStorageUiInterface.cpp b/Engine/Source/Runtime/TypedElementFramework/Private/Elements/Interfaces/TypedElementDataStorageUiInterface.cpp
--- a/Engine/Source/Runtime/TypedElementFramework/Private/Elements/Interfaces/TypedElementDataStorageUiInterface.cpp
+++ b/Engine/Source/Runtime/TypedElementFramework/Private/Elements/Interfaces/TypedElementDataStorageUiInterface.cpp
@@ -210,20 +210,43 @@ TArray<TSharedPtr<const UE::Editor::DataStorage::FColumnSorterInterface>> FTyped
 	const UE::Editor::DataStorage::FMetaDataView& Arguments)
 {
 	using namespace UE::Editor::DataStorage;
+	
+	TArray<TWeakObjectPtr<const UScriptStruct>> ColumnTypes;
+	ColumnTypes.Reserve(MatchedColumnTypes.Num());
+	auto GetAllColumns = [&ColumnTypes](TNotNull<const UScriptStruct*> Column)
+	{
+		ColumnTypes.Add(Column);
+	};
+	
+	ForAllColumnTypes(GetAllColumns);
 
 	TArray<TSharedPtr<const FColumnSorterInterface>> Results;
+	Results.Reserve(ColumnTypes.Num());
+	DataStorageUi->GeneratePropertySorters(Results, ColumnTypes);
 	
-	TArray<TWeakObjectPtr<const UScriptStruct>, TInlineAllocator<32>> DynamicKeys;
-	MatchedDynamicTemplates.GetKeys(DynamicKeys);
+	return Results;
+}
 
-	// Combine both column types and base dynamic column types so they can run in the same GeneratePropertySorter Function & check
-	// for any duplicate constructors.
-	TArray<TWeakObjectPtr<const UScriptStruct>> AllColumnTypes;
-	AllColumnTypes.Reserve(MatchedColumnTypes.Num() + DynamicKeys.Num());
-	AllColumnTypes.Append(MatchedColumnTypes);
-	AllColumnTypes.Append(DynamicKeys);
-	DataStorageUi->GeneratePropertySorters(Results, AllColumnTypes);
+TArray<TSharedPtr<const UE::Editor::DataStorage::FColumnSearcherInterface>> FTypedElementWidgetConstructor::ConstructColumnSearchers(
+	UE::Editor::DataStorage::ICoreProvider* DataStorage,
+	UE::Editor::DataStorage::IUiProvider* DataStorageUi,
+	const UE::Editor::DataStorage::FMetaDataView& Arguments)
+{
+	using namespace UE::Editor::DataStorage;
 
+	TArray<TWeakObjectPtr<const UScriptStruct>> ColumnTypes;
+	ColumnTypes.Reserve(MatchedColumnTypes.Num());
+	auto GetAllColumns = [&ColumnTypes](TNotNull<const UScriptStruct*> Column)
+		{
+			ColumnTypes.Add(Column);
+		};
+	
+	ForAllColumnTypes(GetAllColumns);
+
+	TArray<TSharedPtr<const FColumnSearcherInterface>> Results;
+	Results.Reserve(ColumnTypes.Num());
+	DataStorageUi->GeneratePropertySearchers(Results, ColumnTypes);
+	
 	return Results;
 }
 
@@ -312,6 +335,29 @@ UE::Editor::DataStorage::RowHandle FTypedElementWidgetConstructor::GetTargetRow(
 	return TargetRow;
 }
 
+void FTypedElementWidgetConstructor::ForAllColumnTypes(TFunctionRef<void(TNotNull<const UScriptStruct*>)> Function) const
+{
+	using namespace UE::Editor::DataStorage;
+	for (const TWeakObjectPtr<const UScriptStruct>& MatchedColumn : MatchedColumnTypes)
+	{
+		if (const UScriptStruct* MatchedColumnStruct = MatchedColumn.Get())
+		{
+			Function(TNotNull<const UScriptStruct*>(MatchedColumnStruct));
+		}
+	}
+	
+	TArray<TWeakObjectPtr<const UScriptStruct>, TInlineAllocator<32>> DynamicKeys;
+	MatchedDynamicTemplates.GetKeys(DynamicKeys);
+
+	for (const TWeakObjectPtr<const UScriptStruct>& DynamicKey : DynamicKeys)
+	{
+		if (const UScriptStruct* DynamicKeyStruct = DynamicKey.Get())
+		{
+			Function(TNotNull<const UScriptStruct*>(DynamicKeyStruct));
+		}
+	}
+}
+
 UE::Editor::DataStorage::RowHandle FTypedElementWidgetConstructor::GetWidgetFactoryRow() const
 {
 	return WidgetFactoryRow;
diff --git a/Engine/Source/Runtime/TypedElementFramework/Public/Elements/Framework/TypedElementSearcher.h b/Engine/Source/Runtime/TypedElementFramework/Public/Elements/Framework/TypedElementSearcher.h
new file mode 100644
--- /dev/null
+++ b/Engine/Source/Runtime/TypedElementFramework/Public/Elements/Framework/TypedElementSearcher.h
@@ -0,0 +1,41 @@
+// Copyright Epic Games, Inc. All Rights Reserved.
+
+#pragma once
+
+#include "Containers/UnrealString.h"
+#include "DataStorage/Handles.h"
+#include "Elements/Interfaces/TypedElementDataStorageInterface.h"
+
+namespace UE::Editor::DataStorage
+{
+	/**
+	 * @section Search provider
+	 *
+	 * @description
+	 * Interface and supporting classes to provide searching of rows by column.
+	 */
+	class ICoreProvider;
+
+	/** Interface to provide searching of rows by column. */
+	class FColumnSearcherInterface
+	{
+	public:
+		virtual ~FColumnSearcherInterface() = default;
+
+		virtual void GetSearchableString(FString& SearchableString, const ICoreProvider& Storage, RowHandle Row) const = 0;
+	};
+
+	template<typename ColumnType>
+	class TColumnSearcherInterface : public FColumnSearcherInterface
+	{
+	public:
+		virtual ~TColumnSearcherInterface() override = default;
+
+		virtual void GetSearchableString(FString& SearchableString, const ICoreProvider& Storage, RowHandle Row) const override final;
+
+	protected:
+		virtual void GetSearchableString(FString& SearchableString, const ICoreProvider& Storage, const ColumnType& Column) const = 0;
+	};
+}// namespace UE::Editor::DataStorage
+
+#include "Elements/Framework/TypedElementSearcher.inl"
\ No newline at end of file
diff --git a/Engine/Source/Runtime/TypedElementFramework/Public/Elements/Framework/TypedElementSearcher.inl b/Engine/Source/Runtime/TypedElementFramework/Public/Elements/Framework/TypedElementSearcher.inl
new file mode 100644
--- /dev/null
+++ b/Engine/Source/Runtime/TypedElementFramework/Public/Elements/Framework/TypedElementSearcher.inl
@@ -0,0 +1,18 @@
+// Copyright Epic Games, Inc. All Rights Reserved.
+
+#pragma once
+
+#include "Elements/Interfaces/TypedElementDataStorageInterface.h"
+
+namespace UE::Editor::DataStorage
+{
+	template<typename ColumnType>
+	void TColumnSearcherInterface<ColumnType>::GetSearchableString(
+		FString& SearchableString, const ICoreProvider& Storage, RowHandle Row) const
+	{
+		if (const ColumnType* Column = Storage.GetColumn<ColumnType>(Row))
+		{
+			GetSearchableString(SearchableString, Storage, *Column);
+		}
+	}
+}// namespace UE::Editor::DataStorage
\ No newline at end of file
diff --git a/Engine/Source/Runtime/TypedElementFramework/Public/Elements/Interfaces/TypedElementDataStorageFactory.h b/Engine/Source/Runtime/TypedElementFramework/Public/Elements/Interfaces/TypedElementDataStorageFactory.h
--- a/Engine/Source/Runtime/TypedElementFramework/Public/Elements/Interfaces/TypedElementDataStorageFactory.h
+++ b/Engine/Source/Runtime/TypedElementFramework/Public/Elements/Interfaces/TypedElementDataStorageFactory.h
@@ -52,10 +52,17 @@ public:
 	virtual void RegisterWidgetPurposes(UE::Editor::DataStorage::IUiProvider& DataStorageUi) const {}
 	virtual void RegisterWidgetConstructors(UE::Editor::DataStorage::ICoreProvider& DataStorage,
 		UE::Editor::DataStorage::IUiProvider& DataStorageUi) const {}
+		
+	/** 
+     * Used to register searchers for a specific property type to be used in the default generated property searcher. The widget constructor
+     * can be used to provide a custom searcher if more complex searching logic is needed. Implementations of this function should mostly be
+     * calling `RegisterSearcherGeneratorForProperty` on the UI provider.
+     */
+    virtual void RegisterPropertySearchers(UE::Editor::DataStorage::IUiProvider& DataStorageUi) const {}
 
 	/** 
 	 * Used to register sorters for a specific property type to be used in the default generated property sorter. The widget constructor
-	 * can used to provide a custom sorter if more complex sorting logic is needed. Implementations of this function should mostly be
+	 * can be used to provide a custom sorter if more complex sorting logic is needed. Implementations of this function should mostly be
 	 * calling `RegisterSorterGeneratorForProperty` on the UI provider.
 	 */
 	virtual void RegisterPropertySorters(UE::Editor::DataStorage::IUiProvider& DataStorageUi) const {}
diff --git a/Engine/Source/Runtime/TypedElementFramework/Public/Elements/Interfaces/TypedElementDataStorageUiInterface.h b/Engine/Source/Runtime/TypedElementFramework/Public/Elements/Interfaces/TypedElementDataStorageUiInterface.h
--- a/Engine/Source/Runtime/TypedElementFramework/Public/Elements/Interfaces/TypedElementDataStorageUiInterface.h
+++ b/Engine/Source/Runtime/TypedElementFramework/Public/Elements/Interfaces/TypedElementDataStorageUiInterface.h
@@ -10,6 +10,7 @@
 #include "Elements/Interfaces/TypedElementDataStorageInterface.h"
 #include "Elements/Framework/TypedElementMetaData.h"
 #include "Elements/Framework/TypedElementQueryBuilder.h"
+#include "Elements/Framework/TypedElementSearcher.h"
 #include "Elements/Framework/TypedElementSorter.h"
 #include "Features/IModularFeature.h"
 #include "Templates/UnrealTypeTraits.h"
@@ -154,6 +155,16 @@ public:
 		UE::Editor::DataStorage::ICoreProvider* DataStorage,
 		UE::Editor::DataStorage::IUiProvider* DataStorageUi,
 		const UE::Editor::DataStorage::FMetaDataView& Arguments);
+		
+	/**
+     * Constructs a list of classes that are used to convert the column data into a searchable string on the rows with the column(s)
+     * associated with this widget constructor. A widget constructor can provide multiple searching constructors which will
+     * allow the same list of rows to be searched in different ways.
+     */
+    TYPEDELEMENTFRAMEWORK_API virtual TArray<TSharedPtr<const UE::Editor::DataStorage::FColumnSearcherInterface>> ConstructColumnSearchers(
+        UE::Editor::DataStorage::ICoreProvider* DataStorage,
+        UE::Editor::DataStorage::IUiProvider* DataStorageUi,
+        const UE::Editor::DataStorage::FMetaDataView& Arguments);
 
 	/**
 	 * Helper function to get the row containing the widget factory this constructor was factoried from
@@ -222,6 +233,11 @@ protected:
 	 */
 	TYPEDELEMENTFRAMEWORK_API UE::Editor::DataStorage::RowHandle GetTargetRow(
 		UE::Editor::DataStorage::ICoreProvider* DataStorage, UE::Editor::DataStorage::RowHandle WidgetRow) const;
+
+	/**
+	 * Helper function to perform a function on all column types for this widget (includes matched dynamic templates) 
+	 */
+	TYPEDELEMENTFRAMEWORK_API void ForAllColumnTypes(TFunctionRef<void(TNotNull<const UScriptStruct*>)> Function) const;
 };
 
 /**
@@ -302,6 +318,9 @@ struct TStructOpsTypeTraits<FTypedElementWidgetConstructor> : public TStructOpsT
 
 namespace UE::Editor::DataStorage
 {
+	template<typename PropertyClass> 
+	concept PropertyClassType = std::is_base_of_v<FProperty, PropertyClass>;
+	
 	class IUiProvider : public IModularFeature
 	{
 	public:
@@ -387,6 +406,8 @@ namespace UE::Editor::DataStorage
 		using WidgetPurposeCallback = TFunctionRef<void(FName, EPurposeType, const FText&)>;
 		using PropertySorterConstructorCallback = 
 			TFunction<TSharedPtr<const FColumnSorterInterface>(TWeakObjectPtr<const UScriptStruct> ColumnType, const FProperty& Property)>;
+		using PropertySearcherConstructorCallback = 
+			TFunction<TSharedPtr<const FColumnSearcherInterface>(TWeakObjectPtr<const UScriptStruct> ColumnType, const FProperty& Property)>;
 
 		/**
 		 * @section Widget purposes
@@ -551,6 +572,35 @@ namespace UE::Editor::DataStorage
 		/** Get the table where TEDS UI widgets are stored */
 		virtual TableHandle GetWidgetTable() const = 0;
 
+		/**
+		 * @section Searcher generation
+		 * Widget constructors can provide a custom searcher object that will allow Slate columns to convert Column data to an FString
+		 * that the searcher supports. If none is provided however, TEDS UI will generate one based on the properties in a column that
+		 * are marked with the meta data "Searchable". Out of the box TEDS UI supports string property types (FName, FString, FText), but
+		 * additional properties can be registered with the functions below as well. Note that this is mechanic is exclusively for properties
+		 * that are used in the dynamic generation of searchers. For complex searching that involves data specific knowledge or combines
+		 * multiple values it's highly recommended to create a custom searcher and attach it to the appropriate widget constructor.
+		 */
+
+		/** Build searchers for the provided list of columns */
+		virtual void GeneratePropertySearchers(TArray<TSharedPtr<const FColumnSearcherInterface>>& Results,
+			TArrayView<TWeakObjectPtr<const UScriptStruct>> Columns) const = 0;
+
+		/** Register a function to handle converting a property into a searcher. */
+		virtual void RegisterSearcherGeneratorForProperty(const FFieldClass* PropertyType, 
+			PropertySearcherConstructorCallback PropertySearcherConstructor) = 0;
+
+		/** Register a function to handle converting a property into a searcher. */
+		template<PropertyClassType PropertyClass>
+		void RegisterSearcherGeneratorForProperty(PropertySearcherConstructorCallback PropertySearcherConstructor);
+
+		/** Unregister a function to handle converting a property into a searcher. */
+		virtual void UnregisterSearcherGeneratorForProperty(const FFieldClass* PropertyType) = 0;
+
+		/** Unregister a function to handle converting a property into a searcher. */
+		template<PropertyClassType PropertyClass>
+		void UnregisterSearcherGeneratorForProperty();
+		
 		/**
 		 * @section Sorter generation
 		 * Widget constructors can provide a custom sorter object that will allows Slate columns to sort by the data that the sorter
@@ -570,15 +620,15 @@ namespace UE::Editor::DataStorage
 			PropertySorterConstructorCallback PropertySorterConstructor) = 0;
 
 		/** Register a function to handle converting a property into a sorter. */
-		template<typename PropertyClassType> requires(std::is_base_of_v<FProperty, PropertyClassType>)
-			void RegisterSorterGeneratorForProperty(PropertySorterConstructorCallback PropertySorterConstructor);
+		template<PropertyClassType PropertyClass>
+		void RegisterSorterGeneratorForProperty(PropertySorterConstructorCallback PropertySorterConstructor);
 
 		/** Unregister a function to handle converting a property into a sorter. */
 		virtual void UnregisterSorterGeneratorForProperty(const FFieldClass* PropertyType) = 0;
 
 		/** Unregister a function to handle converting a property into a sorter. */
-		template<typename PropertyClassType> requires(std::is_base_of_v<FProperty, PropertyClassType>)
-			void UnregisterSorterGeneratorForProperty();
+		template<PropertyClassType PropertyClass>
+		void UnregisterSorterGeneratorForProperty();
 
 		/**
 		 * @section Miscellaneous
@@ -693,15 +743,27 @@ namespace UE::Editor::DataStorage
 		return this->RegisterDecoratorWidgetFactory(PurposeRow, Constructor, ColumnType::StaticStruct());
 	}
 
-	template<typename PropertyClassType> requires(std::is_base_of_v<FProperty, PropertyClassType>)
-		void IUiProvider::RegisterSorterGeneratorForProperty(PropertySorterConstructorCallback PropertySorterConstructor)
+	template<PropertyClassType PropertyClass>
+	void IUiProvider::RegisterSearcherGeneratorForProperty(PropertySearcherConstructorCallback PropertySearcherConstructor)
+	{
+		RegisterSearcherGeneratorForProperty(PropertyClass::StaticClass(), PropertySearcherConstructor);
+	}
+
+	template<PropertyClassType PropertyClass>
+	void IUiProvider::UnregisterSearcherGeneratorForProperty()
+	{
+		UnregisterSearcherGeneratorForProperty(PropertyClass::StaticClass());
+	}
+
+	template<PropertyClassType PropertyClass>
+	void IUiProvider::RegisterSorterGeneratorForProperty(PropertySorterConstructorCallback PropertySorterConstructor)
 	{
-		RegisterSorterGeneratorForProperty(PropertyClassType::StaticClass(), PropertySorterConstructor);
+		RegisterSorterGeneratorForProperty(PropertyClass::StaticClass(), PropertySorterConstructor);
 	}
 
-	template<typename PropertyClassType> requires(std::is_base_of_v<FProperty, PropertyClassType>)
-		void IUiProvider::UnregisterSorterGeneratorForProperty()
+	template<PropertyClassType PropertyClass>
+	void IUiProvider::UnregisterSorterGeneratorForProperty()
 	{
-		UnregisterSorterGeneratorForProperty(PropertyClassType::StaticClass());
+		UnregisterSorterGeneratorForProperty(PropertyClass::StaticClass());
 	}
 }
\ No newline at end of file
