[TEDS] Relations

diff --git a/Engine/Plugins/Experimental/EditorDataStorage/Source/TedsCore/Private/TypedElementDatabase.h b/Engine/Plugins/Experimental/EditorDataStorage/Source/TedsCore/Private/TypedElementDatabase.h
--- a/Engine/Plugins/Experimental/EditorDataStorage/Source/TedsCore/Private/TypedElementDatabase.h
+++ b/Engine/Plugins/Experimental/EditorDataStorage/Source/TedsCore/Private/TypedElementDatabase.h
@@ -199,6 +199,105 @@ public:
 	UE_API virtual UE::Editor::DataStorage::RowHandle GetRootScope() const override;
 	UE_API virtual UE::Editor::DataStorage::FHierarchyHandle GetScopeHierarchy() override;
 
+	// Relations Interface - Relation Type Management
+	UE_API virtual UE::Editor::DataStorage::RelationTypeHandle RegisterRelationType(
+		const UE::Editor::DataStorage::FRelationRegistrationParams& Params) override;
+	UE_API virtual UE::Editor::DataStorage::RelationTypeHandle FindRelationType(
+		const FName& InName) const override;
+	UE_API virtual bool IsValidRelationType(
+		UE::Editor::DataStorage::RelationTypeHandle Type) const override;
+	UE_API virtual const UE::Editor::DataStorage::FTedsRelationTraits* GetRelationTypeTraits(
+		UE::Editor::DataStorage::RelationTypeHandle Type) const override;
+	UE_API virtual void ListRelationTypes(
+		TFunctionRef<void(UE::Editor::DataStorage::RelationTypeHandle, const FName&)> Callback) const override;
+	UE_API virtual const UScriptStruct* GetRelationSubjectChangedColumn(
+		UE::Editor::DataStorage::RelationTypeHandle Type) const override;
+	UE_API virtual const UScriptStruct* GetRelationObjectChangedColumn(
+		UE::Editor::DataStorage::RelationTypeHandle Type) const override;
+
+	// Relations Interface - Relation Instance Management
+	UE_API virtual UE::Editor::DataStorage::RowHandle CreateRelation(
+		UE::Editor::DataStorage::RelationTypeHandle Type,
+		UE::Editor::DataStorage::RowHandle Subject,
+		UE::Editor::DataStorage::RowHandle Object) override;
+	UE_API virtual void BatchCreateRelations(
+		UE::Editor::DataStorage::RelationTypeHandle Type,
+		TConstArrayView<UE::Editor::DataStorage::RowHandle> Subjects,
+		TConstArrayView<UE::Editor::DataStorage::RowHandle> Objects,
+		TArray<UE::Editor::DataStorage::RowHandle>* OutRelationRows = nullptr) override;
+	UE_API virtual bool DestroyRelation(
+		UE::Editor::DataStorage::RelationTypeHandle Type,
+		UE::Editor::DataStorage::RowHandle Subject,
+		UE::Editor::DataStorage::RowHandle Object) override;
+	UE_API virtual bool HasRelation(
+		UE::Editor::DataStorage::RelationTypeHandle Type,
+		UE::Editor::DataStorage::RowHandle Subject,
+		UE::Editor::DataStorage::RowHandle Object) const override;
+	UE_API virtual bool HasRelationObject(
+		UE::Editor::DataStorage::RelationTypeHandle Type,
+		UE::Editor::DataStorage::RowHandle Subject) const override;
+	UE_API virtual bool HasRelationSubject(
+		UE::Editor::DataStorage::RelationTypeHandle Type,
+		UE::Editor::DataStorage::RowHandle Object) const override;
+
+	// Relations Interface - Relation Queries
+	UE_API virtual void GetRelationObjects(
+		UE::Editor::DataStorage::RelationTypeHandle Type,
+		UE::Editor::DataStorage::RowHandle Subject,
+		TArray<UE::Editor::DataStorage::RowHandle>& OutObjects) const override;
+	UE_API virtual UE::Editor::DataStorage::RowHandle GetRelationObject(
+		UE::Editor::DataStorage::RelationTypeHandle Type,
+		UE::Editor::DataStorage::RowHandle Subject) const override;
+	UE_API virtual void GetRelationSubjects(
+		UE::Editor::DataStorage::RelationTypeHandle Type,
+		UE::Editor::DataStorage::RowHandle Object,
+		TArray<UE::Editor::DataStorage::RowHandle>& OutSubjects) const override;
+	UE_API virtual UE::Editor::DataStorage::RowHandle GetRelationSubject(
+		UE::Editor::DataStorage::RelationTypeHandle Type,
+		UE::Editor::DataStorage::RowHandle Object) const override;
+
+	// Relations Interface - Hierarchical Queries
+	UE_API virtual bool IsDescendantOf(
+		UE::Editor::DataStorage::RelationTypeHandle Type,
+		UE::Editor::DataStorage::RowHandle Descendant,
+		UE::Editor::DataStorage::RowHandle Ancestor,
+		bool bIncludeSelf = false) const override;
+	UE_API virtual UE::Editor::DataStorage::RowHandle GetHierarchyRoot(
+		UE::Editor::DataStorage::RelationTypeHandle Type,
+		UE::Editor::DataStorage::RowHandle Row) const override;
+	UE_API virtual int32 GetHierarchyDepth(
+		UE::Editor::DataStorage::RelationTypeHandle Type,
+		UE::Editor::DataStorage::RowHandle Row) const override;
+	UE_API virtual void GetDescendants(
+		UE::Editor::DataStorage::RelationTypeHandle Type,
+		UE::Editor::DataStorage::RowHandle Ancestor,
+		TArray<UE::Editor::DataStorage::RowHandle>& OutDescendants) const override;
+	UE_API virtual void GetAncestors(
+		UE::Editor::DataStorage::RelationTypeHandle Type,
+		UE::Editor::DataStorage::RowHandle Descendant,
+		TArray<UE::Editor::DataStorage::RowHandle>& OutAncestors) const override;
+	UE_API virtual void TraverseDescendants(
+		UE::Editor::DataStorage::RelationTypeHandle Type,
+		UE::Editor::DataStorage::RowHandle StartRow,
+		FRelationTraversalCallback Callback,
+		ETraversalOrder Order = ETraversalOrder::PreOrder,
+		int32 MaxDepth = TNumericLimits<int32>::Max()) const override;
+
+	// Relations Interface - Callback-Based Enumeration
+	UE_API virtual void ForEachRelationObject(UE::Editor::DataStorage::RelationTypeHandle Type,
+		UE::Editor::DataStorage::RowHandle Subject,
+		TFunctionRef<void(UE::Editor::DataStorage::RowHandle)> Callback) const override;
+	UE_API virtual void ForEachRelationSubject(UE::Editor::DataStorage::RelationTypeHandle Type,
+		UE::Editor::DataStorage::RowHandle Object,
+		TFunctionRef<void(UE::Editor::DataStorage::RowHandle)> Callback) const override;
+	UE_API virtual void TraverseAncestors(UE::Editor::DataStorage::RelationTypeHandle Type,
+		UE::Editor::DataStorage::RowHandle Descendant,
+		TFunctionRef<bool(UE::Editor::DataStorage::RowHandle, int32)> Callback) const override;
+	UE_API virtual int32 ComputeDescendantCount(UE::Editor::DataStorage::RelationTypeHandle Type,
+		UE::Editor::DataStorage::RowHandle Ancestor) const override;
+	UE_API virtual bool HasRelationSubjects(UE::Editor::DataStorage::RelationTypeHandle Type,
+		UE::Editor::DataStorage::RowHandle Object) const override;
+
 	UE_API void RegisterTickGroup(FName GroupName, UE::Editor::DataStorage::EQueryTickPhase Phase, FName BeforeGroup, FName AfterGroup, UE::Editor::DataStorage::EExecutionMode ExecutionMode);
 	UE_API void UnregisterTickGroup(FName GroupName, UE::Editor::DataStorage::EQueryTickPhase Phase);
 
diff --git a/Engine/Plugins/Experimental/EditorDataStorage/Source/TedsCore/Private/TypedElementDatabaseEnvironment.cpp b/Engine/Plugins/Experimental/EditorDataStorage/Source/TedsCore/Private/TypedElementDatabaseEnvironment.cpp
--- a/Engine/Plugins/Experimental/EditorDataStorage/Source/TedsCore/Private/TypedElementDatabaseEnvironment.cpp
+++ b/Engine/Plugins/Experimental/EditorDataStorage/Source/TedsCore/Private/TypedElementDatabaseEnvironment.cpp
@@ -109,6 +109,26 @@ namespace UE::Editor::DataStorage
 		return HierarchyInterface;
 	}
 
+	FTedsRelationAdapter& FEnvironment::GetRelationAdapter()
+	{
+		return RelationAdapter;
+	}
+
+	const FTedsRelationAdapter& FEnvironment::GetRelationAdapter() const
+	{
+		return RelationAdapter;
+	}
+
+	FTedsHierarchicalRelationManager& FEnvironment::GetHierarchicalRelationManager()
+	{
+		return HierarchicalRelationManager;
+	}
+
+	const FTedsHierarchicalRelationManager& FEnvironment::GetHierarchicalRelationManager() const
+	{
+		return HierarchicalRelationManager;
+	}
+
 	FMassProcessingPhaseManager& FEnvironment::GetMassPhaseManager()
 	{
 		return MassPhaseManager;
diff --git a/Engine/Plugins/Experimental/EditorDataStorage/Source/TedsCore/Private/TypedElementDatabaseEnvironment.h b/Engine/Plugins/Experimental/EditorDataStorage/Source/TedsCore/Private/TypedElementDatabaseEnvironment.h
--- a/Engine/Plugins/Experimental/EditorDataStorage/Source/TedsCore/Private/TypedElementDatabaseEnvironment.h
+++ b/Engine/Plugins/Experimental/EditorDataStorage/Source/TedsCore/Private/TypedElementDatabaseEnvironment.h
@@ -12,6 +12,7 @@
 #include "TypedElementDatabaseScratchBuffer.h"
 #include "TypedElementDatabaseIndexTable.h"
 #include "Hierarchy/EditorDataStorageHierarchyImplementation.h"
+#include "Relations/EditorDataStorageRelationsImplementation.h"
 #include "Memento/TypedElementMementoSystem.h"
 #include "Queries/TypedElementExtendedQueryStore.h"
 #include "Scope/TypedElementScopeSystem.h"
@@ -53,6 +54,12 @@ namespace UE::Editor::DataStorage
 		FTedsHierarchyRegistrar& GetHierarchyRegistrar();
 		const FTedsHierarchyRegistrar& GetHierarchyRegistrar() const;
 
+		FTedsRelationAdapter& GetRelationAdapter();
+		const FTedsRelationAdapter& GetRelationAdapter() const;
+
+		FTedsHierarchicalRelationManager& GetHierarchicalRelationManager();
+		const FTedsHierarchicalRelationManager& GetHierarchicalRelationManager() const;
+
 		FMassProcessingPhaseManager& GetMassPhaseManager();
 		const FMassProcessingPhaseManager& GetMassPhaseManager() const;
 
@@ -103,7 +110,9 @@ namespace UE::Editor::DataStorage
 		FScopeSystem ScopeSystem;
 		FValueTagManager ValueTagManager;
 		FTedsHierarchyRegistrar HierarchyInterface;
-		
+		FTedsRelationAdapter RelationAdapter;
+		FTedsHierarchicalRelationManager HierarchicalRelationManager;
+
 		FMutex CommandQueueMutex;
 		TArray<FEnvironmentCommand> CommandQueues[2];
 		TArray<FEnvironmentCommand>* ActiveCommandQueue;
diff --git a/Engine/Plugins/Experimental/EditorDataStorage/Source/TedsCore/Tests/DynamicColumnTests.cpp b/Engine/Plugins/Experimental/EditorDataStorage/Source/TedsCore/Tests/DynamicColumnTests.cpp
--- a/Engine/Plugins/Experimental/EditorDataStorage/Source/TedsCore/Tests/DynamicColumnTests.cpp
+++ b/Engine/Plugins/Experimental/EditorDataStorage/Source/TedsCore/Tests/DynamicColumnTests.cpp
@@ -327,293 +327,4 @@ namespace UE::Editor::DataStorage::Tests
 #undef TEST_IDENTIFIER_1
 #undef TEST_IDENTIFIER_0
 
-namespace UE::Editor::DataStorage::Tests
-{
-	BEGIN_DEFINE_SPEC(TedsHierarchyTestFixture, "Editor.DataStorage.Hierarchy", EAutomationTestFlags::EditorContext | EAutomationTestFlags::EngineFilter)
-		const FName TestTableName = TEXT("TestTable_HierarchyTestTable");
-		const FName HierarchyName = TEXT("TedsHierarchyTestFixture_Hierarchy");
-		FHierarchyHandle HierarchyHandle;
-		ICoreProvider* TedsInterface = nullptr;
-		TableHandle TestTable;
-		TArray<RowHandle> Rows;
-
-		TableHandle RegisterTestTable() const
-		{
-			const TableHandle Table = TedsInterface->FindTable(TestTableName);
-					
-			if (Table != InvalidTableHandle)
-			{
-				return Table;
-			}
-					
-			return TedsInterface->RegisterTable(
-			{
-				FTestColumnB::StaticStruct()
-			},
-			TestTableName);
-		}
-
-		FHierarchyHandle RegisterHierarchy()
-		{
-			HierarchyHandle = TedsInterface->FindHierarchyByName(HierarchyName);
-			if (!TedsInterface->IsValidHierarchyHandle(HierarchyHandle))
-			{
-				UE::Editor::DataStorage::FHierarchyRegistrationParams EntityHierarchyParams{.Name = HierarchyName};
-				HierarchyHandle = TedsInterface->RegisterHierarchy(EntityHierarchyParams);
-			}
-			return HierarchyHandle;
-		}
-
-		RowHandle CreateTestRow(TableHandle InTableHandle)
-		{
-			RowHandle Row = TedsInterface->AddRow(InTableHandle);
-			Rows.Add(Row);
-			return Row;
-		}
-
-		struct FTreeNodes
-		{
-			RowHandle A = InvalidRowHandle;
-			RowHandle B = InvalidRowHandle;
-			RowHandle C = InvalidRowHandle;
-			RowHandle D = InvalidRowHandle;
-			RowHandle E = InvalidRowHandle;
-		};
-
-		int32 ComputeTreeSize(RowHandle Root)
-		{
-			int32 Count = 0;
-			TedsInterface->WalkDepthFirst(HierarchyHandle, Root, [&Count](const ICoreProvider& Context, RowHandle Owner, RowHandle Target)
-				{
-					Count++;
-				});
-			return Count;
-		};
-
-		// Build out this topology:
-		// A
-		// - B
-		FTreeNodes BuildSimpleTopology()
-		{
-			RowHandle A = CreateTestRow(TestTable);
-			RowHandle B = CreateTestRow(TestTable);
-
-			TedsInterface->SetParentRow(HierarchyHandle, B, A);
-
-			return FTreeNodes
-			{
-				.A = A,
-				.B = B,
-			};
-		};
-
-		// Build out this topology:
-		// A
-		// - B
-		// - C
-		//   - D
-		FTreeNodes BuildNestedTopology()
-		{
-			RowHandle A = CreateTestRow(TestTable);
-			RowHandle B = CreateTestRow(TestTable);
-			RowHandle C = CreateTestRow(TestTable);
-			RowHandle D = CreateTestRow(TestTable);
-
-			TedsInterface->SetParentRow(HierarchyHandle, B, A);
-			TedsInterface->SetParentRow(HierarchyHandle, C, A);
-			TedsInterface->SetParentRow(HierarchyHandle, D, C);
-
-			return FTreeNodes
-			{
-				.A = A,
-				.B = B,
-				.C = C,
-				.D = D
-			};
-		};
-		
-	END_DEFINE_SPEC(TedsHierarchyTestFixture)
-
-	void TedsHierarchyTestFixture::Define()
-	{
-		BeforeEach([this]()
-		{
-			TedsInterface = GetMutableDataStorageFeature<ICoreProvider>(StorageFeatureName);
-			TestTable = RegisterTestTable();
-			HierarchyHandle = RegisterHierarchy();
-		});	
-
-		It("Test Simple Topology", [this]()
-			{
-				auto Tree = BuildSimpleTopology();
-
-				TestEqual("", TedsInterface->GetParentRow(HierarchyHandle, Tree.A), InvalidRowHandle);
-				TestEqual("", TedsInterface->GetParentRow(HierarchyHandle, Tree.B), Tree.A);
-
-				TestEqual("", ComputeTreeSize(Tree.A), 2);
-				TestEqual("", ComputeTreeSize(Tree.B), 1);
-			});
-		
-		It("Test Nested Topology", [this]()
-		{
-			auto Tree = BuildNestedTopology();
-
-			TestEqual("", TedsInterface->GetParentRow(HierarchyHandle, Tree.A), InvalidRowHandle);
-			TestEqual("", TedsInterface->GetParentRow(HierarchyHandle, Tree.B), Tree.A);
-			TestEqual("", TedsInterface->GetParentRow(HierarchyHandle, Tree.C), Tree.A);
-			TestEqual("", TedsInterface->GetParentRow(HierarchyHandle, Tree.D), Tree.C);
-
-			TestEqual("", ComputeTreeSize(Tree.A), 4);
-			TestEqual("", ComputeTreeSize(Tree.B), 1);
-			TestEqual("", ComputeTreeSize(Tree.C), 2);
-			TestEqual("", ComputeTreeSize(Tree.D), 1);
-		});
-			
-		It("Unparent childless leaf from single-child parent", [this]()
-		{
-				auto Tree = BuildSimpleTopology();
-				TedsInterface->SetParentRow(HierarchyHandle, Tree.B, InvalidRowHandle);
-
-				TestEqual("", TedsInterface->GetParentRow(HierarchyHandle, Tree.A), InvalidRowHandle);
-				TestEqual("", TedsInterface->GetParentRow(HierarchyHandle, Tree.B), InvalidRowHandle);
-
-				TestEqual("", ComputeTreeSize(Tree.A), 1);
-				TestEqual("", ComputeTreeSize(Tree.B), 1);
-		});
-
-		It("Delete childless leaf from single-child parent", [this]()
-		{
-				auto Tree = BuildSimpleTopology();
-				TedsInterface->RemoveRow(Tree.B);
-
-				TestEqual("", TedsInterface->GetParentRow(HierarchyHandle, Tree.A), InvalidRowHandle);
-				TestEqual("", TedsInterface->GetParentRow(HierarchyHandle, Tree.B), InvalidRowHandle);
-
-				TestEqual("", ComputeTreeSize(Tree.A), 1);
-				TestEqual("", ComputeTreeSize(Tree.B), 0);
-		});
-
-		It("Delete single-child parent of childless leaf node ", [this]()
-		{
-				auto Tree = BuildSimpleTopology();
-				TedsInterface->RemoveRow(Tree.A);
-
-				TestEqual("", TedsInterface->GetParentRow(HierarchyHandle, Tree.A), InvalidRowHandle);
-				TestEqual("", TedsInterface->GetParentRow(HierarchyHandle, Tree.B), InvalidRowHandle);
-
-				TestEqual("", ComputeTreeSize(Tree.A), 0);
-				TestEqual("", ComputeTreeSize(Tree.B), 1);
-		});
-			
-		It("Unparent childless leaf from multi-child parent", [this]()
-		{
-				auto Tree = BuildNestedTopology();
-				TedsInterface->SetParentRow(HierarchyHandle, Tree.B, InvalidRowHandle);
-
-				TestEqual("", TedsInterface->GetParentRow(HierarchyHandle, Tree.A), InvalidRowHandle);
-				TestEqual("", TedsInterface->GetParentRow(HierarchyHandle, Tree.B), InvalidRowHandle);
-				TestEqual("", TedsInterface->GetParentRow(HierarchyHandle, Tree.C), Tree.A);
-				TestEqual("", TedsInterface->GetParentRow(HierarchyHandle, Tree.D), Tree.C);
-
-				TestEqual("", ComputeTreeSize(Tree.A), 3);
-				TestEqual("", ComputeTreeSize(Tree.B), 1);
-				TestEqual("", ComputeTreeSize(Tree.C), 2);
-				TestEqual("", ComputeTreeSize(Tree.D), 1);
-		});
-
-		It("Unparent middle parent from middle parent", [this]()
-		{
-				auto Tree = BuildNestedTopology();
-				TedsInterface->SetParentRow(HierarchyHandle, Tree.C, InvalidRowHandle);
-
-				TestEqual("", TedsInterface->GetParentRow(HierarchyHandle, Tree.A), InvalidRowHandle);
-				TestEqual("", TedsInterface->GetParentRow(HierarchyHandle, Tree.B), Tree.A);
-				TestEqual("", TedsInterface->GetParentRow(HierarchyHandle, Tree.C), InvalidRowHandle);
-				TestEqual("", TedsInterface->GetParentRow(HierarchyHandle, Tree.D), Tree.C);
-
-				TestEqual("", ComputeTreeSize(Tree.A), 2);
-				TestEqual("", ComputeTreeSize(Tree.B), 1);
-				TestEqual("", ComputeTreeSize(Tree.C), 2);
-				TestEqual("", ComputeTreeSize(Tree.D), 1);
-		});
-
-		It("Unparent childless leaf from middle parent", [this]()
-		{
-				auto Tree = BuildNestedTopology();
-				TedsInterface->SetParentRow(HierarchyHandle, Tree.D, InvalidRowHandle);
-
-				TestEqual("", TedsInterface->GetParentRow(HierarchyHandle, Tree.A), InvalidRowHandle);
-				TestEqual("", TedsInterface->GetParentRow(HierarchyHandle, Tree.B), Tree.A);
-				TestEqual("", TedsInterface->GetParentRow(HierarchyHandle, Tree.C), Tree.A);
-				TestEqual("", TedsInterface->GetParentRow(HierarchyHandle, Tree.D), InvalidRowHandle);
-
-				TestEqual("", ComputeTreeSize(Tree.A), 3);
-				TestEqual("", ComputeTreeSize(Tree.B), 1);
-				TestEqual("", ComputeTreeSize(Tree.C), 1);
-				TestEqual("", ComputeTreeSize(Tree.D), 1);
-		});
-
-		It("Reparent childless leaf to childless leaf", [this]()
-		{
-				auto Tree = BuildNestedTopology();
-				TedsInterface->SetParentRow(HierarchyHandle, Tree.B, Tree.D);
-
-				TestEqual("", TedsInterface->GetParentRow(HierarchyHandle, Tree.A), InvalidRowHandle);
-				TestEqual("", TedsInterface->GetParentRow(HierarchyHandle, Tree.B), Tree.D);
-				TestEqual("", TedsInterface->GetParentRow(HierarchyHandle, Tree.C), Tree.A);
-				TestEqual("", TedsInterface->GetParentRow(HierarchyHandle, Tree.D), Tree.C);
-
-				TestEqual("", ComputeTreeSize(Tree.A), 4);
-				TestEqual("", ComputeTreeSize(Tree.B), 1);
-				TestEqual("", ComputeTreeSize(Tree.C), 3);
-				TestEqual("", ComputeTreeSize(Tree.D), 2);
-		});
-
-		It("Delete childless leaf of multi-child root", [this]()
-		{
-				auto Tree = BuildNestedTopology();
-				TedsInterface->RemoveRow(Tree.B);
-
-				TestEqual("", TedsInterface->GetParentRow(HierarchyHandle, Tree.A), InvalidRowHandle);
-				TestEqual("", TedsInterface->GetParentRow(HierarchyHandle, Tree.B), InvalidRowHandle);
-				TestEqual("", TedsInterface->GetParentRow(HierarchyHandle, Tree.C), Tree.A);
-				TestEqual("", TedsInterface->GetParentRow(HierarchyHandle, Tree.D), Tree.C);
-
-				TestEqual("ComputeTreeSize(Tree.A)", ComputeTreeSize(Tree.A), 3);
-				TestEqual("ComputeTreeSize(Tree.B)", ComputeTreeSize(Tree.B), 0);
-				TestEqual("ComputeTreeSize(Tree.C)", ComputeTreeSize(Tree.C), 2);
-				TestEqual("ComputeTreeSize(Tree.D)", ComputeTreeSize(Tree.D), 1);
-		});
-
-		It("Delete childless leaf of middle parent", [this]()
-		{
-				auto Tree = BuildNestedTopology();
-				TedsInterface->RemoveRow(Tree.D);
-
-				TestEqual("", TedsInterface->GetParentRow(HierarchyHandle, Tree.A), InvalidRowHandle);
-				TestEqual("", TedsInterface->GetParentRow(HierarchyHandle, Tree.B), Tree.A);
-				TestEqual("", TedsInterface->GetParentRow(HierarchyHandle, Tree.C), Tree.A);
-				TestEqual("", TedsInterface->GetParentRow(HierarchyHandle, Tree.D), InvalidRowHandle);
-
-				TestEqual("", ComputeTreeSize(Tree.A), 3);
-				TestEqual("", ComputeTreeSize(Tree.B), 1);
-				TestEqual("", ComputeTreeSize(Tree.C), 1);
-				TestEqual("", ComputeTreeSize(Tree.D), 0);
-		});
-
-		AfterEach([this]()
-		{
-			for (RowHandle Row : Rows)
-			{
-				TedsInterface->RemoveRow(Row);
-			}
-			Rows.Empty(Rows.Num());
-
-			HierarchyHandle = FHierarchyHandle();
-
-			TedsInterface = nullptr;
-		});
-	}
-}
-
 #endif // WITH_TESTS
diff --git a/Engine/Plugins/Experimental/EditorDataStorageFeatures/Source/TedsDebugger/Private/QueryEditor/TedsQueryEditorModel.cpp b/Engine/Plugins/Experimental/EditorDataStorageFeatures/Source/TedsDebugger/Private/QueryEditor/TedsQueryEditorModel.cpp
--- a/Engine/Plugins/Experimental/EditorDataStorageFeatures/Source/TedsDebugger/Private/QueryEditor/TedsQueryEditorModel.cpp
+++ b/Engine/Plugins/Experimental/EditorDataStorageFeatures/Source/TedsDebugger/Private/QueryEditor/TedsQueryEditorModel.cpp
@@ -425,9 +425,27 @@ namespace UE::Editor::DataStorage::Debug::QueryEditor
 	void FTedsQueryEditorModel::SetHierarchy(const FName& InHierarchy)
 	{
 		CurrentHierarchy = InHierarchy;
+		CurrentRelationType.Reset();
 		HierarchyChangedDelegate.Broadcast();
 	}
 
+	void FTedsQueryEditorModel::SetRelationType(RelationTypeHandle InRelationType)
+	{
+		CurrentRelationType = InRelationType;
+		CurrentHierarchy = NAME_None;
+		HierarchyChangedDelegate.Broadcast();
+	}
+
+	bool FTedsQueryEditorModel::HasRelationType() const
+	{
+		return CurrentRelationType.IsSet();
+	}
+
+	RelationTypeHandle FTedsQueryEditorModel::GetRelationType() const
+	{
+		return CurrentRelationType.Get(InvalidRelationTypeHandle);
+	}
+
 	FName FTedsQueryEditorModel::GetHierarchyName() const
 	{
 		return CurrentHierarchy;
diff --git a/Engine/Plugins/Experimental/EditorDataStorageFeatures/Source/TedsDebugger/Private/QueryEditor/TedsQueryEditorModel.h b/Engine/Plugins/Experimental/EditorDataStorageFeatures/Source/TedsDebugger/Private/QueryEditor/TedsQueryEditorModel.h
--- a/Engine/Plugins/Experimental/EditorDataStorageFeatures/Source/TedsDebugger/Private/QueryEditor/TedsQueryEditorModel.h
+++ b/Engine/Plugins/Experimental/EditorDataStorageFeatures/Source/TedsDebugger/Private/QueryEditor/TedsQueryEditorModel.h
@@ -101,6 +101,11 @@ namespace UE::Editor::DataStorage::Debug::QueryEditor
 		void SetHierarchy(const FName& InHierarchy);
 		FName GetHierarchyName() const;
 		FHierarchyHandle GetHierarchy() const;
+
+		// Relation-type hierarchy selection (mutually exclusive with SetHierarchy)
+		void SetRelationType(RelationTypeHandle InRelationType);
+		bool HasRelationType() const;
+		RelationTypeHandle GetRelationType() const;
 	
 	private:
 		
@@ -127,8 +132,10 @@ namespace UE::Editor::DataStorage::Debug::QueryEditor
 		uint64 GeneratedVersion = 0;
 		QueryHandle QueryHandle;
 
-		// The current hierarchy being viewed (only applicable for hierarchy based widgets)
+		// The current hierarchy being viewed (only applicable for hierarchy based widgets).
+		// CurrentHierarchy and CurrentRelationType are mutually exclusive — only one is active at a time.
 		FName CurrentHierarchy;
+		TOptional<RelationTypeHandle> CurrentRelationType;
 		FTedsQueryEditorModel_ModelChanged HierarchyChangedDelegate;
 	};
 } // namespace UE::Editor::DataStorage::Debug::QueryEditor
diff --git a/Engine/Plugins/Experimental/EditorDataStorageFeatures/Source/TedsDebugger/Private/Widgets/QueryEditor/TedsHierarchySelectionComboWidget.cpp b/Engine/Plugins/Experimental/EditorDataStorageFeatures/Source/TedsDebugger/Private/Widgets/QueryEditor/TedsHierarchySelectionComboWidget.cpp
--- a/Engine/Plugins/Experimental/EditorDataStorageFeatures/Source/TedsDebugger/Private/Widgets/QueryEditor/TedsHierarchySelectionComboWidget.cpp
+++ b/Engine/Plugins/Experimental/EditorDataStorageFeatures/Source/TedsDebugger/Private/Widgets/QueryEditor/TedsHierarchySelectionComboWidget.cpp
@@ -11,45 +11,76 @@ namespace UE::Editor::DataStorage::Debug::QueryEditor
 		Model = &InModel;
 
 		GenerateHierarchyList();
-		
+
 		ChildSlot
 		[
 			SAssignNew(SearchableComboBox, SSearchableComboBox)
-            	.OptionsSource(&Hierarchies)
-            	.OnSelectionChanged_Lambda([this](TSharedPtr<FString> SelectedItem, ESelectInfo::Type SelectInfo)
-            	{
-            		if (SelectedItem)
-            		{
-            			Model->SetHierarchy(FName(*SelectedItem));
-            		}
-            	})
-            	.OnGenerateWidget_Lambda([](TSharedPtr<FString> InItem)
-            	{
-            		return SNew(STextBlock)
-            		.Text(FText::FromString(InItem.IsValid() ? *InItem : FString()));
-            	})
+				.OptionsSource(&Hierarchies)
+				.OnSelectionChanged_Lambda([this](TSharedPtr<FString> SelectedItem, ESelectInfo::Type SelectInfo)
+				{
+					if (!SelectedItem)
+					{
+						return;
+					}
+					if (const RelationTypeHandle* Handle = RelationHandleByDisplayName.Find(*SelectedItem))
+					{
+						Model->SetRelationType(*Handle);
+					}
+					else
+					{
+						Model->SetHierarchy(FName(*SelectedItem));
+					}
+				})
+				.OnGenerateWidget_Lambda([](TSharedPtr<FString> InItem)
+				{
+					return SNew(STextBlock)
+						.Text(FText::FromString(InItem.IsValid() ? *InItem : FString()));
+				})
 				.OnComboBoxOpening_Lambda([this]()
 				{
 					GenerateHierarchyList();
 				})
-            	.Content()
-            	[
-            		SNew(STextBlock)
+				.Content()
+				[
+					SNew(STextBlock)
 					.Text_Lambda([this]()
 					{
+						if (Model->HasRelationType())
+						{
+							const RelationTypeHandle SelectedHandle = Model->GetRelationType();
+							for (const TTuple<FString, RelationTypeHandle>& Pair : RelationHandleByDisplayName)
+							{
+								if (Pair.Value == SelectedHandle)
+								{
+									return FText::FromString(Pair.Key);
+								}
+							}
+						}
 						return FText::FromName(Model->GetHierarchyName());
 					})
-            	]
+				]
 		];
 	}
 
 	void SHierarchyComboWidget::GenerateHierarchyList()
 	{
 		Hierarchies.Empty();
-		
+		RelationHandleByDisplayName.Empty();
+
 		Model->GetTedsInterface().ListHierarchyNames([this](const FName& HierarchyName)
 		{
 			Hierarchies.Add(MakeShared<FString>(HierarchyName.ToString()));
 		});
+
+		Model->GetTedsInterface().ListRelationTypes([this](RelationTypeHandle Handle, const FName& Name)
+		{
+			const FTedsRelationTraits* Traits = Model->GetTedsInterface().GetRelationTypeTraits(Handle);
+			if (Traits && Traits->HierarchyMode != EHierarchyMode::Disabled)
+			{
+				FString DisplayName = Name.ToString() + TEXT(" (relation)");
+				RelationHandleByDisplayName.Add(DisplayName, Handle);
+				Hierarchies.Add(MakeShared<FString>(MoveTemp(DisplayName)));
+			}
+		});
 	}
 }
diff --git a/Engine/Plugins/Experimental/EditorDataStorageFeatures/Source/TedsDebugger/Private/Widgets/QueryEditor/TedsHierarchySelectionComboWidget.h b/Engine/Plugins/Experimental/EditorDataStorageFeatures/Source/TedsDebugger/Private/Widgets/QueryEditor/TedsHierarchySelectionComboWidget.h
--- a/Engine/Plugins/Experimental/EditorDataStorageFeatures/Source/TedsDebugger/Private/Widgets/QueryEditor/TedsHierarchySelectionComboWidget.h
+++ b/Engine/Plugins/Experimental/EditorDataStorageFeatures/Source/TedsDebugger/Private/Widgets/QueryEditor/TedsHierarchySelectionComboWidget.h
@@ -20,6 +20,9 @@ namespace UE::Editor::DataStorage::Debug::QueryEditor
 		
 	protected:
 		TArray<TSharedPtr<FString>> Hierarchies;
+		// Maps display name (e.g. "InheritsFrom (relation)") → RelationTypeHandle for entries added from ListRelationTypes.
+		// Entries not present in this map are named hierarchies resolved via FindHierarchyByName.
+		TMap<FString, RelationTypeHandle> RelationHandleByDisplayName;
 		FTedsQueryEditorModel* Model = nullptr;
 		TSharedPtr<SSearchableComboBox> SearchableComboBox;
 	};
diff --git a/Engine/Plugins/Experimental/EditorDataStorageFeatures/Source/TedsDebugger/Private/Widgets/QueryEditor/TedsQueryEditorResultsView.cpp b/Engine/Plugins/Experimental/EditorDataStorageFeatures/Source/TedsDebugger/Private/Widgets/QueryEditor/TedsQueryEditorResultsView.cpp
--- a/Engine/Plugins/Experimental/EditorDataStorageFeatures/Source/TedsDebugger/Private/Widgets/QueryEditor/TedsQueryEditorResultsView.cpp
+++ b/Engine/Plugins/Experimental/EditorDataStorageFeatures/Source/TedsDebugger/Private/Widgets/QueryEditor/TedsQueryEditorResultsView.cpp
@@ -21,7 +21,7 @@
 #include "Widgets/Layout/SBorder.h"
 #include "Widgets/STedsHierarchyViewer.h"
 #include "Widgets/STedsTableViewer.h"
-#include "Widgets/SRowDetails.h"
+#include "Widgets/SRowDetailsNavigator.h"
 #include "Widgets/Text/STextBlock.h"
 #include "Widgets/STedsSearchBox.h"
 
@@ -90,7 +90,7 @@ namespace UE::Editor::DataStorage::Debug::QueryEditor
 				+SSplitter::Slot()
 				.Value(0.5f)
 				[
-					SAssignNew(RowDetailsWidget, SRowDetails)
+					SAssignNew(RowDetailsNavigator, SRowDetailsNavigator)
 				]
 			]
 			+SVerticalBox::Slot()
@@ -197,7 +197,10 @@ namespace UE::Editor::DataStorage::Debug::QueryEditor
 
 	void SResultsView::OnHierarchyChanged()
 	{
-		HierarchyHandle = Model->GetHierarchy();
+		if (!Model->HasRelationType())
+		{
+			HierarchyHandle = Model->GetHierarchy();
+		}
 
 		if (ResultsViewMode == ETableViewMode::Tree)
 		{
@@ -233,15 +236,15 @@ namespace UE::Editor::DataStorage::Debug::QueryEditor
 					.OnSelectionChanged(STedsTableViewer::FOnSelectionChanged::CreateLambda(
 						[this](RowHandle SelectedRow, ESelectInfo::Type SelectInfo)
 							{
-								if(RowDetailsWidget)
+								if(RowDetailsNavigator)
 								{
 									if(SelectedRow != InvalidRowHandle)
 									{
-										RowDetailsWidget->SetRow(SelectedRow);
+										RowDetailsNavigator->SetRow(SelectedRow);
 									}
 									else
 									{
-										RowDetailsWidget->ClearRow();
+										RowDetailsNavigator->ClearRow();
 									}
 								}
 							}));
@@ -249,7 +252,14 @@ namespace UE::Editor::DataStorage::Debug::QueryEditor
 		
 		case ETableViewMode::Tree:
 
-			HierarchyData = MakeShared<FHierarchyViewerData>(Model->GetTedsInterface(), HierarchyHandle);
+			if (Model->HasRelationType())
+			{
+				HierarchyData = MakeShared<FRelationTypeHierarchyViewerData>(Model->GetRelationType());
+			}
+			else
+			{
+				HierarchyData = MakeShared<FHierarchyViewerData>(Model->GetTedsInterface(), HierarchyHandle);
+			}
 
 			TableViewer =
 				SNew(SHierarchyViewer, HierarchyData)
@@ -259,15 +269,15 @@ namespace UE::Editor::DataStorage::Debug::QueryEditor
 					.OnSelectionChanged(STedsTableViewer::FOnSelectionChanged::CreateLambda(
 						[this](RowHandle SelectedRow, ESelectInfo::Type SelectInfo)
 							{
-								if(RowDetailsWidget)
+								if(RowDetailsNavigator)
 								{
 									if(SelectedRow != InvalidRowHandle)
 									{
-										RowDetailsWidget->SetRow(SelectedRow);
+										RowDetailsNavigator->SetRow(SelectedRow);
 									}
 									else
 									{
-										RowDetailsWidget->ClearRow();
+										RowDetailsNavigator->ClearRow();
 									}
 								}
 							}));
@@ -281,9 +291,9 @@ namespace UE::Editor::DataStorage::Debug::QueryEditor
 			return SNullWidget::NullWidget;
 		}
 
-		if (RowDetailsWidget)
+		if (RowDetailsNavigator)
 		{
-			RowDetailsWidget->ClearRow();
+			RowDetailsNavigator->ClearRow();
 		}
 
 		bModelDirty = true; // We need to set the columns again
diff --git a/Engine/Plugins/Experimental/EditorDataStorageFeatures/Source/TedsDebugger/Private/Widgets/QueryEditor/TedsQueryEditorResultsView.h b/Engine/Plugins/Experimental/EditorDataStorageFeatures/Source/TedsDebugger/Private/Widgets/QueryEditor/TedsQueryEditorResultsView.h
--- a/Engine/Plugins/Experimental/EditorDataStorageFeatures/Source/TedsDebugger/Private/Widgets/QueryEditor/TedsQueryEditorResultsView.h
+++ b/Engine/Plugins/Experimental/EditorDataStorageFeatures/Source/TedsDebugger/Private/Widgets/QueryEditor/TedsQueryEditorResultsView.h
@@ -18,7 +18,7 @@ namespace UE::Editor::DataStorage
 	class ITableViewer;
 	class STedsTableViewer;
 	class FTedsTableViewerColumn;
-	class SRowDetails;
+	class SRowDetailsNavigator;
 	class IUiProvider;
 
 	namespace QueryStack
@@ -66,16 +66,17 @@ namespace UE::Editor::DataStorage
 			TSharedPtr<QueryStack::FQueryNode> QueryNode;
 			TSharedPtr<QueryStack::IRowNode> RowQueryStack;
 
-			// Hierarchy Data used by SHierarchyView to resolve parenting
-			TSharedPtr<FHierarchyViewerData> HierarchyData;
-			// The current hierarchy the widget is viewing
+			// Hierarchy Data used by SHierarchyView to resolve parenting.
+			// May be FHierarchyViewerData (named hierarchy) or FRelationTypeHierarchyViewerData (raw hierarchical relation).
+			TSharedPtr<IHierarchyViewerDataInterface> HierarchyData;
+			// The current named hierarchy the widget is viewing (only valid when a named hierarchy is selected)
 			FHierarchyHandle HierarchyHandle;
 			
 			// Custom column for the table viewer to display row handles
 			TSharedPtr<FTedsTableViewerColumn> RowHandleColumn;
 
-			// Widget that displays details of a row
-			TSharedPtr<SRowDetails> RowDetailsWidget;
+			// Breadcrumb-navigable stack of row detail panels
+			TSharedPtr<SRowDetailsNavigator> RowDetailsNavigator;
 			
 			IUiProvider* UiProvider = nullptr;
 
diff --git a/Engine/Plugins/Experimental/EditorDataStorageFeatures/Source/TedsDebugger/Private/Widgets/STedsDiscoveryDebugger.h b/Engine/Plugins/Experimental/EditorDataStorageFeatures/Source/TedsDebugger/Private/Widgets/STedsDiscoveryDebugger.h
--- a/Engine/Plugins/Experimental/EditorDataStorageFeatures/Source/TedsDebugger/Private/Widgets/STedsDiscoveryDebugger.h
+++ b/Engine/Plugins/Experimental/EditorDataStorageFeatures/Source/TedsDebugger/Private/Widgets/STedsDiscoveryDebugger.h
@@ -6,13 +6,15 @@
 #include "TedsColumnsSearchNode.h"
 #include "Widgets/SCompoundWidget.h"
 
-class SRowDetails;
+namespace UE::Editor::DataStorage { class SRowDetails; }
 class SSearchBox;
 class STedsTableViewer;
 class FTedsTableViewerColumn;
 
 namespace UE::Editor::DataStorage::Debug
 {
+	using UE::Editor::DataStorage::SRowDetails;
+
 	enum class ESearchMode : uint8
 	{
 		// Search for a specific RowHandle for specific debugging
diff --git a/Engine/Plugins/Experimental/EditorDataStorageFeatures/Source/TedsTableViewer/Private/HierarchyViewerIntefaces.cpp b/Engine/Plugins/Experimental/EditorDataStorageFeatures/Source/TedsTableViewer/Private/HierarchyViewerIntefaces.cpp
--- a/Engine/Plugins/Experimental/EditorDataStorageFeatures/Source/TedsTableViewer/Private/HierarchyViewerIntefaces.cpp
+++ b/Engine/Plugins/Experimental/EditorDataStorageFeatures/Source/TedsTableViewer/Private/HierarchyViewerIntefaces.cpp
@@ -133,4 +133,27 @@ namespace UE::Editor::DataStorage
 			}
 		}
 	}
+
+	FRelationTypeHierarchyViewerData::FRelationTypeHierarchyViewerData(RelationTypeHandle InRelationType)
+		: RelationType(InRelationType)
+	{
+	}
+
+	RowHandle FRelationTypeHierarchyViewerData::GetParent(const ICoreProvider& Storage, RowHandle InRow) const
+	{
+		return Storage.GetRelationObject(RelationType, InRow);
+	}
+
+	void FRelationTypeHierarchyViewerData::WalkDepthFirst(const ICoreProvider& Storage, RowHandle InRow,
+		ICoreProvider::FHierarchyIterationCallback VisitFn, ICoreProvider::ETraversalOrder TraversalOrder) const
+	{
+		// TraverseDescendants callback: (Current, Parent, Depth)
+		// FHierarchyIterationCallback:  (Context, Owner, Target) where Owner=Parent, Target=Current
+		Storage.TraverseDescendants(RelationType, InRow,
+			[&Storage, &VisitFn](RowHandle Current, RowHandle Parent, int32 /*Depth*/)
+			{
+				VisitFn(Storage, Parent, Current);
+			},
+			TraversalOrder);
+	}
 }
diff --git a/Engine/Plugins/Experimental/EditorDataStorageFeatures/Source/TedsTableViewer/Private/Widgets/SRowDetails.cpp b/Engine/Plugins/Experimental/EditorDataStorageFeatures/Source/TedsTableViewer/Private/Widgets/SRowDetails.cpp
--- a/Engine/Plugins/Experimental/EditorDataStorageFeatures/Source/TedsTableViewer/Private/Widgets/SRowDetails.cpp
+++ b/Engine/Plugins/Experimental/EditorDataStorageFeatures/Source/TedsTableViewer/Private/Widgets/SRowDetails.cpp
@@ -8,6 +8,7 @@
 #include "Elements/Columns/TypedElementTypeInfoColumns.h"
 #include "Elements/Interfaces/TypedElementDataStorageInterface.h"
 #include "Elements/Interfaces/TypedElementDataStorageUiInterface.h"
+#include "Widgets/Input/SButton.h"
 #include "Widgets/Text/STextBlock.h"
 #include "Widgets/Views/SListView.h"
 
@@ -40,6 +41,7 @@ namespace UE::Editor::DataStorage
 		check(DataStorage);
 		check(DataStorageUi);
 
+		OnRelatedRowSelectedDelegate = InArgs._OnRelatedRowSelected;
 		RequestedTedsColumns = InArgs._Columns;
 		NameMetaData = InArgs._NameMetaData;
 		DataMetaData = InArgs._DataMetaData;
@@ -165,6 +167,26 @@ namespace UE::Editor::DataStorage
 			}
 		}
 
+		// Add relation items: for each relation type, list this row's related rows as navigable entries.
+		DataStorage->ListRelationTypes([this, Row](RelationTypeHandle Type, const FName& TypeName)
+		{
+			TArray<RowHandle> Objects;
+			DataStorage->GetRelationObjects(Type, Row, Objects);
+			for (RowHandle Object : Objects)
+			{
+				FName ItemName = FName(*FString::Printf(TEXT("%s (Object)"), *TypeName.ToString()));
+				Items.Add(MakeShared<FRowDetailsItem>(ItemName, Row, Object));
+			}
+
+			TArray<RowHandle> Subjects;
+			DataStorage->GetRelationSubjects(Type, Row, Subjects);
+			for (RowHandle Subject : Subjects)
+			{
+				FName ItemName = FName(*FString::Printf(TEXT("%s (Subject)"), *TypeName.ToString()));
+				Items.Add(MakeShared<FRowDetailsItem>(ItemName, Row, Subject));
+			}
+		});
+
 		ListView->RequestListRefresh();
 	}
 
@@ -219,6 +241,27 @@ namespace UE::Editor::DataStorage
 			return SNullWidget::NullWidget;
 		}
 
+		// Relation items have their own rendering -- no widget row needed.
+		if (Item->IsRelationItem())
+		{
+			if (ColumnName == SRowDetails::NameColumn)
+			{
+				return SNew(STextBlock).Text(FText::FromName(Item->Name));
+			}
+			else if (ColumnName == SRowDetails::DataColumn)
+			{
+				const RowHandle RelatedRow = Item->RelatedRow;
+				return SNew(SButton)
+					.Text(FText::FromString(FString::Printf(TEXT("Row %llu"), RelatedRow)))
+					.OnClicked_Lambda([RowDetails, RelatedRow]()
+					{
+						RowDetails->NotifyRelatedRowSelected(RelatedRow);
+						return FReply::Handled();
+					});
+			}
+			return SNullWidget::NullWidget;
+		}
+
 		ICoreProvider* DataStorage = RowDetails->GetDataStorage();
 		if (!DataStorage->IsRowAvailable(Item->WidgetRow))
 		{
@@ -273,6 +316,11 @@ namespace UE::Editor::DataStorage
 		}
 	}
 
+	void SRowDetails::NotifyRelatedRowSelected(RowHandle Row)
+	{
+		OnRelatedRowSelectedDelegate.ExecuteIfBound(Row);
+	}
+
 	// FRowDetailsItem
 	FRowDetailsItem::FRowDetailsItem(const FName& InItemName, const TArray<TWeakObjectPtr<const UScriptStruct>>& InColumnTypes,
 		IUiProvider::FWidgetConstructorPtr InWidgetConstructor, RowHandle InRow)
@@ -281,7 +329,13 @@ namespace UE::Editor::DataStorage
 		, WidgetConstructor(MoveTemp(InWidgetConstructor))
 		, Row(InRow)
 	{
-	
+	}
+
+	FRowDetailsItem::FRowDetailsItem(const FName& InRelationName, RowHandle InRow, RowHandle InRelatedRow)
+		: Name(InRelationName)
+		, Row(InRow)
+		, RelatedRow(InRelatedRow)
+	{
 	}
 	
 } // namespace UE::Editor::DataStorage
diff --git a/Engine/Plugins/Experimental/EditorDataStorageFeatures/Source/TedsTableViewer/Public/HierarchyViewerIntefaces.h b/Engine/Plugins/Experimental/EditorDataStorageFeatures/Source/TedsTableViewer/Public/HierarchyViewerIntefaces.h
--- a/Engine/Plugins/Experimental/EditorDataStorageFeatures/Source/TedsTableViewer/Public/HierarchyViewerIntefaces.h
+++ b/Engine/Plugins/Experimental/EditorDataStorageFeatures/Source/TedsTableViewer/Public/HierarchyViewerIntefaces.h
@@ -64,13 +64,28 @@ namespace UE::Editor::DataStorage
 		// Does a DFS of each hierarchy one-by-one. If a child is encountered multiple times in different hierarchies, it is only reported in the first instance
 		UE_API virtual void WalkDepthFirst(const ICoreProvider& Storage, RowHandle InRow, ICoreProvider::FHierarchyIterationCallback VisitFn,
 			ICoreProvider::ETraversalOrder TraversalOrder = ICoreProvider::ETraversalOrder::PreOrder) const override;
-		
+
 		// Iterate over each contained hierarchy data. The callback can return false to terminate iteration early
 		UE_API void ForEachHierarchyData(TFunctionRef<bool(const TSharedRef<FHierarchyViewerData>&)> VisitFn) const;
-		
+
 	protected:
 		TArray<TSharedPtr<FHierarchyViewerData>> Hierarchies;
 	};
+
+	// Data interface for a hierarchical relation registered directly via RegisterRelationType (HierarchyMode != Disabled).
+	// These relations are not registered in FTedsHierarchyRegistrar and have no FHierarchyHandle, but still support
+	// parent/child traversal via the RelationTypeHandle API on ICoreProvider.
+	class FRelationTypeHierarchyViewerData : public IHierarchyViewerDataInterface
+	{
+	public:
+		UE_API explicit FRelationTypeHierarchyViewerData(RelationTypeHandle InRelationType);
+		UE_API virtual RowHandle GetParent(const ICoreProvider& Storage, RowHandle InRow) const override;
+		UE_API virtual void WalkDepthFirst(const ICoreProvider& Storage, RowHandle InRow, ICoreProvider::FHierarchyIterationCallback VisitFn,
+			ICoreProvider::ETraversalOrder TraversalOrder = ICoreProvider::ETraversalOrder::PreOrder) const override;
+
+	private:
+		RelationTypeHandle RelationType;
+	};
 }
 
 #undef UE_API
\ No newline at end of file
diff --git a/Engine/Plugins/Experimental/EditorDataStorageFeatures/Source/TedsTableViewer/Public/Widgets/SRowDetails.h b/Engine/Plugins/Experimental/EditorDataStorageFeatures/Source/TedsTableViewer/Public/Widgets/SRowDetails.h
--- a/Engine/Plugins/Experimental/EditorDataStorageFeatures/Source/TedsTableViewer/Public/Widgets/SRowDetails.h
+++ b/Engine/Plugins/Experimental/EditorDataStorageFeatures/Source/TedsTableViewer/Public/Widgets/SRowDetails.h
@@ -16,23 +16,34 @@ namespace UE::Editor::DataStorage
 	class ICoreProvider;
 	class IUiProvider;
 
-	// A row in the SRowDetails widget that represents a column on the TEDS row we are viewing
+	DECLARE_DELEGATE_OneParam(FOnRelatedRowSelected, RowHandle);
+
+	// A row in the SRowDetails widget that represents either a column or a relation on the TEDS row we are viewing
 	struct FRowDetailsItem
 	{
 		// The logical name of this item (composited from the matched columns or grabbed from the widget constructor)
 		FName Name;
 
-		// The columns this row is displaying data for
+		// The columns this row is displaying data for (empty for relation items)
 		TArray<TWeakObjectPtr<const UScriptStruct>> ColumnTypes;
 
-		// Widget for the column
+		// Widget for the column (null for relation items)
 		IUiProvider::FWidgetConstructorPtr WidgetConstructor;
-		
+
 		RowHandle Row = InvalidRowHandle;
 		RowHandle WidgetRow = InvalidRowHandle;
 
+		// For relation items: the related row to navigate to (InvalidRowHandle for column items)
+		RowHandle RelatedRow = InvalidRowHandle;
+
+		bool IsRelationItem() const { return RelatedRow != InvalidRowHandle; }
+
+		// Constructor for column items
 		FRowDetailsItem(const FName& InItemName, const TArray<TWeakObjectPtr<const UScriptStruct>>& InColumnTypes, IUiProvider::FWidgetConstructorPtr InWidgetConstructor,
 			RowHandle InRow);
+
+		// Constructor for relation items
+		FRowDetailsItem(const FName& InRelationName, RowHandle InRow, RowHandle InRelatedRow);
 	};
 	
 	using RowDetailsItemPtr = TSharedPtr<FRowDetailsItem>;
@@ -76,8 +87,11 @@ namespace UE::Editor::DataStorage
 			SLATE_ARGUMENT(TArray<TSharedPtr<FColumnMetaData>>, ColumnMetaData)
 			// The padding inside each Row
 			SLATE_ATTRIBUTE(FMargin, RowPadding)
+
+			// Fired when the user clicks a related row in the relation section
+			SLATE_EVENT(FOnRelatedRowSelected, OnRelatedRowSelected)
 		SLATE_END_ARGS()
-		
+
 		UE_API void Construct(const FArguments& InArgs);
 
 		// Set the row to view
@@ -86,6 +100,9 @@ namespace UE::Editor::DataStorage
 		// Clear the row to view
 		UE_API void ClearRow();
 
+		// Called by SRowDetailsRow to fire the OnRelatedRowSelected delegate
+		UE_API void NotifyRelatedRowSelected(RowHandle Row);
+
 		ICoreProvider* GetDataStorage() const { return DataStorage; }
 		IUiProvider* GetDataStorageUi() const { return DataStorageUi; }
 		const TSharedPtr<FMetaData>& GetNameMetaData() const { return NameMetaData; }
@@ -111,6 +128,8 @@ namespace UE::Editor::DataStorage
 		TAttribute<FMargin> RowPadding;
 		bool bShowAllDetails = true;
 		bool bShowUnpopulatedColumns = false;
+
+		FOnRelatedRowSelected OnRelatedRowSelectedDelegate;
 	};
 
 	class SRowDetailsRow : public SMultiColumnTableRow<RowDetailsItemPtr>
diff --git a/Engine/Source/Runtime/MassEntity/Public/MassEntityManager.h b/Engine/Source/Runtime/MassEntity/Public/MassEntityManager.h
--- a/Engine/Source/Runtime/MassEntity/Public/MassEntityManager.h
+++ b/Engine/Source/Runtime/MassEntity/Public/MassEntityManager.h
@@ -1263,7 +1263,7 @@ public:
 	//-----------------------------------------------------------------------------
 	// Type management
 	//-----------------------------------------------------------------------------
-	void OnNewTypeRegistered(UE::Mass::FTypeHandle RegisteredTypeHandle);
+	UE_API void OnNewTypeRegistered(UE::Mass::FTypeHandle RegisteredTypeHandle);
 
 protected:
 	void OnRelationTypeRegistered(UE::Mass::FTypeHandle RegisteredTypeHandle, const UE::Mass::FRelationTypeTraits& RelationTypeTraits);
diff --git a/Engine/Source/Runtime/TypedElementFramework/Private/Elements/Framework/TypedElementQueryBuilder.cpp b/Engine/Source/Runtime/TypedElementFramework/Private/Elements/Framework/TypedElementQueryBuilder.cpp
--- a/Engine/Source/Runtime/TypedElementFramework/Private/Elements/Framework/TypedElementQueryBuilder.cpp
+++ b/Engine/Source/Runtime/TypedElementFramework/Private/Elements/Framework/TypedElementQueryBuilder.cpp
@@ -244,6 +244,256 @@ namespace UE::Editor::DataStorage::Queries
 		return *this;
 	}
 
+	// ============================================================================
+	// Relation Conditions
+	// ============================================================================
+
+	FSimpleQuery& FSimpleQuery::IsSubjectOf(RelationTypeHandle RelationType)
+	{
+		FQueryDescription::FRelationCondition Condition;
+		Condition.RelationType = RelationType;
+		Condition.Role = FQueryDescription::ERelationRole::Subject;
+		Condition.TargetRow = InvalidRowHandle;
+		Query->RelationConditions.Add(Condition);
+		return *this;
+	}
+
+	FSimpleQuery& FSimpleQuery::IsSubjectOf(RelationTypeHandle RelationType, RowHandle Object)
+	{
+		FQueryDescription::FRelationCondition Condition;
+		Condition.RelationType = RelationType;
+		Condition.Role = FQueryDescription::ERelationRole::Subject;
+		Condition.TargetRow = Object;
+		Query->RelationConditions.Add(Condition);
+		return *this;
+	}
+
+	FSimpleQuery& FSimpleQuery::IsObjectOf(RelationTypeHandle RelationType)
+	{
+		FQueryDescription::FRelationCondition Condition;
+		Condition.RelationType = RelationType;
+		Condition.Role = FQueryDescription::ERelationRole::Object;
+		Condition.TargetRow = InvalidRowHandle;
+		Query->RelationConditions.Add(Condition);
+		return *this;
+	}
+
+	FSimpleQuery& FSimpleQuery::IsObjectOf(RelationTypeHandle RelationType, RowHandle Subject)
+	{
+		FQueryDescription::FRelationCondition Condition;
+		Condition.RelationType = RelationType;
+		Condition.Role = FQueryDescription::ERelationRole::Object;
+		Condition.TargetRow = Subject;
+		Query->RelationConditions.Add(Condition);
+		return *this;
+	}
+
+	// ============================================================================
+	// Hierarchy Conditions
+	// ============================================================================
+
+	FSimpleQuery& FSimpleQuery::IsDescendantOf(RelationTypeHandle RelationType, RowHandle Ancestor, bool bIncludeSelf)
+	{
+		FQueryDescription::FHierarchyCondition Condition;
+		Condition.RelationType = RelationType;
+		Condition.ReferenceRow = Ancestor;
+		Condition.bIncludeSelf = bIncludeSelf;
+		Query->DescendantOfConditions.Add(Condition);
+		return *this;
+	}
+
+	FSimpleQuery& FSimpleQuery::IsAncestorOf(RelationTypeHandle RelationType, RowHandle Descendant, bool bIncludeSelf)
+	{
+		FQueryDescription::FHierarchyCondition Condition;
+		Condition.RelationType = RelationType;
+		Condition.ReferenceRow = Descendant;
+		Condition.bIncludeSelf = bIncludeSelf;
+		Query->AncestorOfConditions.Add(Condition);
+		return *this;
+	}
+
+	// ============================================================================
+	// Relation Predicate Factory Functions
+	// ============================================================================
+
+	FQueryDescription::FRelationPredicate IsSubjectOf(RelationTypeHandle RelationType)
+	{
+		FQueryDescription::FRelationPredicate Pred;
+		Pred.Type = FQueryDescription::FRelationPredicate::EType::SubjectOf;
+		Pred.RelationType = RelationType;
+		return Pred;
+	}
+
+	FQueryDescription::FRelationPredicate IsSubjectOf(RelationTypeHandle RelationType, RowHandle Object)
+	{
+		FQueryDescription::FRelationPredicate Pred;
+		Pred.Type = FQueryDescription::FRelationPredicate::EType::SubjectOf;
+		Pred.RelationType = RelationType;
+		Pred.ReferenceRow = Object;
+		return Pred;
+	}
+
+	FQueryDescription::FRelationPredicate IsObjectOf(RelationTypeHandle RelationType)
+	{
+		FQueryDescription::FRelationPredicate Pred;
+		Pred.Type = FQueryDescription::FRelationPredicate::EType::ObjectOf;
+		Pred.RelationType = RelationType;
+		return Pred;
+	}
+
+	FQueryDescription::FRelationPredicate IsObjectOf(RelationTypeHandle RelationType, RowHandle Subject)
+	{
+		FQueryDescription::FRelationPredicate Pred;
+		Pred.Type = FQueryDescription::FRelationPredicate::EType::ObjectOf;
+		Pred.RelationType = RelationType;
+		Pred.ReferenceRow = Subject;
+		return Pred;
+	}
+
+	FQueryDescription::FRelationPredicate IsDescendantOf(RelationTypeHandle RelationType, RowHandle Ancestor, bool bIncludeSelf)
+	{
+		FQueryDescription::FRelationPredicate Pred;
+		Pred.Type = FQueryDescription::FRelationPredicate::EType::DescendantOf;
+		Pred.RelationType = RelationType;
+		Pred.ReferenceRow = Ancestor;
+		Pred.bIncludeSelf = bIncludeSelf;
+		return Pred;
+	}
+
+	FQueryDescription::FRelationPredicate IsAncestorOf(RelationTypeHandle RelationType, RowHandle Descendant, bool bIncludeSelf)
+	{
+		FQueryDescription::FRelationPredicate Pred;
+		Pred.Type = FQueryDescription::FRelationPredicate::EType::AncestorOf;
+		Pred.RelationType = RelationType;
+		Pred.ReferenceRow = Descendant;
+		Pred.bIncludeSelf = bIncludeSelf;
+		return Pred;
+	}
+
+	// ============================================================================
+	// FQueryExpression operator implementations
+	// ============================================================================
+
+	namespace
+	{
+		/** Dispatch a FRelationPredicate to the appropriate FQueryDescription condition array. */
+		static void AddPredicateToDescription(FQueryDescription& Desc, const FQueryDescription::FRelationPredicate& Pred)
+		{
+			using EType = FQueryDescription::FRelationPredicate::EType;
+			switch (Pred.Type)
+			{
+				case EType::SubjectOf:
+				{
+					FQueryDescription::FRelationCondition Cond;
+					Cond.RelationType = Pred.RelationType;
+					Cond.Role = FQueryDescription::ERelationRole::Subject;
+					Cond.TargetRow = Pred.ReferenceRow;
+					Desc.RelationConditions.Add(Cond);
+					break;
+				}
+				case EType::ObjectOf:
+				{
+					FQueryDescription::FRelationCondition Cond;
+					Cond.RelationType = Pred.RelationType;
+					Cond.Role = FQueryDescription::ERelationRole::Object;
+					Cond.TargetRow = Pred.ReferenceRow;
+					Desc.RelationConditions.Add(Cond);
+					break;
+				}
+				case EType::DescendantOf:
+				{
+					FQueryDescription::FHierarchyCondition Cond;
+					Cond.RelationType = Pred.RelationType;
+					Cond.ReferenceRow = Pred.ReferenceRow;
+					Cond.bIncludeSelf = Pred.bIncludeSelf;
+					Desc.DescendantOfConditions.Add(Cond);
+					break;
+				}
+				case EType::AncestorOf:
+				{
+					FQueryDescription::FHierarchyCondition Cond;
+					Cond.RelationType = Pred.RelationType;
+					Cond.ReferenceRow = Pred.ReferenceRow;
+					Cond.bIncludeSelf = Pred.bIncludeSelf;
+					Desc.AncestorOfConditions.Add(Cond);
+					break;
+				}
+			}
+		}
+	}
+
+	FQueryExpression operator&&(FColumnBase Lhs, const FQueryDescription::FRelationPredicate& Rhs)
+	{
+		FQueryExpression Result;
+		Result.ColumnConditions = FConditions(Lhs);
+		Result.RelationAtoms.Add(Rhs);
+		return Result;
+	}
+
+	FQueryExpression operator&&(const FQueryDescription::FRelationPredicate& Lhs, FColumnBase Rhs)
+	{
+		FQueryExpression Result;
+		Result.RelationAtoms.Add(Lhs);
+		Result.ColumnConditions = FConditions(Rhs);
+		return Result;
+	}
+
+	FQueryExpression operator&&(const FConditions& Lhs, const FQueryDescription::FRelationPredicate& Rhs)
+	{
+		FQueryExpression Result;
+		Result.ColumnConditions = Lhs;
+		Result.RelationAtoms.Add(Rhs);
+		return Result;
+	}
+
+	FQueryExpression operator&&(const FQueryDescription::FRelationPredicate& Lhs, const FConditions& Rhs)
+	{
+		FQueryExpression Result;
+		Result.RelationAtoms.Add(Lhs);
+		Result.ColumnConditions = Rhs;
+		return Result;
+	}
+
+	FQueryExpression operator&&(const FQueryDescription::FRelationPredicate& Lhs, const FQueryDescription::FRelationPredicate& Rhs)
+	{
+		FQueryExpression Result;
+		Result.RelationAtoms.Add(Lhs);
+		Result.RelationAtoms.Add(Rhs);
+		return Result;
+	}
+
+	FQueryExpression operator&&(FQueryExpression Lhs, const FQueryDescription::FRelationPredicate& Rhs)
+	{
+		Lhs.RelationAtoms.Add(Rhs);
+		return Lhs;
+	}
+
+	FQueryExpression operator&&(FQueryExpression Lhs, FColumnBase Rhs)
+	{
+		if (Lhs.ColumnConditions.IsSet())
+		{
+			Lhs.ColumnConditions = Lhs.ColumnConditions.GetValue() && Rhs;
+		}
+		else
+		{
+			Lhs.ColumnConditions = FConditions(Rhs);
+		}
+		return Lhs;
+	}
+
+	FQueryExpression operator&&(FQueryExpression Lhs, const FConditions& Rhs)
+	{
+		if (Lhs.ColumnConditions.IsSet())
+		{
+			Lhs.ColumnConditions = Lhs.ColumnConditions.GetValue() && Rhs;
+		}
+		else
+		{
+			Lhs.ColumnConditions = Rhs;
+		}
+		return Lhs;
+	}
+
 	FDependency FSimpleQuery::DependsOn()
 	{
 		return FDependency{ Query };
@@ -272,6 +522,9 @@ namespace UE::Editor::DataStorage::Queries
 		Query->CachedDependencies.Shrink();
 		Query->Subqueries.Shrink();
 		Query->MetaData.Shrink();
+		Query->RelationConditions.Shrink();
+		Query->DescendantOfConditions.Shrink();
+		Query->AncestorOfConditions.Shrink();
 		return MoveTemp(*Query);
 	}
 
@@ -533,6 +786,19 @@ namespace UE::Editor::DataStorage::Queries
 		return FQueryConditionQuery(&Query);
 	}
 
+	FQueryConditionQuery Select::Where(Queries::FQueryExpression Expression)
+	{
+		if (Expression.ColumnConditions.IsSet())
+		{
+			Query.Conditions.Emplace(MoveTemp(Expression.ColumnConditions.GetValue()));
+		}
+		for (FQueryDescription::FRelationPredicate& Pred : Expression.RelationAtoms)
+		{
+			AddPredicateToDescription(Query, Pred);
+		}
+		return FQueryConditionQuery(&Query);
+	}
+
 	FQueryDescription&& Select::Compile()
 	{
 		return MoveTemp(Query);
diff --git a/Engine/Source/Runtime/TypedElementFramework/Public/DataStorage/Handles.h b/Engine/Source/Runtime/TypedElementFramework/Public/DataStorage/Handles.h
--- a/Engine/Source/Runtime/TypedElementFramework/Public/DataStorage/Handles.h
+++ b/Engine/Source/Runtime/TypedElementFramework/Public/DataStorage/Handles.h
@@ -18,6 +18,13 @@ namespace UE::Editor::DataStorage
 	using QueryHandle = uint64;
 	static constexpr QueryHandle InvalidQueryHandle = TNumericLimits<QueryHandle>::Max();
 
+	/**
+	 * Handle to a registered relation type.
+	 * Relation types define how two rows can be connected (e.g., SubclassOf, ImplementsInterface).
+	 */
+	using RelationTypeHandle = uint64;
+	static constexpr RelationTypeHandle InvalidRelationTypeHandle = TNumericLimits<RelationTypeHandle>::Max();
+
 	struct FHierarchyHandle
 	{
 		uint64 Reserved = 0;
diff --git a/Engine/Source/Runtime/TypedElementFramework/Public/DataStorage/Queries/Description.h b/Engine/Source/Runtime/TypedElementFramework/Public/DataStorage/Queries/Description.h
--- a/Engine/Source/Runtime/TypedElementFramework/Public/DataStorage/Queries/Description.h
+++ b/Engine/Source/Runtime/TypedElementFramework/Public/DataStorage/Queries/Description.h
@@ -69,6 +69,75 @@ namespace UE::Editor::DataStorage
 			FName MatchValue;
 		};
 
+		/**
+		 * Specifies the role a row plays in a relation.
+		 */
+		enum class ERelationRole : uint8
+		{
+			/** The "from" / "child" side of a relation */
+			Subject,
+			/** The "to" / "parent" side of a relation */
+			Object
+		};
+
+		/**
+		 * Condition for filtering rows by their participation in relations.
+		 * If TargetRow is set, only matches relations with that specific counterpart.
+		 */
+		struct FRelationCondition
+		{
+			/** The relation type to filter by */
+			RelationTypeHandle RelationType = InvalidRelationTypeHandle;
+
+			/** Whether we're looking for rows as subjects or objects of the relation */
+			ERelationRole Role = ERelationRole::Subject;
+
+			/** Optional: If set, only match relations where the counterpart is this specific row */
+			RowHandle TargetRow = InvalidRowHandle;
+		};
+
+		/**
+		 * Condition for filtering rows by their position in a hierarchical relation.
+		 * Used for IsDescendantOf/IsAncestorOf queries.
+		 */
+		struct FHierarchyCondition
+		{
+			/** The hierarchical relation type to traverse */
+			RelationTypeHandle RelationType = InvalidRelationTypeHandle;
+
+			/** The reference row to check ancestry/descendancy against */
+			RowHandle ReferenceRow = InvalidRowHandle;
+
+			/** Whether to include the reference row itself in the results */
+			bool bIncludeSelf = false;
+		};
+
+		/**
+		 * A single atomic relation predicate. Multiple predicates are AND-composed.
+		 * Unifies the four existing condition types into one value type.
+		 */
+		struct FRelationPredicate
+		{
+			enum class EType : uint8
+			{
+				SubjectOf,      //< Row must be a subject of the relation (optionally with a specific object)
+				ObjectOf,       //< Row must be an object of the relation (optionally with a specific subject)
+				DescendantOf,   //< Row must be a descendant of ReferenceRow in a hierarchical relation
+				AncestorOf,     //< Row must be an ancestor of ReferenceRow in a hierarchical relation
+			};
+
+			EType              Type         = EType::SubjectOf;
+			RelationTypeHandle RelationType = InvalidRelationTypeHandle;
+
+			/** For SubjectOf/ObjectOf: the specific counterpart row, or InvalidRowHandle to match any.
+			 *  For DescendantOf: the ancestor to test against.
+			 *  For AncestorOf: the descendant to test against. */
+			RowHandle          ReferenceRow = InvalidRowHandle;
+
+			/** For DescendantOf/AncestorOf: if true, the reference row itself also passes. */
+			bool               bIncludeSelf = false;
+		};
+
 		struct FCallbackData
 		{
 			// Necessary to avoid deprecation warnings in clang when using compiler generated copy and move ctors. Can
@@ -133,6 +202,15 @@ namespace UE::Editor::DataStorage
 
 		FName Hierarchy;
 
+		/** Relation conditions for filtering rows by relation participation */
+		TArray<FRelationCondition, TInlineAllocator<2>> RelationConditions;
+
+		/** Descendant-of conditions for hierarchical filtering */
+		TArray<FHierarchyCondition, TInlineAllocator<1>> DescendantOfConditions;
+
+		/** Ancestor-of conditions for hierarchical filtering */
+		TArray<FHierarchyCondition, TInlineAllocator<1>> AncestorOfConditions;
+
 		EActionType Action;
 		bool bShouldBatchModifications = false;
 	};
diff --git a/Engine/Source/Runtime/TypedElementFramework/Public/Elements/Framework/TypedElementQueryBuilder.h b/Engine/Source/Runtime/TypedElementFramework/Public/Elements/Framework/TypedElementQueryBuilder.h
--- a/Engine/Source/Runtime/TypedElementFramework/Public/Elements/Framework/TypedElementQueryBuilder.h
+++ b/Engine/Source/Runtime/TypedElementFramework/Public/Elements/Framework/TypedElementQueryBuilder.h
@@ -144,6 +144,56 @@ namespace UE::Editor::DataStorage::Queries
 		FQueryDescription* Query;
 	};
 
+	// ============================================================================
+	// Relation Predicate Factory Functions
+	// ============================================================================
+
+	TYPEDELEMENTFRAMEWORK_API FQueryDescription::FRelationPredicate IsSubjectOf(RelationTypeHandle RelationType);
+	TYPEDELEMENTFRAMEWORK_API FQueryDescription::FRelationPredicate IsSubjectOf(RelationTypeHandle RelationType, RowHandle Object);
+	TYPEDELEMENTFRAMEWORK_API FQueryDescription::FRelationPredicate IsObjectOf(RelationTypeHandle RelationType);
+	TYPEDELEMENTFRAMEWORK_API FQueryDescription::FRelationPredicate IsObjectOf(RelationTypeHandle RelationType, RowHandle Subject);
+	TYPEDELEMENTFRAMEWORK_API FQueryDescription::FRelationPredicate IsDescendantOf(RelationTypeHandle RelationType, RowHandle Ancestor, bool bIncludeSelf = false);
+	TYPEDELEMENTFRAMEWORK_API FQueryDescription::FRelationPredicate IsAncestorOf(RelationTypeHandle RelationType, RowHandle Descendant, bool bIncludeSelf = false);
+
+	// ============================================================================
+	// FQueryExpression — composite expression of column conditions and relation predicates.
+	//
+	// Build one with operator&& between:
+	//   - TColumn<T>() or FConditions  (column conditions, evaluated at archetype level)
+	//   - IsDescendantOf / IsSubjectOf / etc. free functions (relation predicates, per-entity)
+	// All conditions are AND-composed.
+	//
+	// Example:
+	//   using namespace UE::Editor::DataStorage::Queries;
+	//   Select(...).Where(TColumn<FMyColumn>() && IsDescendantOf(MyRelation, RootRow)).Compile()
+	// ============================================================================
+
+	/**
+	 * Composite query expression holding both column conditions (evaluated at archetype selection time
+	 * via SetArchetypeMatchOverride) and relation predicates (evaluated per-entity after chunk selection).
+	 * Column and relation conditions are AND-composed with each other.
+	 */
+	struct FQueryExpression
+	{
+		/** Column conditions. Evaluated at archetype level. Absent when only relation predicates are combined. */
+		TOptional<FConditions> ColumnConditions;
+		/** Relation predicates AND-composed with each other and with ColumnConditions. */
+		TArray<FQueryDescription::FRelationPredicate, TInlineAllocator<2>> RelationAtoms;
+	};
+
+	// ---- Produce FQueryExpression: column atom + relation predicate ----
+	TYPEDELEMENTFRAMEWORK_API FQueryExpression operator&&(FColumnBase Lhs, const FQueryDescription::FRelationPredicate& Rhs);
+	TYPEDELEMENTFRAMEWORK_API FQueryExpression operator&&(const FQueryDescription::FRelationPredicate& Lhs, FColumnBase Rhs);
+	// ---- Produce FQueryExpression: FConditions + relation predicate ----
+	TYPEDELEMENTFRAMEWORK_API FQueryExpression operator&&(const FConditions& Lhs, const FQueryDescription::FRelationPredicate& Rhs);
+	TYPEDELEMENTFRAMEWORK_API FQueryExpression operator&&(const FQueryDescription::FRelationPredicate& Lhs, const FConditions& Rhs);
+	// ---- Produce FQueryExpression: relation predicate + relation predicate ----
+	TYPEDELEMENTFRAMEWORK_API FQueryExpression operator&&(const FQueryDescription::FRelationPredicate& Lhs, const FQueryDescription::FRelationPredicate& Rhs);
+	// ---- Extend an existing FQueryExpression ----
+	TYPEDELEMENTFRAMEWORK_API FQueryExpression operator&&(FQueryExpression Lhs, const FQueryDescription::FRelationPredicate& Rhs);
+	TYPEDELEMENTFRAMEWORK_API FQueryExpression operator&&(FQueryExpression Lhs, FColumnBase Rhs);
+	TYPEDELEMENTFRAMEWORK_API FQueryExpression operator&&(FQueryExpression Lhs, const FConditions& Rhs);
+
 	class FSimpleQuery final
 	{
 	public:
@@ -224,6 +274,55 @@ namespace UE::Editor::DataStorage::Queries
 		template<auto Value, TEnumType EnumT = decltype(Value)>
 		FSimpleQuery& All();
 
+		// ============================================================================
+		// Relation Conditions
+		// ============================================================================
+
+		/**
+		 * Filter to rows that are subjects in relations of the given type.
+		 * E.g., IsSubjectOf(SubclassOfType) finds all rows that inherit from something.
+		 */
+		TYPEDELEMENTFRAMEWORK_API FSimpleQuery& IsSubjectOf(RelationTypeHandle RelationType);
+
+		/**
+		 * Filter to rows that are subjects in relations of the given type where the object
+		 * is the specified row. E.g., IsSubjectOf(SubclassOfType, UObject) finds direct subclasses of UObject.
+		 */
+		TYPEDELEMENTFRAMEWORK_API FSimpleQuery& IsSubjectOf(RelationTypeHandle RelationType, RowHandle Object);
+
+		/**
+		 * Filter to rows that are objects in relations of the given type.
+		 * E.g., IsObjectOf(SubclassOfType) finds all rows that have something inheriting from them.
+		 */
+		TYPEDELEMENTFRAMEWORK_API FSimpleQuery& IsObjectOf(RelationTypeHandle RelationType);
+
+		/**
+		 * Filter to rows that are objects in relations of the given type where the subject
+		 * is the specified row. E.g., IsObjectOf(SubclassOfType, UDerivedClass) finds UDerivedClass's parent.
+		 */
+		TYPEDELEMENTFRAMEWORK_API FSimpleQuery& IsObjectOf(RelationTypeHandle RelationType, RowHandle Subject);
+
+		// ============================================================================
+		// Hierarchy Conditions (for hierarchical relation types)
+		// ============================================================================
+
+		/**
+		 * Filter to rows that are descendants of the specified ancestor in the given hierarchical relation.
+		 * Uses interval encoding for O(1) evaluation per row.
+		 * @param RelationType Must be a hierarchical relation type (HierarchyMode != EHierarchyMode::Disabled)
+		 * @param Ancestor The ancestor row to test against
+		 * @param bIncludeSelf If true, the ancestor itself will also match
+		 */
+		TYPEDELEMENTFRAMEWORK_API FSimpleQuery& IsDescendantOf(RelationTypeHandle RelationType, RowHandle Ancestor, bool bIncludeSelf = false);
+
+		/**
+		 * Filter to rows that are ancestors of the specified descendant in the given hierarchical relation.
+		 * @param RelationType Must be a hierarchical relation type (HierarchyMode != EHierarchyMode::Disabled)
+		 * @param Descendant The descendant row to test against
+		 * @param bIncludeSelf If true, the descendant itself will also match
+		 */
+		TYPEDELEMENTFRAMEWORK_API FSimpleQuery& IsAncestorOf(RelationTypeHandle RelationType, RowHandle Descendant, bool bIncludeSelf = false);
+
 	private:
 		TYPEDELEMENTFRAMEWORK_API explicit FSimpleQuery(FQueryDescription* Query);
 
@@ -430,6 +529,7 @@ namespace UE::Editor::DataStorage::Queries
 
 		// Query Condition API
 		TYPEDELEMENTFRAMEWORK_API FQueryConditionQuery Where(Queries::FConditions Condition);
+		TYPEDELEMENTFRAMEWORK_API FQueryConditionQuery Where(Queries::FQueryExpression Expression);
 
 		TYPEDELEMENTFRAMEWORK_API FQueryDescription&& Compile();
 		TYPEDELEMENTFRAMEWORK_API FSimpleQuery Where();
diff --git a/Engine/Source/Runtime/TypedElementFramework/Public/Elements/Interfaces/TypedElementDataStorageInterface.h b/Engine/Source/Runtime/TypedElementFramework/Public/Elements/Interfaces/TypedElementDataStorageInterface.h
--- a/Engine/Source/Runtime/TypedElementFramework/Public/Elements/Interfaces/TypedElementDataStorageInterface.h
+++ b/Engine/Source/Runtime/TypedElementFramework/Public/Elements/Interfaces/TypedElementDataStorageInterface.h
@@ -58,13 +58,112 @@ struct FHierarchyRegistrationParams
 	FName Name;
 
 	// Optionally add a column to any row participating in the hierarchy if its parent changes (useful for change detection).
-	// This behaves similar to the sync tags and is removed at the end of the frame
+	// This behaves similar to the sync tags and is removed at the end of the frame.
+	// For the Relations backend this maps to bEnableSubjectChangedColumn on the underlying relation type.
 	bool bEnableParentChangedColumn = false;
-	
+
 	UE_DEPRECATED(5.8, "Providing an explicit column has been deprecated, use the bool argument + ICoreProvider::GetParentChangedColumn()")
 	const UScriptStruct* ParentChangedColumn = nullptr;
+
+	/** Selects which storage backend implements this hierarchy. */
+	enum class EBackend : uint8
+	{
+		/** Column-on-row storage (FTedsHierarchyRegistrar). Existing behavior. */
+		Legacy,
+		/** Relation-row storage with interval encoding for O(1) descendant queries. */
+		Relations
+	};
+	EBackend Backend = EBackend::Legacy;
 };
-	
+
+/**
+ * Configuration for a role (Subject or Object) in a relation type.
+ */
+struct FTedsRelationRoleTraits
+{
+	/**
+	 * Whether only one entity can hold this role per relation.
+	 * For hierarchies: Object (parent) is typically exclusive (one parent per child).
+	 */
+	bool bExclusive = false;
+
+	/**
+	 * Policy for what happens when the entity in this role is destroyed.
+	 */
+	enum class EDestructionPolicy : uint8
+	{
+		/** Just remove the relation row, leave other participant unchanged */
+		CleanUp,
+		/** Destroy the entity in the other role (e.g., destroy child when parent is deleted) */
+		Cascade,
+		/** Leave the other participant without relation (orphan it) */
+		Orphan
+	};
+	EDestructionPolicy DestructionPolicy = EDestructionPolicy::CleanUp;
+};
+
+/** Controls how hierarchical queries (IsDescendantOf, etc.) are implemented for a relation type. */
+enum class EHierarchyMode : uint8
+{
+	/** Not a hierarchy. Hierarchical APIs are not available. */
+	Disabled,
+	/** IsDescendantOf walks the parent chain. Always correct, O(depth). */
+	WalkOnly,
+	/** IsDescendantOf uses interval encoding for O(1) checks. Falls back to walk when intervals
+	 *  are stale; intervals are automatically refreshed at frame end. Recommended for dynamic hierarchies. */
+	IntervalEncoded,
+};
+
+/**
+ * Configuration for registering a new relation type.
+ * Relation types define how two rows can be connected.
+ */
+struct FTedsRelationTraits
+{
+	/** Configuration for the Subject role (the "from" / "child" side) */
+	FTedsRelationRoleTraits Subject;
+
+	/** Configuration for the Object role (the "to" / "parent" side) */
+	FTedsRelationRoleTraits Object;
+
+	/** Selects how hierarchical descendant queries are implemented for this relation type. */
+	EHierarchyMode HierarchyMode = EHierarchyMode::Disabled;
+
+	/**
+	 * Initial gap between interval values for new nodes.
+	 * Larger gaps allow more insertions before rebalancing is needed.
+	 * Only used when HierarchyMode is IntervalEncoded.
+	 */
+	int64 IntervalGap = 1000;
+};
+
+/**
+ * Parameters for registering a relation type.
+ */
+struct FRelationRegistrationParams
+{
+	/** Unique name for this relation type */
+	FName Name;
+
+	/** Relation type configuration */
+	FTedsRelationTraits Traits;
+
+	/**
+	 * When true, a tag column is stamped on the Subject row whenever a relation of this type
+	 * is created or destroyed with that row as Subject. The column is removed at FrameEnd.
+	 * Retrieve the column type via ICoreProvider::GetRelationSubjectChangedColumn.
+	 * Use with FObserver::OnAdd or a FProcessor Where().All(col) for change detection.
+	 */
+	bool bEnableSubjectChangedColumn = false;
+
+	/**
+	 * When true, a tag column is stamped on the Object row whenever a relation of this type
+	 * is created or destroyed with that row as Object. The column is removed at FrameEnd.
+	 * Retrieve the column type via ICoreProvider::GetRelationObjectChangedColumn.
+	 */
+	bool bEnableObjectChangedColumn = false;
+};
+
 class ICoreProvider : public 
 	Queries::ITableInfo<
 	Queries::ITableManagement<
@@ -444,7 +543,224 @@ public:
 	 * Outputs the registered query callbacks to the given output device for debugging purposes.
 	 */
 	virtual void DebugPrintQueryCallbacks(FOutputDevice& Output) = 0;
-	
+
+	// Relations Interface
+	// ===================
+	// Relations provide a way to connect two rows together. Unlike the hierarchy system which stores
+	// parent-child data on participating rows, relations are stored as separate rows containing
+	// Subject (from) and Object (to) handles. This allows:
+	// - Zero archetype changes on participating entity rows
+	// - Unlimited participation in different relation types
+	// - Per-relation metadata (timestamps, weights, etc.)
+	// - O(1) hierarchical queries via interval encoding
+
+	/**
+	 * @section Relation Type Management
+	 */
+
+	/**
+	 * Registers a new relation type with the given configuration.
+	 * Returns InvalidRelationTypeHandle if registration fails (e.g., name already exists).
+	 */
+	virtual RelationTypeHandle RegisterRelationType(const FRelationRegistrationParams& Params) = 0;
+
+	/** Finds a previously registered relation type by name. Returns InvalidRelationTypeHandle if not found. */
+	virtual RelationTypeHandle FindRelationType(const FName& Name) const = 0;
+
+	/** Checks if a relation type handle is valid. */
+	virtual bool IsValidRelationType(RelationTypeHandle Type) const = 0;
+
+	/** Gets the traits for a registered relation type. Returns nullptr if handle is invalid. */
+	virtual const FTedsRelationTraits* GetRelationTypeTraits(RelationTypeHandle Type) const = 0;
+
+	/** Lists all registered relation types by calling the callback for each. */
+	virtual void ListRelationTypes(TFunctionRef<void(RelationTypeHandle, const FName&)> Callback) const = 0;
+
+	/**
+	 * Returns the tag column type stamped on Subject rows when a relation of this type changes.
+	 * Only non-null when the relation was registered with bEnableSubjectChangedColumn=true.
+	 * Returns nullptr for an invalid handle or when the column was not enabled.
+	 */
+	virtual const UScriptStruct* GetRelationSubjectChangedColumn(RelationTypeHandle Type) const = 0;
+
+	/**
+	 * Returns the tag column type stamped on Object rows when a relation of this type changes.
+	 * Only non-null when the relation was registered with bEnableObjectChangedColumn=true.
+	 * Returns nullptr for an invalid handle or when the column was not enabled.
+	 */
+	virtual const UScriptStruct* GetRelationObjectChangedColumn(RelationTypeHandle Type) const = 0;
+
+	/**
+	 * @section Relation Instance Management
+	 */
+
+	/**
+	 * Creates a relation between two rows.
+	 * Returns the handle of the newly created relation row, or InvalidRowHandle if creation fails.
+	 *
+	 * For hierarchical relations, this also:
+	 * - Computes and stores depth, root, and interval values
+	 * - Updates indexes for efficient queries
+	 */
+	virtual RowHandle CreateRelation(RelationTypeHandle Type, RowHandle Subject, RowHandle Object) = 0;
+
+	/**
+	 * Creates multiple relations at once (more efficient than individual calls).
+	 * Subjects and Objects arrays must have the same length, or Objects can have length 1
+	 * (in which case all subjects relate to the same object).
+	 *
+	 * @param OutRelationRows Optional output array to receive the created relation row handles.
+	 */
+	virtual void BatchCreateRelations(
+		RelationTypeHandle Type,
+		TConstArrayView<RowHandle> Subjects,
+		TConstArrayView<RowHandle> Objects,
+		TArray<RowHandle>* OutRelationRows = nullptr) = 0;
+
+	/**
+	 * Destroys a specific relation between two rows.
+	 * Returns true if a relation was found and destroyed.
+	 */
+	virtual bool DestroyRelation(RelationTypeHandle Type, RowHandle Subject, RowHandle Object) = 0;
+
+	/**
+	 * Checks if a specific relation exists between two rows.
+	 */
+	virtual bool HasRelation(RelationTypeHandle Type, RowHandle Subject, RowHandle Object) const = 0;
+
+	/** Returns true if Subject has at least one Object in this relation type. */
+	virtual bool HasRelationObject(RelationTypeHandle Type, RowHandle Subject) const = 0;
+
+	/** Returns true if Object has at least one Subject in this relation type. */
+	virtual bool HasRelationSubject(RelationTypeHandle Type, RowHandle Object) const = 0;
+
+	/**
+	 * @section Relation Queries
+	 */
+
+	/**
+	 * Gets all Objects that the given Subject has relations to.
+	 * For a "ChildOf" relation, this returns the parent(s) of the given row.
+	 */
+	virtual void GetRelationObjects(
+		RelationTypeHandle Type,
+		RowHandle Subject,
+		TArray<RowHandle>& OutObjects) const = 0;
+
+	/**
+	 * Gets the single Object for a given Subject. Use this for exclusive-object relations
+	 * (bExclusiveObject=true), where a subject has at most one object.
+	 * Returns InvalidRowHandle if the subject has no object in this relation.
+	 */
+	virtual RowHandle GetRelationObject(RelationTypeHandle Type, RowHandle Subject) const = 0;
+
+	/**
+	 * Gets all Subjects that have relations to the given Object.
+	 * For a "ChildOf" relation, this returns the children of the given row.
+	 */
+	virtual void GetRelationSubjects(
+		RelationTypeHandle Type,
+		RowHandle Object,
+		TArray<RowHandle>& OutSubjects) const = 0;
+
+	/**
+	 * Gets the single Subject for a given Object. Use this for exclusive-subject relations
+	 * (bExclusiveSubject=true), where an object has at most one subject.
+	 * Returns InvalidRowHandle if the object has no subject in this relation.
+	 */
+	virtual RowHandle GetRelationSubject(RelationTypeHandle Type, RowHandle Object) const = 0;
+
+	/**
+	 * @section Hierarchical Relation Queries (requires HierarchyMode != EHierarchyMode::Disabled)
+	 */
+
+	/**
+	 * Checks if Descendant is anywhere below Ancestor in the hierarchy.
+	 * @param bIncludeSelf If true, returns true when Descendant == Ancestor.
+	 */
+	virtual bool IsDescendantOf(
+		RelationTypeHandle Type,
+		RowHandle Descendant,
+		RowHandle Ancestor,
+		bool bIncludeSelf = false) const = 0;
+
+	/**
+	 * Checks if Ancestor is anywhere above Descendant in the hierarchy.
+	 * Equivalent to IsDescendantOf with arguments swapped.
+	 * @param bIncludeSelf If true, returns true when Ancestor == Descendant.
+	 */
+	inline bool IsAncestorOf(
+		RelationTypeHandle Type,
+		RowHandle Ancestor,
+		RowHandle Descendant,
+		bool bIncludeSelf = false) const
+	{
+		return IsDescendantOf(Type, Descendant, Ancestor, bIncludeSelf);
+	}
+
+	/**
+	 * Gets the root of the hierarchy containing this row.
+	 * Returns the row itself if it has no parent (is a root).
+	 */
+	virtual RowHandle GetHierarchyRoot(RelationTypeHandle Type, RowHandle Row) const = 0;
+
+	/**
+	 * Gets the depth of a row in the hierarchy.
+	 * Roots return 0, their direct children return 1, etc.
+	 */
+	virtual int32 GetHierarchyDepth(RelationTypeHandle Type, RowHandle Row) const = 0;
+
+	/** Gets all descendants of a row in the hierarchy. */
+	virtual void GetDescendants(
+		RelationTypeHandle Type,
+		RowHandle Ancestor,
+		TArray<RowHandle>& OutDescendants) const = 0;
+
+	/**
+	 * Gets all ancestors of a row in the hierarchy (parent, grandparent, etc.).
+	 */
+	virtual void GetAncestors(
+		RelationTypeHandle Type,
+		RowHandle Descendant,
+		TArray<RowHandle>& OutAncestors) const = 0;
+
+	/** Callback for TraverseDescendants. Receives the current row, its object, and depth relative to StartRow. */
+	using FRelationTraversalCallback = TFunctionRef<void(RowHandle Current, RowHandle Object, int32 Depth)>;
+
+	/** Traverses descendants of StartRow depth-first, calling Callback for each. */
+	virtual void TraverseDescendants(
+		RelationTypeHandle Type,
+		RowHandle StartRow,
+		FRelationTraversalCallback Callback,
+		ETraversalOrder Order = ETraversalOrder::PreOrder,
+		int32 MaxDepth = TNumericLimits<int32>::Max()) const = 0;
+
+	/**
+	 * @section Callback-Based Relation Enumeration
+	 * Prefer these over the TArray variants when you only need iteration.
+	 */
+
+	/** Calls Callback for each Object related to Subject. */
+	virtual void ForEachRelationObject(RelationTypeHandle Type, RowHandle Subject,
+		TFunctionRef<void(RowHandle Object)> Callback) const = 0;
+
+	/** Calls Callback for each Subject related to Object. */
+	virtual void ForEachRelationSubject(RelationTypeHandle Type, RowHandle Object,
+		TFunctionRef<void(RowHandle Subject)> Callback) const = 0;
+
+	/**
+	 * Walks ancestors (parent, grandparent, ...) calling Callback for each.
+	 * Return false from Callback to stop early.
+	 */
+	virtual void TraverseAncestors(RelationTypeHandle Type, RowHandle Descendant,
+		TFunctionRef<bool(RowHandle Ancestor, int32 Depth)> Callback) const = 0;
+
+	/** Counts descendants by walking the subtree. Not cached. */
+	virtual int32 ComputeDescendantCount(RelationTypeHandle Type, RowHandle Ancestor) const = 0;
+
+	/** Returns true if Object has at least one Subject in this relation type. */
+	virtual bool HasRelationSubjects(RelationTypeHandle Type, RowHandle Object) const = 0;
+
 	/**
 	 * @section Query
 	 * @description
diff --git a/Engine/Source/Runtime/TypedElementFramework/Public/Elements/Interfaces/TypedElementQueryStorageInterfaces.h b/Engine/Source/Runtime/TypedElementFramework/Public/Elements/Interfaces/TypedElementQueryStorageInterfaces.h
--- a/Engine/Source/Runtime/TypedElementFramework/Public/Elements/Interfaces/TypedElementQueryStorageInterfaces.h
+++ b/Engine/Source/Runtime/TypedElementFramework/Public/Elements/Interfaces/TypedElementQueryStorageInterfaces.h
@@ -119,6 +119,18 @@ namespace UE::Editor::DataStorage
 		virtual RowHandle GetParentRow(
 			RowHandle Target) const = 0;
 
+		/**
+		 * Defers a relation change until the end of the current chunk loop.
+		 * Removes any existing relation of this type where Target is the Subject,
+		 * then creates a new one with NewObject as the Object.
+		 * Pass InvalidRowHandle for NewObject to remove the existing relation without creating a new one.
+		 * Safe to call from within a processor callback — does not call ICoreProvider directly.
+		 */
+		virtual void SetRelation(
+			RelationTypeHandle Type,
+			RowHandle Target,
+			RowHandle NewObject) = 0;
+
 		/** Returns the amount of time now and the last update in seconds. */
 		virtual float GetDeltaTimeSeconds() const = 0;
 
