[TEDS] Deprecated `CurrentRowHasColumns` and `CurrentBatchTableHasColumns` in favor of `CurrentTableHasColumns`.

#Synced-CL 52688201

diff --git a/Engine/Plugins/Experimental/EditorDataStorage/Source/TedsCore/Private/Processors/QueryCapabilityImplementationForwarder.h b/Engine/Plugins/Experimental/EditorDataStorage/Source/TedsCore/Private/Processors/QueryCapabilityImplementationForwarder.h
--- a/Engine/Plugins/Experimental/EditorDataStorage/Source/TedsCore/Private/Processors/QueryCapabilityImplementationForwarder.h
+++ b/Engine/Plugins/Experimental/EditorDataStorage/Source/TedsCore/Private/Processors/QueryCapabilityImplementationForwarder.h
@@ -96,13 +96,16 @@ namespace UE::Editor::DataStorage::Queries::Private
 
 #define CapabilityEnd(Capability) \
 	}; \
-	
+
+#define DeprecatedFunction(Version, Msg) UE_DEPRECATED(Version, Msg)
+
 #include "Elements/Framework/TypedElementQueryCapabilities.inl"
 
 #undef ArgTypeName
 #undef ArgName
 #undef ForwardArg
 #undef CallFunction
+#undef DeprecatedFunction
 #undef Function0
 #undef Function1
 #undef Function2
diff --git a/Engine/Plugins/Experimental/EditorDataStorageFeatures/Source/TedsOutliner/Private/Widgets/SceneOutlinerWidget.cpp b/Engine/Plugins/Experimental/EditorDataStorageFeatures/Source/TedsOutliner/Private/Widgets/SceneOutlinerWidget.cpp
--- a/Engine/Plugins/Experimental/EditorDataStorageFeatures/Source/TedsOutliner/Private/Widgets/SceneOutlinerWidget.cpp
+++ b/Engine/Plugins/Experimental/EditorDataStorageFeatures/Source/TedsOutliner/Private/Widgets/SceneOutlinerWidget.cpp
@@ -208,7 +208,7 @@ void FTedsSceneOutlinerWidgetConstructor::GetShowOptionsFilters(UE::Editor::Data
 			LOCTEXT("ToggleHideTemporaryObjectsToolTip", "When enabled, hides temporary/run-time objects."),
 		NAME_None,
 		nullptr,
-		[DataStorage, IsLevelALevelInstance](TConstQueryContext<SingleRowInfo> Context, const FTypedElementWorldColumn& WorldColumn)
+		[DataStorage, IsLevelALevelInstance](TConstQueryContext<SingleRowInfo, CurrentTableInfo> Context, const FTypedElementWorldColumn& WorldColumn)
 	 	{
 			if (DataStorage)
 			{
@@ -218,7 +218,7 @@ void FTedsSceneOutlinerWidgetConstructor::GetShowOptionsFilters(UE::Editor::Data
 					{
 						if (Object->HasAnyFlags(RF_Transient))
 						{
-							if (Context.CurrentRowHasColumns<FPieObjectTag>())
+							if (Context.CurrentTableHasColumns<FPieObjectTag>())
 							{
 								return true;
 							}
@@ -244,16 +244,16 @@ void FTedsSceneOutlinerWidgetConstructor::GetShowOptionsFilters(UE::Editor::Data
 		LOCTEXT("ToggleHideLevelInstancesToolTip", "When enabled, hides all level instance content."),
 		NAME_None,
 		nullptr,
-		[DataStorage, IsLevelALevelInstance](TConstQueryContext<SingleRowInfo> Context, const FTypedElementWorldColumn& WorldColumn)
+		[DataStorage, IsLevelALevelInstance](TConstQueryContext<SingleRowInfo, CurrentTableInfo> Context, const FTypedElementWorldColumn& WorldColumn)
 		 {
-			if (DataStorage && Context.CurrentRowHasColumns<FTypedElementLevelColumn>())
+			if (DataStorage && Context.CurrentTableHasColumns<FTypedElementLevelColumn>())
 			{
 				if (const FTypedElementLevelColumn* LevelColumn = DataStorage->GetColumn<FTypedElementLevelColumn>(Context.GetCurrentRow()))
 				{
 					return !IsLevelALevelInstance(*LevelColumn, WorldColumn);
 				}
 			}
-				
+
 			return true;
 		 }));
 
@@ -280,11 +280,11 @@ void FTedsSceneOutlinerWidgetConstructor::GetShowOptionsFilters(UE::Editor::Data
 		LOCTEXT("ToggleShowOnlyCurrentLevelToolTip", "When enabled, only shows objects that are in the Current Level."),
 		NAME_None,
 		nullptr,
-		[DataStorage](TConstQueryContext<SingleRowInfo> Context)
+		[DataStorage](TConstQueryContext<SingleRowInfo, CurrentTableInfo> Context)
 		 {
 			if (DataStorage)
 			{
-				if (Context.CurrentRowHasColumns<FEditorDataStorageLevelTag>())
+				if (Context.CurrentTableHasColumns<FEditorDataStorageLevelTag>())
 				{
 					if (const FTypedElementUObjectColumn* ObjectColumn = DataStorage->GetColumn<FTypedElementUObjectColumn>(Context.GetCurrentRow()))
 					{
@@ -339,16 +339,16 @@ void FTedsSceneOutlinerWidgetConstructor::GetShowOptionsFilters(UE::Editor::Data
 	const FName HideEmptyFoldersFilterName = TEXT("HideEmptyFoldersFilter");
 	const FHierarchyHandle Hierarchy = DataStorage->FindHierarchyByName(UTedsEditorHierarchyFactory::EditorObjectHierarchyName);
 	const UScriptStruct* FolderParentTag = DataStorage->GetParentTagType(Hierarchy);
-	
+
 	Filters.Add(HideEmptyFoldersFilterName, MakeShared<FTedsOutlinerFilter>(
 		HideEmptyFoldersFilterName,
 		LOCTEXT("ToggleHideEmptyFolders", "Hide Empty Folders"),
 		LOCTEXT("ToggleHideEmptyFoldersToolTip", "When enabled, hides all empty folders."),
 		NAME_None,
 		nullptr,
-		[FolderParentTag](TConstQueryContext<SingleRowInfo> Context)
+		[FolderParentTag](TConstQueryContext<CurrentTableInfo> Context)
 		 {
-			return !Context.CurrentRowHasColumns<FFolderCompatibilityColumn>() || Context.CurrentRowHasColumns({FolderParentTag});
+			return !Context.CurrentTableHasColumns<FFolderCompatibilityColumn>() || Context.CurrentTableHasColumns({FolderParentTag});
 		 }));
 }
 
diff --git a/Engine/Plugins/Experimental/MeshPartition/Source/MeshPartitionEditorUI/Private/UI/MeshPartitionSettingsWidget.cpp b/Engine/Plugins/Experimental/MeshPartition/Source/MeshPartitionEditorUI/Private/UI/MeshPartitionSettingsWidget.cpp
--- a/Engine/Plugins/Experimental/MeshPartition/Source/MeshPartitionEditorUI/Private/UI/MeshPartitionSettingsWidget.cpp
+++ b/Engine/Plugins/Experimental/MeshPartition/Source/MeshPartitionEditorUI/Private/UI/MeshPartitionSettingsWidget.cpp
@@ -655,10 +655,10 @@ TSharedRef<UE::Editor::DataStorage::STedsCompositeHierarchyViewer> SMegaMeshSett
 	TArray<TSharedPtr<DataStorage::FTedsFilter>> Filters;
 
 	{
-		TConstQueryFunction<bool> FilterQuery = 
-			BuildConstQueryFunction<bool>([](TConstQueryContext<SingleRowInfo> Context)
+		TConstQueryFunction<bool> FilterQuery =
+			BuildConstQueryFunction<bool>([](TConstQueryContext<CurrentTableInfo> Context)
 			{
-				if (Context.CurrentRowHasColumns<FMegaMeshAffectsFilterBounds>())
+				if (Context.CurrentTableHasColumns<FMegaMeshAffectsFilterBounds>())
 				{
 					return true;
 				}
diff --git a/Engine/Source/Runtime/TypedElementFramework/Public/Elements/Framework/TypedElementQueryCapabilities.h b/Engine/Source/Runtime/TypedElementFramework/Public/Elements/Framework/TypedElementQueryCapabilities.h
--- a/Engine/Source/Runtime/TypedElementFramework/Public/Elements/Framework/TypedElementQueryCapabilities.h
+++ b/Engine/Source/Runtime/TypedElementFramework/Public/Elements/Framework/TypedElementQueryCapabilities.h
@@ -167,8 +167,11 @@ namespace UE::Editor::DataStorage::Queries
 			true; \
 	};
 
+#define DeprecatedFunction(Version, Msg)
+
 #include "Elements/Framework/TypedElementQueryCapabilities.inl"
 
+#undef DeprecatedFunction
 #undef Function0
 #undef Function1
 #undef Function2
@@ -212,6 +215,8 @@ namespace UE::Editor::DataStorage::Queries
 #define ConstFunction3(Capability, Return, Function, Arg1, Arg2, Arg3) virtual Return Function(ArgTypeName Arg1 , ArgTypeName Arg2 , ArgTypeName Arg3 ) const = 0;
 #define ConstFunction4(Capability, Return, Function, Arg1, Arg2, Arg3, Arg4) virtual Return Function(ArgTypeName Arg1 , ArgTypeName Arg2 , ArgTypeName Arg3 , ArgTypeName Arg4 ) const = 0;
 
+#define DeprecatedFunction(Version, Msg) UE_DEPRECATED(Version, Msg)
+
 #define CapabilityEnd(InCapability) \
 	}; \
 	using InCapability = I##InCapability <IContextCapability>; \
@@ -222,6 +227,7 @@ namespace UE::Editor::DataStorage::Queries
 #undef ArgTypeName
 #undef ArgType
 #undef ArgName
+#undef DeprecatedFunction
 #undef Function0
 #undef Function1
 #undef Function2
diff --git a/Engine/Source/Runtime/TypedElementFramework/Public/Elements/Framework/TypedElementQueryCapabilities.inl b/Engine/Source/Runtime/TypedElementFramework/Public/Elements/Framework/TypedElementQueryCapabilities.inl
--- a/Engine/Source/Runtime/TypedElementFramework/Public/Elements/Framework/TypedElementQueryCapabilities.inl
+++ b/Engine/Source/Runtime/TypedElementFramework/Public/Elements/Framework/TypedElementQueryCapabilities.inl
@@ -98,13 +98,18 @@ CapabilityStart(TableManagement, EContextCapabilityFlags::SupportsSingle | ECont
 	{
 		return this->RegisterTable(TConstArrayView<const UScriptStruct*>({ ColumnTypes::StaticStruct()... }), TableName);
 	}
-	/** Register a table by name. */
+
+	/** Register a table by name with additional configuration options. */
 	template<TColumnType... ColumnTypes>
 	TableHandle RegisterTable(const FName& TableName, const FTableRegistrationOptions& Options)
 	{
 		return this->RegisterTable(TConstArrayView<const UScriptStruct*>({ ColumnTypes::StaticStruct()... }), TableName, Options);
 	}
 
+	/**
+	 * Register a function that's used to generate a foreign key for a particular row. Foreign keys create a unique identifier that
+	 * represents the row for a specific domain.
+	 */
 	template<FunctionType Function>
 	void RegisterTableForeignKey(TableHandle Table, const FName& Domain, Function&& Callback)
 	{
@@ -135,19 +140,23 @@ CapabilityStart(SingleRowInfo, EContextCapabilityFlags::SupportsSingle)
 	 * Returns whether a column matches the requested type or not. This version only applies to current row. This version is faster than
 	 * querying for an arbitrary row.
 	 */
+	DeprecatedFunction(5.8, "Use CurrentTableHasColumns instead. Make sure to add CurrentTableInfo to the query's context capabilities.")
 	ConstFunction1(SingleRowInfo, bool, CurrentRowHasColumns, (TConstArrayView<const UScriptStruct*>, Columns))
 	/** Returns the currently active row. */
 	ConstFunction0(SingleRowInfo, RowHandle, GetCurrentRow)
 
 #if defined(WithWrappers)
 	/**
-	 * Returns whether a column matches the requested type or not. This version only applies to current row. This version is faster than
-	 * querying for an arbitrary row.
+	 * Returns whether the current row's table has the requested columns.
+	 * @deprecated Use CurrentTableHasColumns instead. Make sure to add CurrentTableInfo to the query's context capabilities.
 	 */
 	template<TColumnType... ColumnTypes>
+	DeprecatedFunction(5.8, "Use CurrentTableHasColumns instead. Make sure to add CurrentTableInfo to the query's context capabilities.")
 	bool CurrentRowHasColumns() const
 	{
+		PRAGMA_DISABLE_DEPRECATION_WARNINGS
 		return this->CurrentRowHasColumns(TConstArrayView<const UScriptStruct*>({ ColumnTypes::StaticStruct()... }));
+		PRAGMA_ENABLE_DEPRECATION_WARNINGS
 	}
 #endif
 CapabilityEnd(SingleRowInfo)
@@ -161,6 +170,7 @@ CapabilityStart(RowBatchInfo, EContextCapabilityFlags::SupportsBatch)
 	/** Returns an view with the rows used by this batch. */
 	ConstFunction0(RowBatchInfo, FRowHandleArrayView, GetBatchRowHandles)
 	/** Checks if the rows in the current batch have the requested columns. */
+	DeprecatedFunction(5.8, "Use CurrentTableHasColumns instead. Make sure to add CurrentTableInfo to the query's context capabilities.")
 	ConstFunction1(RowBatchInfo, bool, CurrentBatchTableHasColumns, (TConstArrayView<const UScriptStruct*>, Columns))
 	/** Direct access to the const addresses of a column or null if not found in this query. */
 	ConstFunction1(RowBatchInfo, const void*, GetColumnBatchAddress, (const UScriptStruct*, ColumnType))
@@ -168,11 +178,17 @@ CapabilityStart(RowBatchInfo, EContextCapabilityFlags::SupportsBatch)
 	Function1(RowBatchInfo, void*, GetMutableColumnBatchAddress, (const UScriptStruct*, ColumnType))
 
 #if defined(WithWrappers)
-	/** Checks if the rows in the current batch have the requested columns. */
+	/**
+	 * Checks if the rows in the current batch have the requested columns.
+	 * @deprecated Use CurrentTableHasColumns instead. Make sure to add CurrentTableInfo to the query's context capabilities.
+	 */
 	template<TColumnType... ColumnTypes>
+	DeprecatedFunction(5.8, "Use CurrentTableHasColumns instead. Make sure to add CurrentTableInfo to the query's context capabilities.")
 	bool CurrentBatchTableHasColumns() const
 	{
+		PRAGMA_DISABLE_DEPRECATION_WARNINGS
 		return this->CurrentBatchTableHasColumns(TConstArrayView<const UScriptStruct*>({ ColumnTypes::StaticStruct()... }));
+		PRAGMA_ENABLE_DEPRECATION_WARNINGS
 	}
 
 	/**
@@ -287,6 +303,7 @@ CapabilityStart(DynamicColumnInfo, EContextCapabilityFlags::SupportsSingle | ECo
 	{
 		return Description.TemplateType ? this->FindDynamicColumnType(*Description.TemplateType, Description.Identifier) : nullptr;
 	}
+
 	UE_DEPRECATED(5.8, "Please one of the new versions of FindDynamicColumnType.")
 	const UScriptStruct* FindDynamicColumn(const FDynamicColumnDescription& Description) const
 	{
diff --git a/Engine/Source/Runtime/TypedElementFramework/Public/Elements/Framework/TypedElementQueryCapabilityForwarder.h b/Engine/Source/Runtime/TypedElementFramework/Public/Elements/Framework/TypedElementQueryCapabilityForwarder.h
--- a/Engine/Source/Runtime/TypedElementFramework/Public/Elements/Framework/TypedElementQueryCapabilityForwarder.h
+++ b/Engine/Source/Runtime/TypedElementFramework/Public/Elements/Framework/TypedElementQueryCapabilityForwarder.h
@@ -74,13 +74,16 @@ namespace UE::Editor::DataStorage::Queries::Private
 #define CapabilityEnd(Capability) \
 	};
 
+#define DeprecatedFunction(Version, Msg) UE_DEPRECATED(Version, Msg)
+
 #define WithWrappers 1
-	
+
 #include "Elements/Framework/TypedElementQueryCapabilities.inl"
 
 #undef ArgTypeName
 #undef ArgName
 #undef ForwardArg
+#undef DeprecatedFunction
 #undef Function0
 #undef Function1
 #undef Function2
diff --git a/Engine/Source/Runtime/TypedElementFramework/Public/Elements/Framework/TypedElementQueryContextImplementation.h b/Engine/Source/Runtime/TypedElementFramework/Public/Elements/Framework/TypedElementQueryContextImplementation.h
--- a/Engine/Source/Runtime/TypedElementFramework/Public/Elements/Framework/TypedElementQueryContextImplementation.h
+++ b/Engine/Source/Runtime/TypedElementFramework/Public/Elements/Framework/TypedElementQueryContextImplementation.h
@@ -27,6 +27,8 @@
 
 #define CapabilityEnd(Capability)
 
+#define DeprecatedFunction(Version, Msg) UE_DEPRECATED(Version, Msg)
+
 #define WithWrappers 1
 
 namespace UE::Editor::DataStorage::Queries
@@ -70,6 +72,7 @@ namespace UE::Editor::DataStorage::Queries
 // namespace UE::Editor::DataStorage::Queries
 
 #undef ArgTypeName
+#undef DeprecatedFunction
 #undef Function0
 #undef Function1
 #undef Function2
diff --git a/Engine/Source/Runtime/TypedElementFramework/Public/Elements/Framework/TypedElementQueryContextImplementation.inl b/Engine/Source/Runtime/TypedElementFramework/Public/Elements/Framework/TypedElementQueryContextImplementation.inl
--- a/Engine/Source/Runtime/TypedElementFramework/Public/Elements/Framework/TypedElementQueryContextImplementation.inl
+++ b/Engine/Source/Runtime/TypedElementFramework/Public/Elements/Framework/TypedElementQueryContextImplementation.inl
@@ -83,7 +83,9 @@ namespace UE::Editor::DataStorage::Queries
 		if constexpr (!bIsConst || (bIsConst && ConstFunction)) \
 		{ \
 			static_assert(bFunctionFound, "Function '" #FunctionName "' in capability '" #Capability "' is missing from the implementation. This is required by context implementations." ); \
+			PRAGMA_DISABLE_DEPRECATION_WARNINGS \
 			return ImplementationType:: FunctionName ( Environment PrependCommaIfNeeded(ExpandArgPack CallArgs) ); \
+			PRAGMA_ENABLE_DEPRECATION_WARNINGS \
 		} \
 		else \
 		{ \
@@ -162,6 +164,8 @@ namespace UE::Editor::DataStorage::Queries
 
 #define CapabilityEnd(Capability)
 
+#define DeprecatedFunction(Version, Msg)
+
 #include "Elements/Framework/TypedElementQueryCapabilities.inl"
 
 #undef ArgTypeName
@@ -170,6 +174,7 @@ namespace UE::Editor::DataStorage::Queries
 #undef ExpandArgPack
 #undef PrependCommaIfNeeded
 #undef FunctionCommon
+#undef DeprecatedFunction
 #undef Function0
 #undef Function1
 #undef Function2
diff --git a/Engine/Source/Runtime/TypedElementFramework/Public/Elements/Framework/TypedElementQueryContract.h b/Engine/Source/Runtime/TypedElementFramework/Public/Elements/Framework/TypedElementQueryContract.h
--- a/Engine/Source/Runtime/TypedElementFramework/Public/Elements/Framework/TypedElementQueryContract.h
+++ b/Engine/Source/Runtime/TypedElementFramework/Public/Elements/Framework/TypedElementQueryContract.h
@@ -159,10 +159,11 @@ namespace UE::Editor::DataStorage::Queries
 #define ConstFunction3(Capability, Return, Function, Arg1, Arg2, Arg3)
 #define ConstFunction4(Capability, Return, Function, Arg1, Arg2, Arg3, Arg4)
 
+#define DeprecatedFunction(Version, Msg)
 #define CapabilityStart(Capability, Flags) , I##Capability
 #define CapabilityEnd(Capability)
 
-	/** 
+	/**
 	 * Contract between context and implementation to be able to communicate.
 	 * The contract contains all functions for all capabilities that are available. Based on the supported capabilities
 	 * an implementation may opt to only partially implement the contract, with the remaining functions asserting. On 
@@ -195,6 +196,7 @@ namespace UE::Editor::DataStorage::Queries
 		virtual TConstArrayView<FName> GetSupportedCapabilities() const = 0;
 	};
 
+#undef DeprecatedFunction
 #undef Function0
 #undef Function1
 #undef Function2
diff --git a/Engine/Source/Runtime/TypedElementFramework/Tests/Elements/Framework/TypedElementQueryContextMock.h b/Engine/Source/Runtime/TypedElementFramework/Tests/Elements/Framework/TypedElementQueryContextMock.h
--- a/Engine/Source/Runtime/TypedElementFramework/Tests/Elements/Framework/TypedElementQueryContextMock.h
+++ b/Engine/Source/Runtime/TypedElementFramework/Tests/Elements/Framework/TypedElementQueryContextMock.h
@@ -118,9 +118,12 @@ namespace UE::Editor::DataStorage::Queries
 
 #define CapabilityEnd(Capability)
 
+#define DeprecatedFunction(Version, Msg)
+
 #include "Elements/Framework/TypedElementQueryCapabilities.inl"
 
 #undef CapabilityStart
+#undef DeprecatedFunction
 #undef Function0
 #undef Function1
 #undef Function2
@@ -229,9 +232,12 @@ namespace UE::Editor::DataStorage::Queries
 #define CapabilityEnd(Capability) \
 		};
 
+#define DeprecatedFunction(Version, Msg) UE_DEPRECATED(Version, Msg)
+
 #include "Elements/Framework/TypedElementQueryCapabilities.inl"
 
 #undef CapabilityStart
+#undef DeprecatedFunction
 #undef Function0
 #undef Function1
 #undef Function2
@@ -269,6 +275,8 @@ using QueryContextMock = TQueryContextImpl<false, FMockContextEnvironment
 #define ConstFunction3(Capability, Return, Function, Arg1, Arg2, Arg3)
 #define ConstFunction4(Capability, Return, Function, Arg1, Arg2, Arg3, Arg4)
 
+#define DeprecatedFunction(Version, Msg)
+
 #define CapabilityStart(Capability, Flags) , F##Capability##Mock
 #define CapabilityEnd(Capability)
 
@@ -285,6 +293,7 @@ using QueryContextMock = TQueryContextImpl<false, FMockContextEnvironment
 #undef ConstFunction2
 #undef ConstFunction3
 #undef ConstFunction4
+#undef DeprecatedFunction
 #undef CapabilityStart
 #undef CapabilityEnd
 
diff --git a/Engine/Source/Runtime/TypedElementFramework/Tests/Elements/Framework/TypedElementQueryFunctionsTests.cpp b/Engine/Source/Runtime/TypedElementFramework/Tests/Elements/Framework/TypedElementQueryFunctionsTests.cpp
--- a/Engine/Source/Runtime/TypedElementFramework/Tests/Elements/Framework/TypedElementQueryFunctionsTests.cpp
+++ b/Engine/Source/Runtime/TypedElementFramework/Tests/Elements/Framework/TypedElementQueryFunctionsTests.cpp
@@ -347,10 +347,10 @@ namespace UE::Editor::DataStorage::Queries::Tests
 		SECTION("Call and use context")
 		{
 			TQueryFunction<void> Function = BuildQueryFunction<void>(
-				[](TQueryContext<SingleRowInfo> Context)
+				[](TQueryContext<SingleRowInfo, CurrentTableInfo> Context)
 				{
 					Context.GetCurrentRow();
-					Context.CurrentRowHasColumns<FTestColumnString>();
+					Context.CurrentTableHasColumns<FTestColumnString>();
 				});
 
 			FMockContextEnvironment Environment;
@@ -363,12 +363,12 @@ namespace UE::Editor::DataStorage::Queries::Tests
 					bGetCurrentRowCalled = true;
 					return 0;
 				});
-			
-			bool bCurrentRowHasColumnsCalled = false;
-			Environment.Assign_CurrentRowHasColumns_Const(
-				[&bCurrentRowHasColumnsCalled](TConstArrayView<const UScriptStruct*>)
+
+			bool bCurrentTableHasColumnsCalled = false;
+			Environment.Assign_CurrentTableHasColumns_Const(
+				[&bCurrentTableHasColumnsCalled](TConstArrayView<const UScriptStruct*>)
 				{
-					bCurrentRowHasColumnsCalled = true;
+					bCurrentTableHasColumnsCalled = true;
 					return false;
 				});
 
@@ -377,7 +377,7 @@ namespace UE::Editor::DataStorage::Queries::Tests
 			Function.Call(Context, Response);
 
 			CHECK_MESSAGE(TEXT("GetCurrentRow not called on context."), bGetCurrentRowCalled);
-			CHECK_MESSAGE(TEXT("CurrentRowHasColumns not called on context."), bCurrentRowHasColumnsCalled);
+			CHECK_MESSAGE(TEXT("CurrentTableHasColumns not called on context."), bCurrentTableHasColumnsCalled);
 		}
 
 		SECTION("Call and check provided row handles")
